NumericEntryBox now displays the value in the tooltip by default.

diff --git a/Engine/Source/Runtime/Slate/Public/Widgets/Input/SNumericEntryBox.h b/Engine/Source/Runtime/Slate/Public/Widgets/Input/SNumericEntryBox.h
--- a/Engine/Source/Runtime/Slate/Public/Widgets/Input/SNumericEntryBox.h
+++ b/Engine/Source/Runtime/Slate/Public/Widgets/Input/SNumericEntryBox.h
@@ -178,20 +178,20 @@ public:
 
 		if( bAllowSpin )
 		{
-			SAssignNew( SpinBox, SSpinBox<NumericType> )
-				.Style( InArgs._SpinBoxStyle )
-				.Font( InArgs._Font.IsSet() ? InArgs._Font : InArgs._EditableTextBoxStyle->Font )
-				.ContentPadding( TextMargin )
-				.Value( this, &SNumericEntryBox<NumericType>::OnGetValueForSpinBox )
-				.Delta( InArgs._Delta )
+			SAssignNew(SpinBox, SSpinBox<NumericType>)
+				.Style(InArgs._SpinBoxStyle)
+				.Font(InArgs._Font.IsSet() ? InArgs._Font : InArgs._EditableTextBoxStyle->Font)
+				.ContentPadding(TextMargin)
+				.Value(this, &SNumericEntryBox<NumericType>::OnGetValueForSpinBox)
+				.Delta(InArgs._Delta)
 				.ShiftMouseMovePixelPerDelta(InArgs._ShiftMouseMovePixelPerDelta)
 				.LinearDeltaSensitivity(InArgs._LinearDeltaSensitivity)
 				.SupportDynamicSliderMaxValue(InArgs._SupportDynamicSliderMaxValue)
 				.SupportDynamicSliderMinValue(InArgs._SupportDynamicSliderMinValue)
 				.OnDynamicSliderMaxValueChanged(InArgs._OnDynamicSliderMaxValueChanged)
 				.OnDynamicSliderMinValueChanged(InArgs._OnDynamicSliderMinValueChanged)
-				.OnValueChanged( OnValueChanged )
-				.OnValueCommitted( OnValueCommitted )
+				.OnValueChanged(OnValueChanged)
+				.OnValueCommitted(OnValueCommitted)
 				.MinSliderValue(InArgs._MinSliderValue)
 				.MaxSliderValue(InArgs._MaxSliderValue)
 				.MaxValue(InArgs._MaxValue)
@@ -201,21 +201,23 @@ public:
 				.OnBeginSliderMovement(InArgs._OnBeginSliderMovement)
 				.OnEndSliderMovement(InArgs._OnEndSliderMovement)
 				.MinDesiredWidth(InArgs._MinDesiredValueWidth)
-				.TypeInterface(Interface);
+				.TypeInterface(Interface)
+				.ToolTipText(this, &SNumericEntryBox<NumericType>::GetValueAsText);
 		}
 
 		// Always create an editable text box.  In the case of an undetermined value being passed in, we cant use the spinbox.
-		SAssignNew( EditableText, SEditableText )
-			.Text( this, &SNumericEntryBox<NumericType>::OnGetValueForTextBox )
-			.Visibility( bAllowSpin ? EVisibility::Collapsed : EVisibility::Visible )
-			.Font( InArgs._Font.IsSet() ? InArgs._Font : InArgs._EditableTextBoxStyle->Font )
-			.SelectAllTextWhenFocused( true )
-			.ClearKeyboardFocusOnCommit( false )
-			.OnTextChanged( this, &SNumericEntryBox<NumericType>::OnTextChanged  )
-			.OnTextCommitted( this, &SNumericEntryBox<NumericType>::OnTextCommitted )
-			.SelectAllTextOnCommit( true )
-			.ContextMenuExtender( InArgs._ContextMenuExtender )
-			.MinDesiredWidth(InArgs._MinDesiredValueWidth);
+		SAssignNew(EditableText, SEditableText)
+			.Text(this, &SNumericEntryBox<NumericType>::OnGetValueForTextBox)
+			.Visibility(bAllowSpin ? EVisibility::Collapsed : EVisibility::Visible)
+			.Font(InArgs._Font.IsSet() ? InArgs._Font : InArgs._EditableTextBoxStyle->Font)
+			.SelectAllTextWhenFocused(true)
+			.ClearKeyboardFocusOnCommit(false)
+			.OnTextChanged(this, &SNumericEntryBox<NumericType>::OnTextChanged)
+			.OnTextCommitted(this, &SNumericEntryBox<NumericType>::OnTextCommitted)
+			.SelectAllTextOnCommit(true)
+			.ContextMenuExtender(InArgs._ContextMenuExtender)
+			.MinDesiredWidth(InArgs._MinDesiredValueWidth)
+			.ToolTipText(this, &SNumericEntryBox<NumericType>::GetValueAsText);
 
 		TSharedRef<SHorizontalBox> HorizontalBox = SNew(SHorizontalBox);
 	
@@ -356,6 +358,17 @@ private:
 		return 0;
 	}
 
+	FText GetValueAsText() const
+	{
+		const TOptional<NumericType>& Value = ValueAttribute.Get();
+		if (Value.IsSet() == true)
+		{
+			return FText::FromString(Interface->ToString(Value.GetValue()));
+		}
+		
+		return FText::GetEmpty();
+	}
+
 	/**
 	 * Called to get the value for the text box as FText                 
 	 */
