Combobox dropdowns are no longer scaled up twice by the high DPI UI scaling factor.

diff --git a/Engine/Source/Runtime/Slate/Private/Widgets/Input/SMenuAnchor.cpp b/Engine/Source/Runtime/Slate/Private/Widgets/Input/SMenuAnchor.cpp
--- a/Engine/Source/Runtime/Slate/Private/Widgets/Input/SMenuAnchor.cpp
+++ b/Engine/Source/Runtime/Slate/Private/Widgets/Input/SMenuAnchor.cpp
@@ -348,7 +348,6 @@ void SMenuAnchor::SetIsOpen( bool InIsOpen, const bool bFocusMenu, const int32 F
 
 					const FVector2D NewPosition = MyGeometry.AbsolutePosition;
 					FVector2D NewWindowSize = DesiredContentSize;
-					const FVector2D SummonLocationSize = MyGeometry.GetLocalSize();
 
 					FPopupTransitionEffect TransitionEffect( FPopupTransitionEffect::None );
 					if ( PlacementMode == MenuPlacement_ComboBox || PlacementMode == MenuPlacement_ComboBoxRight )
@@ -363,6 +362,7 @@ void SMenuAnchor::SetIsOpen( bool InIsOpen, const bool bFocusMenu, const int32 F
 					else if ( PlacementMode == MenuPlacement_MenuRight )
 					{
 						TransitionEffect = FPopupTransitionEffect( FPopupTransitionEffect::SubMenu );
+						NewWindowSize = MyGeometry.GetAbsoluteSize();
 					}
 
 					MethodInUse = Method.IsSet()
@@ -375,7 +375,7 @@ void SMenuAnchor::SetIsOpen( bool InIsOpen, const bool bFocusMenu, const int32 F
 						if (MethodInUse.GetPopupMethod() == EPopupMethod::CreateNewWindow)
 						{
 							// Open the pop-up
-							TSharedPtr<IMenu> NewMenu = FSlateApplication::Get().PushMenu(AsShared(), MyWidgetPath, MenuContentRef, NewPosition, TransitionEffect, bFocusMenu, MyGeometry.GetAbsoluteSize(), MethodInUse.GetPopupMethod(), bIsCollapsedByParent);
+							TSharedPtr<IMenu> NewMenu = FSlateApplication::Get().PushMenu(AsShared(), MyWidgetPath, MenuContentRef, NewPosition, TransitionEffect, bFocusMenu, NewWindowSize, MethodInUse.GetPopupMethod(), bIsCollapsedByParent);
 							
 							if (ensure(NewMenu.IsValid()))
 							{
@@ -428,8 +428,7 @@ void SMenuAnchor::SetIsOpen( bool InIsOpen, const bool bFocusMenu, const int32 F
 							// Start pop-up windows out transparent, then fade them in over time
 							const EWindowTransparency Transparency(EWindowTransparency::PerWindow);
 
-							const float TargetWindowOpacity = 1.0f;
-							FSlateRect Anchor(NewPosition, NewPosition + SummonLocationSize);
+							FSlateRect Anchor(NewPosition, NewPosition + MyGeometry.GetLocalSize());
 							EOrientation Orientation = (TransitionEffect.SlideDirection == FPopupTransitionEffect::SubMenu) ? Orient_Horizontal : Orient_Vertical;
 		
 							// @todo slate: Assumes that popup is not Scaled up or down from application scale.
