Fix for Mac Menu Bar Assertion in Debug Editor. All Mac Menu items are now destroyed in thee Game Thread.

diff --git a/Engine/Source/Runtime/Slate/Private/Framework/MultiBox/Mac/MacMenu.cpp b/Engine/Source/Runtime/Slate/Private/Framework/MultiBox/Mac/MacMenu.cpp
--- a/Engine/Source/Runtime/Slate/Private/Framework/MultiBox/Mac/MacMenu.cpp
+++ b/Engine/Source/Runtime/Slate/Private/Framework/MultiBox/Mac/MacMenu.cpp
@@ -34,6 +34,12 @@ struct FMacMenuItemState
 static TMap<FMacMenu*, TSharedPtr<TArray<FMacMenuItemState>>> GCachedMenuState;
 static FCriticalSection GCachedMenuStateCS;
 
+// The FMenuEntryBlock in FMacMenuItem will be deallocated when menu items
+// are removed from the NSMenu in FMacMenu::UpdateWithMultiBox() via the Main Thread. But SWidgets should only be destroyed in the Game Thread.
+// We save the FMenuEntryBlock ptrs in this queue and release them in the Game Thread
+static TArray<TSharedPtr<const FMenuEntryBlock>> GMacMenuEntryBlockDeletionQueue;
+static FCriticalSection GMacMenuEntryBlockCS;;
+
 @interface FMacMenuItem : NSMenuItem
 @property (assign) TSharedPtr<const FMenuEntryBlock> MenuEntryBlock;
 - (void)performAction;
@@ -48,6 +54,17 @@ static FCriticalSection GCachedMenuStateCS;
 	return self;
 }
 
+- (void) dealloc
+{
+	{
+		FScopeLock Lock(&GMacMenuEntryBlockCS);
+		// we save the block to prevent it from being released in the main thread
+		// releasing the block in Main Thread will trigger an assertion in FSlateApplicationBase::Get()
+		GMacMenuEntryBlockDeletionQueue.Add(self.MenuEntryBlock);
+	}
+	[super dealloc];
+}
+
 - (void)performAction
 {
 	FCocoaMenu* CocoaMenu = [[self menu] isKindOfClass:[FCocoaMenu class]] ? (FCocoaMenu*)[self menu] : nil;
@@ -289,8 +306,20 @@ static FStartupApplicationToMacMenuBinder StaticInitializer;
 
 void FSlateMacMenu::CleanupOnShutdown()
 {
-	FScopeLock Lock(&GCachedMenuStateCS);
-	GCachedMenuState.Reset();
+	{
+		FScopeLock Lock(&GCachedMenuStateCS);
+		GCachedMenuState.Reset();
+	}
+	{
+		GameThreadCall(^{
+			FScopeLock Lock(&GMacMenuEntryBlockCS);
+			for(auto& Block : GMacMenuEntryBlockDeletionQueue)
+			{
+				Block.Reset();
+			}
+			GMacMenuEntryBlockDeletionQueue.Empty();
+		}, @[ NSDefaultRunLoopMode ], true);
+	}
 }
 
 void FSlateMacMenu::PostInitStartup()
@@ -556,6 +585,15 @@ void FSlateMacMenu::UpdateWithMultiBox(const TSharedPtr< FMultiBox > MultiBox)
 
 		FPlatformApplicationMisc::bChachedMacMenuStateNeedsUpdate = true;
 	}, NSDefaultRunLoopMode, false);
+	// Now we destroy all the menu blocks in the game thread to ensure ~SWidget() is only called in Game Thread
+	GameThreadCall(^{
+		FScopeLock Lock(&GMacMenuEntryBlockCS);
+		for(auto& Block : GMacMenuEntryBlockDeletionQueue)
+		{
+			Block.Reset();
+		}
+		GMacMenuEntryBlockDeletionQueue.Empty();
+	}, @[NSDefaultRunLoopMode], false);
 }
 
 void FSlateMacMenu::UpdateMenu(FMacMenu* Menu)
