Bug fixed: When opening a BP in VR mode, a slight hitch occurred and the blueprint opened followed by the warning that "StandaloneToolkit" tab failed to spawn. The Warning has been replaced for a Log with the same message.

diff --git a/Engine/Source/Runtime/Slate/Private/Framework/Docking/TabManager.cpp b/Engine/Source/Runtime/Slate/Private/Framework/Docking/TabManager.cpp
--- a/Engine/Source/Runtime/Slate/Private/Framework/Docking/TabManager.cpp
+++ b/Engine/Source/Runtime/Slate/Private/Framework/Docking/TabManager.cpp
@@ -1664,10 +1664,10 @@ TSharedPtr<SDockTab> FTabManager::SpawnTab(const FTabId& TabId, const TSharedPtr
 		{
 			StringToDisplay = FString("Unknown");
 		}
-		// If an output must be generated, create an "unrecognized tab"
+		// If an output must be generated, create an "unrecognized tab" and log it
 		if (!bCanOutputBeNullptr)
 		{
-			UE_LOG(LogSlate, Warning,
+			UE_LOG(LogSlate, Log,
 				TEXT("The tab \"%s\" attempted to spawn but failed for some reason. An \"unrecognized tab\" will be returned instead."), *StringToDisplay
 			);
 
@@ -1686,10 +1686,10 @@ TSharedPtr<SDockTab> FTabManager::SpawnTab(const FTabId& TabId, const TSharedPtr
 
 			NewTabWidget->SetLayoutIdentifier(TabId);
 		}
-		// If we can return nullptr, report it in the log
+		// If we can return nullptr, log it
 		else
 		{
-			UE_LOG(LogSlate, Warning,
+			UE_LOG(LogSlate, Log,
 				TEXT("The tab \"%s\" attempted to spawn but failed for some reason. It will not be displayed but it will still be saved in the layout settings file."), *StringToDisplay
 			);
 		}
