SComboButton now has a HasDownArrowDropShadow argument which adds a drop shadow to the arrow. You still need to add a drop shadow to the button content if you want one.

diff --git a/Engine/Source/Runtime/Slate/Private/Widgets/Input/SComboButton.cpp b/Engine/Source/Runtime/Slate/Private/Widgets/Input/SComboButton.cpp
--- a/Engine/Source/Runtime/Slate/Private/Widgets/Input/SComboButton.cpp
+++ b/Engine/Source/Runtime/Slate/Private/Widgets/Input/SComboButton.cpp
@@ -58,11 +58,26 @@ void SComboButton::Construct( const FArguments& InArgs )
 				.VAlign( VAlign_Center )
 				.Padding( InArgs._HasDownArrow ? 2 : 0 )
 				[
-					SNew( SImage )
-					.Visibility( InArgs._HasDownArrow ? EVisibility::Visible : EVisibility::Collapsed )
-					.Image( &InArgs._ComboButtonStyle->DownArrowImage )
-					// Inherit tinting from parent
-					. ColorAndOpacity( FSlateColor::UseForeground() )
+					SNew(SOverlay)
+					// drop shadow
+					+ SOverlay::Slot()
+					.VAlign(VAlign_Top)
+					.Padding(0, 1, 0, 0)
+					[
+						SNew(SImage)
+						.Visibility( InArgs._HasDownArrow && InArgs._HasDownArrowDropShadow ? EVisibility::Visible : EVisibility::Collapsed )
+						.Image( &InArgs._ComboButtonStyle->DownArrowImage )
+						.ColorAndOpacity( FLinearColor::Black )
+					]
+					+ SOverlay::Slot()
+					.VAlign(VAlign_Top)
+					[
+						SNew(SImage)
+						.Visibility( InArgs._HasDownArrow ? EVisibility::Visible : EVisibility::Collapsed )
+						.Image( &InArgs._ComboButtonStyle->DownArrowImage )
+						// Inherit tinting from parent
+						.ColorAndOpacity( FSlateColor::UseForeground() )
+					]
 				]
 			]
 		]
diff --git a/Engine/Source/Runtime/Slate/Public/Widgets/Input/SComboButton.h b/Engine/Source/Runtime/Slate/Public/Widgets/Input/SComboButton.h
--- a/Engine/Source/Runtime/Slate/Public/Widgets/Input/SComboButton.h
+++ b/Engine/Source/Runtime/Slate/Public/Widgets/Input/SComboButton.h
@@ -56,6 +56,7 @@ public:
 		SLATE_EVENT( FOnComboBoxOpened, OnComboBoxOpened )
 		SLATE_ARGUMENT( bool, IsFocusable )
 		SLATE_ARGUMENT( bool, HasDownArrow )
+		SLATE_ARGUMENT( bool, HasDownArrowDropShadow )
 
 		SLATE_ATTRIBUTE( FSlateColor, ForegroundColor )
 		SLATE_ATTRIBUTE( FSlateColor, ButtonColorAndOpacity )
