Slate - If the incoming tint color and opacity changes, invalidate the invalidation panel.

Matt.Kuhlenschmidt

diff --git a/Engine/Source/Runtime/Slate/Private/Widgets/SInvalidationPanel.cpp b/Engine/Source/Runtime/Slate/Private/Widgets/SInvalidationPanel.cpp
--- a/Engine/Source/Runtime/Slate/Private/Widgets/SInvalidationPanel.cpp
+++ b/Engine/Source/Runtime/Slate/Private/Widgets/SInvalidationPanel.cpp
@@ -46,6 +46,8 @@ SInvalidationPanel::SInvalidationPanel()
 	SetInvalidationRootHittestGrid(HittestGrid);
 	SetCanTick(false);
 
+	LastIncomingColorAndOpacity = FLinearColor::White;
+
 	FSlateApplicationBase::Get().OnGlobalInvalidationToggled().AddRaw(this, &SInvalidationPanel::OnGlobalInvalidationToggled);
 }
 
@@ -104,7 +106,7 @@ void SInvalidationPanel::OnGlobalInvalidationToggled(bool bGlobalInvalidationEna
 	ClearAllFastPathData(false);
 }
 
-bool SInvalidationPanel::UpdateCachePrequisites(FSlateWindowElementList& OutDrawElements, const FGeometry& AllottedGeometry, const FSlateRect& MyCullingRect, int32 LayerId) const
+bool SInvalidationPanel::UpdateCachePrequisites(FSlateWindowElementList& OutDrawElements, const FGeometry& AllottedGeometry, const FSlateRect& MyCullingRect, int32 LayerId, const FWidgetStyle& InWidgetStyle) const
 {
 	bool bNeedsRecache = false;
 #if WITH_SLATE_DEBUGGING
@@ -143,6 +145,12 @@ bool SInvalidationPanel::UpdateCachePrequisites(FSlateWindowElementList& OutDraw
 		LastClippingState = ClippingState;
 		bNeedsRecache = true;
 	}
+
+	if (LastIncomingColorAndOpacity != InWidgetStyle.GetColorAndOpacityTint())
+	{
+		LastIncomingColorAndOpacity = InWidgetStyle.GetColorAndOpacityTint();
+		bNeedsRecache = true;
+	}
 	
 	return bNeedsRecache;
 }
@@ -198,7 +206,7 @@ int32 SInvalidationPanel::OnPaint( const FPaintArgs& Args, const FGeometry& Allo
 		NewArgs.GetHittestGrid().SetUserIndex(Args.RootGrid.GetUserIndex());
 		check(!GSlateEnableGlobalInvalidation);
 
-		const bool bRequiresRecache = UpdateCachePrequisites(OutDrawElements, AllottedGeometry, MyCullingRect, LayerId);
+		const bool bRequiresRecache = UpdateCachePrequisites(OutDrawElements, AllottedGeometry, MyCullingRect, LayerId, InWidgetStyle);
 		if (bHittestCleared || bRequiresRecache)
 		{
 			// @todo: Overly aggressive?
diff --git a/Engine/Source/Runtime/Slate/Public/Widgets/SInvalidationPanel.h b/Engine/Source/Runtime/Slate/Public/Widgets/SInvalidationPanel.h
--- a/Engine/Source/Runtime/Slate/Public/Widgets/SInvalidationPanel.h
+++ b/Engine/Source/Runtime/Slate/Public/Widgets/SInvalidationPanel.h
@@ -58,7 +58,7 @@ protected:
 	virtual int32 PaintSlowPath(const FSlateInvalidationContext& Context) override;
 private:
 	void OnGlobalInvalidationToggled(bool bGlobalInvalidationEnabled);
-	bool UpdateCachePrequisites(FSlateWindowElementList& OutDrawElements, const FGeometry& AllottedGeometry, const FSlateRect& MyCullingRect, int32 LayerId) const;
+	bool UpdateCachePrequisites(FSlateWindowElementList& OutDrawElements, const FGeometry& AllottedGeometry, const FSlateRect& MyCullingRect, int32 LayerId, const FWidgetStyle& InWidgetStyle) const;
 private:
 	FSimpleSlot EmptyChildSlot;
 
@@ -68,6 +68,7 @@ private:
 	mutable FGeometry LastAllottedGeometry;
 	mutable FVector2D LastClipRectSize;
 	mutable int32 LastIncomingLayerId;
+	mutable FLinearColor LastIncomingColorAndOpacity;
 
 	bool bCanCache;
 
