Preventing any out-of-order processing of touch events in FSlateApplication
- All touch events are ignored if they aren't between a TouchStarted and a TouchEnded
- Solves issues where mobile platforms can send force change events after ending the touch, causing FSlateUser to re-establish info for that pointer when it's not actually around/active

Also now resetting LastWidgetsUnderPointer info in FSlateUser when a touch pointer is released
- Solves issues where FSlateUser still had for the touch pointer the next time it was pressed, causing us to trigger a redundant OnMouseLeave on the widget that the pointer was over when it was released

diff --git a/Engine/Source/Runtime/Slate/Private/Framework/Application/SlateApplication.cpp b/Engine/Source/Runtime/Slate/Private/Framework/Application/SlateApplication.cpp
--- a/Engine/Source/Runtime/Slate/Private/Framework/Application/SlateApplication.cpp
+++ b/Engine/Source/Runtime/Slate/Private/Framework/Application/SlateApplication.cpp
@@ -5657,28 +5657,28 @@ bool FSlateApplication::OnTouchStarted( const TSharedPtr< FGenericWindow >& Plat
 
 void FSlateApplication::ProcessTouchStartedEvent( const TSharedPtr< FGenericWindow >& PlatformWindow, const FPointerEvent& InTouchEvent )
 {
-	GetOrCreateUser(InTouchEvent)->GetGestureDetector().OnTouchStarted(InTouchEvent.GetPointerIndex(), InTouchEvent.GetScreenSpacePosition());
+	GetOrCreateUser(InTouchEvent)->NotifyTouchStarted(InTouchEvent);
 	ProcessMouseButtonDownEvent(PlatformWindow, InTouchEvent);
 }
 
 bool FSlateApplication::OnTouchMoved( const FVector2D& Location, float Force, int32 TouchIndex, int32 ControllerId )
 {
-	// Don't process touches that overlap or surpass with the cursor pointer index.
-	if (TouchIndex >= (int32)ETouchIndex::CursorPointerIndex)
+	TSharedRef<FSlateUser> User = GetOrCreateUser(ControllerId);
+	if (User->IsTouchPointerActive(TouchIndex))
 	{
-		return false;
-	}
+		FPointerEvent PointerEvent(
+			ControllerId,
+			TouchIndex,
+			Location,
+			User->GetPreviousPointerPosition(TouchIndex),
+			Force,
+			true);
+		ProcessTouchMovedEvent(PointerEvent);
 
-	FPointerEvent PointerEvent(
-		ControllerId,
-		TouchIndex,
-		Location,
-		GetOrCreateUser(ControllerId)->GetPreviousPointerPosition(TouchIndex),
-		Force,
-		true);
-	ProcessTouchMovedEvent(PointerEvent);
+		return true;
+	}
 
-	return true;
+	return false;
 }
 
 void FSlateApplication::ProcessTouchMovedEvent( const FPointerEvent& PointerEvent )
@@ -5688,22 +5688,26 @@ void FSlateApplication::ProcessTouchMovedEvent( const FPointerEvent& PointerEven
 
 bool FSlateApplication::OnTouchEnded( const FVector2D& Location, int32 TouchIndex, int32 ControllerId )
 {
-	// Don't process touches that overlap or surpass with the cursor pointer index.
-	if (TouchIndex >= ((int32)ETouchIndex::CursorPointerIndex))
+	TSharedRef<FSlateUser> User = GetOrCreateUser(ControllerId);
+	if (User->IsTouchPointerActive(TouchIndex))
 	{
-		return false;
-	}
+		FPointerEvent PointerEvent(
+			ControllerId,
+			TouchIndex,
+			Location,
+			Location,
+			0.0f,
+			true);
+		ProcessTouchEndedEvent(PointerEvent);
 
-	FPointerEvent PointerEvent(
-		ControllerId,
-		TouchIndex,
-		Location,
-		Location,
-		0.0f,
-		true);
-	ProcessTouchEndedEvent(PointerEvent);
+#if WITH_SLATE_DEBUGGING
+		ensure(!User->IsTouchPointerActive(TouchIndex));
+#endif
 
-	return true;
+		return true;
+	}
+
+	return false;
 }
 
 void FSlateApplication::ProcessTouchEndedEvent(const FPointerEvent& PointerEvent)
@@ -5713,48 +5717,46 @@ void FSlateApplication::ProcessTouchEndedEvent(const FPointerEvent& PointerEvent
 
 bool FSlateApplication::OnTouchForceChanged(const FVector2D& Location, float Force, int32 TouchIndex, int32 ControllerId)
 {
-	// Don't process touches that overlap or surpass with the cursor pointer index.
-	if (TouchIndex >= ((int32)ETouchIndex::CursorPointerIndex))
+	TSharedRef<FSlateUser> User = GetOrCreateUser(ControllerId);
+	if (User->IsTouchPointerActive(TouchIndex))
 	{
-		return false;
-	}
+		FPointerEvent PointerEvent(
+			ControllerId,
+			TouchIndex,
+			Location,
+			Location,
+			Force,
+			true,
+			true,
+			false);
+		ProcessTouchMovedEvent(PointerEvent);
 
-	FPointerEvent PointerEvent(
-		ControllerId,
-		TouchIndex,
-		Location,
-		Location,
-		Force,
-		true,
-		true,
-		false);
-	ProcessTouchMovedEvent(PointerEvent);
-
-	return true;
+		return true;
+	}
 
+	return false;
 }
 
 bool FSlateApplication::OnTouchFirstMove(const FVector2D& Location, float Force, int32 TouchIndex, int32 ControllerId)
 {
-	// Don't process touches that overlap or surpass with the cursor pointer index.
-	if (TouchIndex >= ((int32)ETouchIndex::CursorPointerIndex))
+	TSharedRef<FSlateUser> User = GetOrCreateUser(ControllerId);
+	if (User->IsTouchPointerActive(TouchIndex))
 	{
-		return false;
-	}
-
-	FPointerEvent PointerEvent(
-		ControllerId,
-		TouchIndex,
-		Location,
-		GetOrCreateUser(ControllerId)->GetPreviousPointerPosition(TouchIndex),
-		Force,
-		true,
-		false,
-		true);
-	ProcessTouchMovedEvent(PointerEvent);
+		FPointerEvent PointerEvent(
+			ControllerId,
+			TouchIndex,
+			Location,
+			User->GetPreviousPointerPosition(TouchIndex),
+			Force,
+			true,
+			false,
+			true);
+		ProcessTouchMovedEvent(PointerEvent);
 
-	return true;
+		return true;
+	}
 
+	return false;
 }
 
 void FSlateApplication::ShouldSimulateGesture(EGestureEvent Gesture, bool bEnable)
diff --git a/Engine/Source/Runtime/Slate/Private/Framework/Application/SlateUser.cpp b/Engine/Source/Runtime/Slate/Private/Framework/Application/SlateUser.cpp
--- a/Engine/Source/Runtime/Slate/Private/Framework/Application/SlateUser.cpp
+++ b/Engine/Source/Runtime/Slate/Private/Framework/Application/SlateUser.cpp
@@ -532,6 +532,11 @@ void FSlateUser::CloseTooltip()
 	}
 }
 
+bool FSlateUser::IsTouchPointerActive(int32 TouchPointerIndex) const
+{
+	return TouchPointerIndex < (int32)ETouchIndex::CursorPointerIndex && PointerPositionsByIndex.Contains(TouchPointerIndex);
+}
+
 void FSlateUser::DrawWindowlessDragDropContent(const TSharedRef<SWindow>& WindowToDraw, FSlateWindowElementList& WindowElementList, int32& MaxLayerId)
 {
 	if (DragDropContent && DragDropContent->IsWindowlessOperation())
@@ -902,6 +907,14 @@ TSharedRef<SWindow> FSlateUser::GetOrCreateTooltipWindow()
 	return NewTooltipWindow;
 }
 
+void FSlateUser::NotifyTouchStarted(const FPointerEvent& TouchEvent)
+{
+	UE_CLOG(PointerPositionsByIndex.Contains(TouchEvent.GetUserIndex()), LogSlate, Error, TEXT("SlateUser [%d] notified of a touch starting for pointer [%d] without finding out it ever ended."));
+	
+	GestureDetector.OnTouchStarted(TouchEvent.GetPointerIndex(), TouchEvent.GetScreenSpacePosition());
+	PointerPositionsByIndex.FindOrAdd(TouchEvent.GetPointerIndex()) = TouchEvent.GetScreenSpacePosition();
+}
+
 void FSlateUser::NotifyPointerMoveBegin(const FPointerEvent& PointerEvent)
 {
 	if (PointerEvent.IsTouchEvent())
@@ -1025,7 +1038,7 @@ void FSlateUser::NotifyPointerReleased(const FPointerEvent& PointerEvent, const
 				WidgetsUnderCursor.Widgets[WidgetIndex].Widget->OnMouseLeave(PointerEvent);
 			}
 		}
-#if !UE_BUILD_SHIPPING
+#if WITH_SLATE_DEBUGGING
 		const bool bIsFingerCapturedPostRelease = PointerEvent.IsTouchEvent() && HasCapture(PointerEvent.GetPointerIndex());
 		ensureMsgf(!bIsFingerCapturedPostRelease, TEXT("Touch pointer was captured while in the process of being released. Since touch pointers cease to exist when the touch ends, this makes no sense."));
 #endif
@@ -1051,9 +1064,10 @@ void FSlateUser::NotifyPointerReleased(const FPointerEvent& PointerEvent, const
 		// fingers no longer in contact with the screen.
 		ReleaseCapture(PointerIdx);
 
-		// When touch pointers are released, they also effectively cease to exist
+		// When touch pointers are released, they also effectively cease to exist until a touch begins again
 		PointerPositionsByIndex.Remove(PointerIdx);
 		PreviousPointerPositionsByIndex.Remove(PointerIdx);
+		WidgetsUnderPointerLastEventByIndex.Remove(PointerIdx);
 	}
 }
 
diff --git a/Engine/Source/Runtime/Slate/Public/Framework/Application/SlateUser.h b/Engine/Source/Runtime/Slate/Public/Framework/Application/SlateUser.h
--- a/Engine/Source/Runtime/Slate/Public/Framework/Application/SlateUser.h
+++ b/Engine/Source/Runtime/Slate/Public/Framework/Application/SlateUser.h
@@ -124,7 +124,7 @@ SLATE_SCOPE:
 	
 	FORCEINLINE bool HasValidFocusPath() const { return WeakFocusPath.IsValid(); }
 	FORCEINLINE const FWeakWidgetPath& GetWeakFocusPath() const { return WeakFocusPath; }
-
+	
 	FORCEINLINE TSharedRef<FWidgetPath> GetFocusPath() const
 	{
 		if (!StrongFocusPath.IsValid())
@@ -154,6 +154,9 @@ SLATE_SCOPE:
 	void FinishFrame();
 	void NotifyWindowDestroyed(TSharedRef<SWindow> DestroyedWindow);
 
+	bool IsTouchPointerActive(int32 TouchPointerIndex) const;
+
+	void NotifyTouchStarted(const FPointerEvent& TouchEvent);
 	void NotifyPointerMoveBegin(const FPointerEvent& PointerEvent);
 	void NotifyPointerMoveComplete(const FPointerEvent& PointerEvent, const FWidgetPath& WidgetsUnderPointer);
 	void NotifyPointerReleased(const FPointerEvent& PointerEvent, const FWidgetPath& WidgetsUnderCursor, TSharedPtr<FDragDropOperation> DroppedContent, bool bWasHandled);
