Copied cl 10905945

Fixed assert from unset optional access

#rnx

diff --git a/Engine/Source/Runtime/Slate/Private/Widgets/Input/SRotatorInputBox.cpp b/Engine/Source/Runtime/Slate/Private/Widgets/Input/SRotatorInputBox.cpp
--- a/Engine/Source/Runtime/Slate/Private/Widgets/Input/SRotatorInputBox.cpp
+++ b/Engine/Source/Runtime/Slate/Private/Widgets/Input/SRotatorInputBox.cpp
@@ -21,10 +21,6 @@ void SRotatorInputBox::Construct( const SRotatorInputBox::FArguments& InArgs )
 	const FLinearColor LabelColorY = InArgs._bColorAxisLabels ? SNumericEntryBox<float>::GreenLabelBackgroundColor : FLinearColor(0.0f, 0.0f, 0.0f, .5f);
 	const FLinearColor LabelColorZ = InArgs._bColorAxisLabels ? SNumericEntryBox<float>::BlueLabelBackgroundColor : FLinearColor(0.0f, 0.0f, 0.0f, .5f);
 
-	TAttribute<TOptional<float>> RollValue = InArgs._Roll;
-	TAttribute<TOptional<float>> PitchValue = InArgs._Pitch;
-	TAttribute<TOptional<float>> YawValue = InArgs._Yaw;
-
 	this->ChildSlot
 	[
 		SNew(SHorizontalBox)
@@ -49,10 +45,13 @@ void SRotatorInputBox::Construct( const SRotatorInputBox::FArguments& InArgs )
 			.OnBeginSliderMovement( InArgs._OnBeginSliderMovement )
 			.OnEndSliderMovement( InArgs._OnEndSliderMovement )
 			.UndeterminedString( LOCTEXT("MultipleValues", "Multiple Values") )
-			.ToolTipText(MakeAttributeLambda([RollValue]
+			.ToolTipText_Lambda([RollAttr = InArgs._Roll]
 			{
-				return FText::Format(LOCTEXT("Roll_ToolTip", "Roll Value = {0}"), RollValue.Get().GetValue());
-			}))
+				const TOptional<float>& RollValue = RollAttr.Get();
+				return RollValue.IsSet() 
+					? FText::Format(LOCTEXT("Roll_ToolTip", "Roll Value = {0}"), RollValue.GetValue())
+					: LOCTEXT("MultipleValues", "Multiple Values");
+			})
 			.TypeInterface(InArgs._TypeInterface)
 		]
 		+SHorizontalBox::Slot()
@@ -76,10 +75,13 @@ void SRotatorInputBox::Construct( const SRotatorInputBox::FArguments& InArgs )
 			.OnBeginSliderMovement( InArgs._OnBeginSliderMovement )
 			.OnEndSliderMovement( InArgs._OnEndSliderMovement )
 			.UndeterminedString( LOCTEXT("MultipleValues", "Multiple Values") )
-			.ToolTipText(MakeAttributeLambda([PitchValue]
+			.ToolTipText_Lambda([PitchAttr = InArgs._Pitch]
 			{
-				return FText::Format(LOCTEXT("Pitch_ToolTip", "Pitch Value = {0}"), PitchValue.Get().GetValue());
-			}))
+				const TOptional<float>& PitchValue = PitchAttr.Get();
+				return PitchValue.IsSet()
+					? FText::Format(LOCTEXT("Pitch_ToolTip", "Pitch Value = {0}"), PitchValue.GetValue())
+					: LOCTEXT("MultipleValues", "Multiple Values");
+			})
 			.TypeInterface(InArgs._TypeInterface)
 		]
 		+SHorizontalBox::Slot()
@@ -103,10 +105,13 @@ void SRotatorInputBox::Construct( const SRotatorInputBox::FArguments& InArgs )
 			.OnBeginSliderMovement( InArgs._OnBeginSliderMovement )
 			.OnEndSliderMovement( InArgs._OnEndSliderMovement )
 			.UndeterminedString( LOCTEXT("MultipleValues", "Multiple Values") )
-			.ToolTipText(MakeAttributeLambda([YawValue]
+			.ToolTipText_Lambda([YawAttr = InArgs._Yaw]
 			{
-				return FText::Format(LOCTEXT("Yaw_ToolTip", "Yaw Value = {0}"), YawValue.Get().GetValue());
-			}))
+				const TOptional<float>& YawValue = YawAttr.Get();
+				return YawValue.IsSet()
+					? FText::Format(LOCTEXT("Yaw_ToolTip", "Yaw Value = {0}"), YawValue.GetValue())
+					: LOCTEXT("MultipleValues", "Multiple Values");
+			})
 			.TypeInterface(InArgs._TypeInterface)
 		]
 	];
