Back out changelist 10635678 to avoid heap-buffer-overflow, revert Fix for font size multiplier on text with Decorator
#rnx

diff --git a/Engine/Source/Runtime/Slate/Private/Framework/Text/RichTextLayoutMarshaller.cpp b/Engine/Source/Runtime/Slate/Private/Framework/Text/RichTextLayoutMarshaller.cpp
--- a/Engine/Source/Runtime/Slate/Private/Framework/Text/RichTextLayoutMarshaller.cpp
+++ b/Engine/Source/Runtime/Slate/Private/Framework/Text/RichTextLayoutMarshaller.cpp
@@ -8,7 +8,6 @@
 #include "Framework/Text/RichTextMarkupProcessing.h"
 #include "Framework/Text/SlateTextLayout.h"
 #include "Framework/Text/SlateTextUnderlineLineHighlighter.h"
-#include "Templates/SharedPointer.h"
 
 #if WITH_FANCY_TEXT
 
@@ -179,7 +178,14 @@ void FRichTextLayoutMarshaller::AppendRunsForText(
 		ModelRange.EndIndex = InOutModelText->Len();
 
 		// Create run.
-		Run = FSlateTextRun::Create(RunInfo, InOutModelText, *TextBlockStyle, ModelRange);
+		TSharedPtr< FSlateTextRun > SlateTextRun = FSlateTextRun::Create(RunInfo, InOutModelText, *TextBlockStyle, ModelRange);
+
+		if (SlateTextRun)
+		{
+			// Apply the FontSizeMultiplier at the style use by the IRun
+			SlateTextRun->ApplyFontSizeMultiplierOnTextStyle(FontSizeMultiplier);
+		}
+		Run = SlateTextRun;
 
 		if (!TextBlockStyle->UnderlineBrush.GetResourceName().IsNone())
 		{
@@ -206,15 +212,6 @@ void FRichTextLayoutMarshaller::AppendRunsForText(
 		}
 	}
 
-	if (Run.IsValid())
-	{
-		TSharedRef<FSlateTextRun> SlateTextRun = StaticCastSharedRef<FSlateTextRun>(Run.ToSharedRef());
-		if (SlateTextRun.Get().DoesSharedInstanceExist())
-		{
-			SlateTextRun.Get().ApplyFontSizeMultiplierOnTextStyle(FontSizeMultiplier);
-		}
-	}
-
 	Runs.Add(Run.ToSharedRef());
 }
 
