Don't attempt to bring a tab to the foreground if there isn't any tabs left after the one being dragged is removed.

diff --git a/Engine/Source/Runtime/Slate/Private/Framework/Docking/SDockingTabWell.cpp b/Engine/Source/Runtime/Slate/Private/Framework/Docking/SDockingTabWell.cpp
--- a/Engine/Source/Runtime/Slate/Private/Framework/Docking/SDockingTabWell.cpp
+++ b/Engine/Source/Runtime/Slate/Private/Framework/Docking/SDockingTabWell.cpp
@@ -352,15 +352,17 @@ void SDockingTabWell::OnDragLeave( const FDragDropEvent& DragDropEvent )
 			// showing it in the TabWell.
 			this->TabBeingDraggedPtr.Reset();
 
-			// Also stop showing its content; switch to the last tab that was active.
-			BringTabToFront( FMath::Max(LastForegroundTabIndex-1, 0) );
-
 			// We may have removed the last tab that this DockNode had.
 			if ( Tabs.Num() == 0 )
 			{
 				// Let the DockNode know that it is no longer needed.
 				ParentTabStack->OnLastTabRemoved();
 			}
+			else
+			{
+				// Also stop showing its content; switch to the last tab that was active.
+				BringTabToFront(FMath::Max(LastForegroundTabIndex - 1, 0));
+			}
 
 			GetDockArea()->CleanUp( SDockingNode::TabRemoval_DraggedOut );
 
