Fix Mac editor menus not opening

diff --git a/Engine/Source/Runtime/Slate/Private/Framework/MultiBox/Mac/MacMenu.cpp b/Engine/Source/Runtime/Slate/Private/Framework/MultiBox/Mac/MacMenu.cpp
--- a/Engine/Source/Runtime/Slate/Private/Framework/MultiBox/Mac/MacMenu.cpp
+++ b/Engine/Source/Runtime/Slate/Private/Framework/MultiBox/Mac/MacMenu.cpp
@@ -780,7 +780,10 @@ void FSlateMacMenu::UpdateCachedState()
 							// Have the menu fill its contents
 							MenuEntryBlock->EntryBuilder.ExecuteIfBound(MenuBuilder);
 						}
-						Widget = MenuBuilder.MakeWidget();
+ 
+                        // Use INT_MAX for MaxHeight to so that SMultiBoxWidget is returned by MakeWidget instead of SVerticalBox
+                        const int32 MaxHeight = INT_MAX;
+						Widget = MenuBuilder.MakeWidget(/*InMakeMultiBoxBuilderOverride=*/ nullptr, MaxHeight);
 					}
 
 					if (Widget->GetType() == FName(TEXT("SMultiBoxWidget")))
@@ -789,7 +792,9 @@ void FSlateMacMenu::UpdateCachedState()
 					}
 					else
 					{
-						UE_LOG(LogMac, Warning, TEXT("Unsupported type of menu widget in FSlateMacMenu::UpdateCachedState(): %s"), *Widget->GetType().ToString());
+                        const FName ActionName = MenuEntryBlock->GetAction().IsValid() ? MenuEntryBlock->GetAction()->GetCommandName() : NAME_None;
+						UE_LOG(LogMac, Warning, TEXT("Unsupported type of menu widget in FSlateMacMenu::UpdateCachedState(): %s, %s, %s"),
+                               *Widget->GetType().ToString(), *MenuEntryBlock->GetExtensionHook().ToString(), *ActionName.ToString());
 					}
 				}
 			}
