Copied cl 10883414

Fix to SVectorInput ToolTip that triggers a check on multiple values.

#rnx

diff --git a/Engine/Source/Runtime/Slate/Private/Widgets/Input/SVectorInputBox.cpp b/Engine/Source/Runtime/Slate/Private/Widgets/Input/SVectorInputBox.cpp
--- a/Engine/Source/Runtime/Slate/Private/Widgets/Input/SVectorInputBox.cpp
+++ b/Engine/Source/Runtime/Slate/Private/Widgets/Input/SVectorInputBox.cpp
@@ -54,7 +54,11 @@ void SVectorInputBox::ConstructX( const FArguments& InArgs, TSharedRef<SHorizont
 		.OnValueCommitted( InArgs._OnXCommitted )
 		.ToolTipText(MakeAttributeLambda([Value]
 		{
-			return FText::Format(LOCTEXT("X_ToolTip", "X Value = {0}"), Value.Get().GetValue());
+			if (Value.Get().IsSet())
+			{
+				return FText::Format(LOCTEXT("X_ToolTip", "X Value = {0}"), Value.Get().GetValue());
+			}
+			return LOCTEXT("MultipleValues", "Multiple Values");
 		}))
 		.UndeterminedString( LOCTEXT("MultipleValues", "Multiple Values") )
 		.LabelPadding(0)
@@ -102,7 +106,11 @@ void SVectorInputBox::ConstructY( const FArguments& InArgs, TSharedRef<SHorizont
 		.OnValueCommitted( InArgs._OnYCommitted )
 		.ToolTipText(MakeAttributeLambda([Value]
 		{
-			return FText::Format(LOCTEXT("Y_ToolTip", "Y Value = {0}"), Value.Get().GetValue());
+			if (Value.Get().IsSet())
+			{
+				return FText::Format(LOCTEXT("Y_ToolTip", "Y Value = {0}"), Value.Get().GetValue());
+			}
+			return LOCTEXT("MultipleValues", "Multiple Values");
 		}))
 		.UndeterminedString( LOCTEXT("MultipleValues", "Multiple Values") )
 		.LabelPadding(0)
@@ -149,8 +157,12 @@ void SVectorInputBox::ConstructZ( const FArguments& InArgs, TSharedRef<SHorizont
 		.OnValueChanged( InArgs._OnZChanged )
 		.OnValueCommitted( InArgs._OnZCommitted )
 		.ToolTipText(MakeAttributeLambda([Value]
-		{
-			return FText::Format(LOCTEXT("Z_ToolTip", "Z Value = {0}"), Value.Get().GetValue());
+		{	
+			if (Value.Get().IsSet())
+			{
+				return FText::Format(LOCTEXT("Z_ToolTip", "Z Value = {0}"), Value.Get().GetValue());
+			}
+			return LOCTEXT("MultipleValues", "Multiple Values");
 		}))
 		.UndeterminedString( LOCTEXT("MultipleValues", "Multiple Values") )
 		.LabelPadding(0)
