TListTypeTraits will no longer reference raw FProperty pointers but will be using TFieldPath instead to prevent crashes in attributes editor

diff --git a/Engine/Plugins/Runtime/GameplayAbilities/Source/GameplayAbilitiesEditor/Private/GameplayTagBlueprintPropertyMappingDetails.cpp b/Engine/Plugins/Runtime/GameplayAbilities/Source/GameplayAbilitiesEditor/Private/GameplayTagBlueprintPropertyMappingDetails.cpp
--- a/Engine/Plugins/Runtime/GameplayAbilities/Source/GameplayAbilitiesEditor/Private/GameplayTagBlueprintPropertyMappingDetails.cpp
+++ b/Engine/Plugins/Runtime/GameplayAbilities/Source/GameplayAbilitiesEditor/Private/GameplayTagBlueprintPropertyMappingDetails.cpp
@@ -84,7 +84,7 @@ void FGameplayTagBlueprintPropertyMappingDetails::CustomizeHeader(TSharedRef<IPr
 	}
 
 	// Sort the options list alphabetically.
-	PropertyOptions.StableSort([](const FProperty& A, const FProperty& B) { return (A.GetName() < B.GetName()); });
+	PropertyOptions.StableSort([](const TFieldPath<FProperty>& A, const TFieldPath<FProperty>& B) { return (A->GetName() < B->GetName()); });
 
 	if ((FoundProperty == nullptr) || (FoundProperty != SelectedProperty) || (FoundProperty->GetName() != SelectedPropertyName))
 	{
@@ -124,7 +124,7 @@ void FGameplayTagBlueprintPropertyMappingDetails::CustomizeChildren(TSharedRef<I
 			.HAlign(HAlign_Fill)
 			.Padding(0.0f, 0.0f, 2.0f, 0.0f)
 			[
-				SNew(SComboBox< FProperty* >)
+				SNew(SComboBox< TFieldPath<FProperty> >)
 				.OptionsSource(&PropertyOptions)
 				.OnGenerateWidget(this, &FGameplayTagBlueprintPropertyMappingDetails::GeneratePropertyWidget)
 				.OnSelectionChanged(this, &FGameplayTagBlueprintPropertyMappingDetails::OnChangeProperty)
@@ -139,7 +139,7 @@ void FGameplayTagBlueprintPropertyMappingDetails::CustomizeChildren(TSharedRef<I
 		];
 }
 
-void FGameplayTagBlueprintPropertyMappingDetails::OnChangeProperty(FProperty* ItemSelected, ESelectInfo::Type SelectInfo)
+void FGameplayTagBlueprintPropertyMappingDetails::OnChangeProperty(TFieldPath<FProperty> ItemSelected, ESelectInfo::Type SelectInfo)
 {
 	if (NamePropertyHandle.IsValid() && GuidPropertyHandle.IsValid())
 	{
@@ -158,7 +158,7 @@ void FGameplayTagBlueprintPropertyMappingDetails::OnChangeProperty(FProperty* It
 	}
 }
 
-FGuid FGameplayTagBlueprintPropertyMappingDetails::GetPropertyGuid(FProperty* Property) const
+FGuid FGameplayTagBlueprintPropertyMappingDetails::GetPropertyGuid(TFieldPath<FProperty> Property) const
 {
 	FGuid Guid;
 
@@ -170,14 +170,14 @@ FGuid FGameplayTagBlueprintPropertyMappingDetails::GetPropertyGuid(FProperty* Pr
 	return Guid;
 }
 
-FString FGameplayTagBlueprintPropertyMappingDetails::GetPropertyName(FProperty* Property) const
+FString FGameplayTagBlueprintPropertyMappingDetails::GetPropertyName(TFieldPath<FProperty> Property) const
 {
-	return (Property ? Property->GetName() : TEXT("None"));
+	return (Property != nullptr ? Property->GetName() : TEXT("None"));
 }
 
-TSharedRef<SWidget> FGameplayTagBlueprintPropertyMappingDetails::GeneratePropertyWidget(FProperty* Property)
+TSharedRef<SWidget> FGameplayTagBlueprintPropertyMappingDetails::GeneratePropertyWidget(TFieldPath<FProperty> Property)
 {
-	return SNew(STextBlock).Text(FText::FromString(GetPropertyName(Property)));
+	return SNew(STextBlock).Text(FText::FromString(GetPropertyName(Property.Get())));
 }
 
 FText FGameplayTagBlueprintPropertyMappingDetails::GetSelectedValueText() const
diff --git a/Engine/Plugins/Runtime/GameplayAbilities/Source/GameplayAbilitiesEditor/Private/GameplayTagBlueprintPropertyMappingDetails.h b/Engine/Plugins/Runtime/GameplayAbilities/Source/GameplayAbilitiesEditor/Private/GameplayTagBlueprintPropertyMappingDetails.h
--- a/Engine/Plugins/Runtime/GameplayAbilities/Source/GameplayAbilitiesEditor/Private/GameplayTagBlueprintPropertyMappingDetails.h
+++ b/Engine/Plugins/Runtime/GameplayAbilities/Source/GameplayAbilitiesEditor/Private/GameplayTagBlueprintPropertyMappingDetails.h
@@ -2,6 +2,7 @@
 
 #pragma once
 
+#include "UObject/FieldPath.h"
 #include "IPropertyTypeCustomization.h"
 
 
@@ -22,12 +23,12 @@ public:
 
 protected:
 
-	void OnChangeProperty(FProperty* ItemSelected, ESelectInfo::Type SelectInfo);
+	void OnChangeProperty(TFieldPath<FProperty> ItemSelected, ESelectInfo::Type SelectInfo);
 
-	FGuid GetPropertyGuid(FProperty* Property) const;
-	FString GetPropertyName(FProperty* Property) const;
+	FGuid GetPropertyGuid(TFieldPath<FProperty> Property) const;
+	FString GetPropertyName(TFieldPath<FProperty> Property) const;
 
-	TSharedRef<SWidget> GeneratePropertyWidget(FProperty* Property);
+	TSharedRef<SWidget> GeneratePropertyWidget(TFieldPath<FProperty> Property);
 
 	FText GetSelectedValueText() const;
 
@@ -36,7 +37,7 @@ protected:
 	TSharedPtr<IPropertyHandle> NamePropertyHandle;
 	TSharedPtr<IPropertyHandle> GuidPropertyHandle;
 
-	TArray<FProperty*> PropertyOptions;
+	TArray<TFieldPath<FProperty>> PropertyOptions;
 
-	FProperty* SelectedProperty;
+	TFieldPath<FProperty> SelectedProperty;
 };
diff --git a/Engine/Source/Runtime/CoreUObject/Public/UObject/FieldPath.h b/Engine/Source/Runtime/CoreUObject/Public/UObject/FieldPath.h
--- a/Engine/Source/Runtime/CoreUObject/Public/UObject/FieldPath.h
+++ b/Engine/Source/Runtime/CoreUObject/Public/UObject/FieldPath.h
@@ -93,6 +93,7 @@ public:
 		InitialFieldClass = nullptr;
 #endif // WITH_EDITORONLY_DATA
 		ResolvedFieldOwner = -1;
+		SerialNumber = -1;
 	}
 
 	/**
diff --git a/Engine/Source/Runtime/Slate/Public/Framework/Views/TableViewTypeTraits.h b/Engine/Source/Runtime/Slate/Public/Framework/Views/TableViewTypeTraits.h
--- a/Engine/Source/Runtime/Slate/Public/Framework/Views/TableViewTypeTraits.h
+++ b/Engine/Source/Runtime/Slate/Public/Framework/Views/TableViewTypeTraits.h
@@ -5,6 +5,7 @@
 #include "CoreMinimal.h"
 #include "UObject/UObjectBase.h"
 #include "UObject/GCObject.h"
+#include "UObject/FieldPath.h"
 
 class ITableRow;
 struct FSparseItemInfo;
@@ -79,7 +80,7 @@ struct TIsValidListItem<TWeakObjectPtr<T>>
 	};
 };
 template <typename T>
-struct TIsValidListItem<T*, typename TEnableIf<TPointerIsConvertibleFromTo<T, FField>::Value>::Type>
+struct TIsValidListItem<TFieldPath<T>>
 {
 	enum
 	{
@@ -94,7 +95,7 @@ struct TIsValidListItem<T*, typename TEnableIf<TPointerIsConvertibleFromTo<T, FF
  */
 template <typename T, typename Enable=void> struct TListTypeTraits
 {
-	static_assert(TIsValidListItem<T>::Value, "Item type T must be a UObjectBase pointer, FField pointer, TSharedRef, or TSharedPtr.");
+	static_assert(TIsValidListItem<T>::Value, "Item type T must be a UObjectBase pointer, TFieldPath, TSharedRef, or TSharedPtr.");
 };
 
 
@@ -434,30 +435,26 @@ public:
 
 
 /**
- * Lists of pointer types only work if the pointers are deriving from UObject*.
- * In addition to testing and setting the pointers to null, Lists of UObjects
- * will serialize the objects they are holding onto.
+ * Pointer-related functionality (e.g. setting to null, testing for null) specialized for TFieldPaths.
  */
-template <typename T>
-struct TListTypeTraits<T*, typename TEnableIf<TPointerIsConvertibleFromTo<T, FField>::Value>::Type>
+template <typename T> struct TListTypeTraits< TFieldPath<T> >
 {
 public:
-	typedef T* NullableType;
+	typedef TFieldPath<T> NullableType;
 
-	using MapKeyFuncs = TDefaultMapHashableKeyFuncs<T*, TSharedRef<ITableRow>, false>;
-	using MapKeyFuncsSparse = TDefaultMapHashableKeyFuncs<T*, FSparseItemInfo, false>;
-	using SetKeyFuncs = DefaultKeyFuncs<T*>;
+	using MapKeyFuncs = TDefaultMapHashableKeyFuncs<TFieldPath<T>, TSharedRef<ITableRow>, false>;
+	using MapKeyFuncsSparse = TDefaultMapHashableKeyFuncs<TFieldPath<T>, FSparseItemInfo, false>;
+	using SetKeyFuncs = DefaultKeyFuncs< TFieldPath<T> >;
 
 	template<typename U>
 	static void AddReferencedObjects(FReferenceCollector& Collector,
-		TArray<T*>& ItemsWithGeneratedWidgets,
-		TSet<T*>& SelectedItems,
-		TMap< const U*, T* >& WidgetToItemMap)
+		TArray< TFieldPath<T> >& ItemsWithGeneratedWidgets,
+		TSet< TFieldPath<T> >& SelectedItems,
+		TMap< const U*, TFieldPath<T> >& WidgetToItemMap)
 	{
-		// Serialize generated items
-		for (T* Item : ItemsWithGeneratedWidgets)
+		for (TFieldPath<T>& Item : ItemsWithGeneratedWidgets)
 		{
-			if (Item)
+			if (Item != nullptr)
 			{
 				Item->AddReferencedObjects(Collector);
 			}
@@ -467,34 +464,47 @@ public:
 		// Also, we cannot AddReferencedObject to the Keys of the ItemToWidgetMap or we end up with keys being set to 0 when the UObject is destroyed which generate an invalid id in the map.
 		for (auto& It : WidgetToItemMap)
 		{
-			if (It.Value)
+			if (It.Value != nullptr)
 			{
 				It.Value->AddReferencedObjects(Collector);
 			}
 		}
 
 		// Serialize the selected items
-		for (T* Item : SelectedItems)
+		for (TFieldPath<T>& Item : SelectedItems)
 		{
-			if (Item)
+			if (Item != nullptr)
 			{
 				Item->AddReferencedObjects(Collector);
 			}
 		}
 	}
 
-	static bool IsPtrValid(T* InPtr) { return InPtr != nullptr; }
+	static bool IsPtrValid(const TFieldPath<T>& InPtr)
+	{
+		return !!InPtr.Get();
+	}
 
-	static void ResetPtr(T*& InPtr) { InPtr = nullptr; }
+	static void ResetPtr(TFieldPath<T>& InPtr)
+	{
+		InPtr.Reset();
+	}
 
-	static T* MakeNullPtr() { return nullptr; }
+	static TFieldPath<T> MakeNullPtr()
+	{
+		return nullptr;
+	}
 
-	static T* NullableItemTypeConvertToItemType(T* InPtr) { return InPtr; }
+	static TFieldPath<T> NullableItemTypeConvertToItemType(const TFieldPath<T>& InPtr)
+	{
+		return InPtr;
+	}
 
-	static FString DebugDump(T* InPtr)
+	static FString DebugDump(TFieldPath<T> InPtr)
 	{
-		return InPtr ? FString::Printf(TEXT("0x%08x [%s]"), InPtr, *InPtr->GetName()) : FString(TEXT("nullptr"));
+		T* ObjPtr = InPtr.Get();
+		return ObjPtr ? FString::Printf(TEXT("0x%08x [%s]"), ObjPtr, *ObjPtr->GetName()) : FString(TEXT("nullptr"));
 	}
 
-	typedef FGCObject SerializerType;
+	class SerializerType {};
 };
\ No newline at end of file
