Fix menu search not starting when max height is set
#rnx

diff --git a/Engine/Source/Runtime/Slate/Private/Framework/MultiBox/MultiBox.cpp b/Engine/Source/Runtime/Slate/Private/Framework/MultiBox/MultiBox.cpp
--- a/Engine/Source/Runtime/Slate/Private/Framework/MultiBox/MultiBox.cpp
+++ b/Engine/Source/Runtime/Slate/Private/Framework/MultiBox/MultiBox.cpp
@@ -332,7 +332,7 @@ void FMultiBox::InsertCustomMultiBlock( TSharedRef<const FMultiBlock> InBlock, i
  *
  * @return  MultiBox widget object
  */
-TSharedRef< SMultiBoxWidget > FMultiBox::MakeWidget( bool bSearchable, FOnMakeMultiBoxBuilderOverride* InMakeMultiBoxBuilderOverride /* = nullptr */ )
+TSharedRef< SMultiBoxWidget > FMultiBox::MakeWidget( bool bSearchable, FOnMakeMultiBoxBuilderOverride* InMakeMultiBoxBuilderOverride /* = nullptr */, TAttribute<float> InMaxHeight )
 {	
 	TSharedRef< SMultiBoxWidget > NewMultiBoxWidget =
 		SNew( SMultiBoxWidget );
@@ -343,6 +343,9 @@ TSharedRef< SMultiBoxWidget > FMultiBox::MakeWidget( bool bSearchable, FOnMakeMu
 	// Assign ourselves to the MultiBox widget
 	NewMultiBoxWidget->SetMultiBox( AsShared() );
 
+	// Set the maximum height the MultiBox widget should be
+	NewMultiBoxWidget->SetMaxHeight( InMaxHeight );
+
 	if( (InMakeMultiBoxBuilderOverride != nullptr) && (InMakeMultiBoxBuilderOverride->IsBound()) )
 	{
 		TSharedRef<FMultiBox> ThisMultiBox = AsShared();
@@ -510,22 +513,22 @@ TSharedRef<ITableRow> SMultiBoxWidget::GenerateTiles(TSharedPtr<SWidget> Item, c
 
 float SMultiBoxWidget::GetItemWidth() const
 {
-	float MaxWidth = 0;
+	float MaxItemWidth = 0;
 	for (int32 i = 0; i < TileViewWidgets.Num(); ++i)
 	{
-		MaxWidth = FMath::Max(TileViewWidgets[i]->GetDesiredSize().X, MaxWidth);
+		MaxItemWidth = FMath::Max(TileViewWidgets[i]->GetDesiredSize().X, MaxItemWidth);
 	}
-	return MaxWidth;
+	return MaxItemWidth;
 }
 
 float SMultiBoxWidget::GetItemHeight() const
 {
-	float MaxHeight = 0;
+	float MaxItemHeight = 0;
 	for (int32 i = 0; i < TileViewWidgets.Num(); ++i)
 	{
-		MaxHeight = FMath::Max(TileViewWidgets[i]->GetDesiredSize().Y, MaxHeight);
+		MaxItemHeight = FMath::Max(TileViewWidgets[i]->GetDesiredSize().Y, MaxItemHeight);
 	}
-	return MaxHeight;
+	return MaxItemHeight;
 }
 
 bool SMultiBoxWidget::IsBlockBeingDragged( TSharedPtr<const FMultiBlock> Block ) const
@@ -745,12 +748,29 @@ void SMultiBoxWidget::BuildMultiBoxWidget()
 		break;
 	case EMultiBoxType::Menu:
 		{
-			// wrap menu content in a scrollbox to support vertical scrolling if needed
-			MainWidget = SNew(SScrollBox)
-				+SScrollBox::Slot()
-				[
-					SAssignNew( VerticalBox, SVerticalBox )
-				];
+			if (MaxHeight.IsSet())
+			{
+				MainWidget = SNew(SVerticalBox)
+					+ SVerticalBox::Slot()
+					.MaxHeight(MaxHeight)
+					[
+						// wrap menu content in a scrollbox to support vertical scrolling if needed
+						SNew(SScrollBox)
+						+ SScrollBox::Slot()
+						[
+							SAssignNew(VerticalBox, SVerticalBox)
+						]
+					];
+			}
+			else
+			{
+				// wrap menu content in a scrollbox to support vertical scrolling if needed
+				MainWidget = SNew(SScrollBox)
+					+ SScrollBox::Slot()
+					[
+						SAssignNew(VerticalBox, SVerticalBox)
+					];
+			}
 		}
 		break;
 	}
diff --git a/Engine/Source/Runtime/Slate/Private/Framework/MultiBox/MultiBoxBuilder.cpp b/Engine/Source/Runtime/Slate/Private/Framework/MultiBox/MultiBoxBuilder.cpp
--- a/Engine/Source/Runtime/Slate/Private/Framework/MultiBox/MultiBoxBuilder.cpp
+++ b/Engine/Source/Runtime/Slate/Private/Framework/MultiBox/MultiBoxBuilder.cpp
@@ -175,22 +175,12 @@ void FBaseMenuBuilder::AddMenuEntry( const FUIAction& UIAction, const TSharedRef
 TSharedRef< class SWidget > FMenuBuilder::MakeWidget( FMultiBox::FOnMakeMultiBoxBuilderOverride* InMakeMultiBoxBuilderOverride /* = nullptr */, uint32 MaxHeight)
 {
 	// Make menu builders searchable (by default)
-	TSharedRef< class SWidget > MenuWidget = MultiBox->MakeWidget(bSearchable, InMakeMultiBoxBuilderOverride);
+	TAttribute<float> MaxHeightAttribute;
 	if (MaxHeight < INT_MAX)
 	{
-		TSharedRef<SWidget> ConstrainedMenu = SNew(SVerticalBox)
-			+ SVerticalBox::Slot()
-			.MaxHeight((float)MaxHeight)
-			[
-				SNew(SScrollBox)
-				+ SScrollBox::Slot()
-				[
-					MenuWidget
-				]
-			];
-		return ConstrainedMenu;
+		MaxHeightAttribute.Set((float)MaxHeight);
 	}
-	return MenuWidget;
+	return MultiBox->MakeWidget(bSearchable, InMakeMultiBoxBuilderOverride, MaxHeightAttribute);
 }
 
 void FMenuBuilder::BeginSection( FName InExtensionHook, const TAttribute< FText >& InHeadingText )
diff --git a/Engine/Source/Runtime/Slate/Public/Framework/MultiBox/MultiBox.h b/Engine/Source/Runtime/Slate/Public/Framework/MultiBox/MultiBox.h
--- a/Engine/Source/Runtime/Slate/Public/Framework/MultiBox/MultiBox.h
+++ b/Engine/Source/Runtime/Slate/Public/Framework/MultiBox/MultiBox.h
@@ -296,7 +296,7 @@ public:
 	 *
 	 * @return  MultiBox widget object
 	 */
-	TSharedRef< class SMultiBoxWidget > MakeWidget( bool bSearchable, FOnMakeMultiBoxBuilderOverride* InMakeMultiBoxBuilderOverride = nullptr );
+	TSharedRef< class SMultiBoxWidget > MakeWidget( bool bSearchable, FOnMakeMultiBoxBuilderOverride* InMakeMultiBoxBuilderOverride = nullptr, TAttribute<float> InMaxHeight = TAttribute<float>() );
 
 
 	/**
@@ -583,7 +583,12 @@ public:
 	* @return	Whether this block is searchable
 	*/
 	bool GetSearchable() const;
-	
+
+	/**
+	* Sets optional maximum height of widget
+	*/
+	void SetMaxHeight(TAttribute<float> InMaxHeight) { MaxHeight = InMaxHeight; }
+
 	/**
 	 * Builds this MultiBox widget up from the MultiBox associated with it
 	 */
@@ -770,4 +775,7 @@ private:
 
 	/** Whether this multibox can be searched */
 	bool bSearchable;
+
+	/** Optional maximum height of widget */
+	TAttribute<float> MaxHeight;
 };
