Editor TabManager: Check filter before spawning tabs
#rnx

diff --git a/Engine/Source/Runtime/Slate/Private/Framework/Docking/TabManager.cpp b/Engine/Source/Runtime/Slate/Private/Framework/Docking/TabManager.cpp
--- a/Engine/Source/Runtime/Slate/Private/Framework/Docking/TabManager.cpp
+++ b/Engine/Source/Runtime/Slate/Private/Framework/Docking/TabManager.cpp
@@ -1633,6 +1633,11 @@ TSharedRef<FBlacklistNames>& FTabManager::GetTabBlacklist()
 
 bool FTabManager::IsValidTabForSpawning( const FTab& SomeTab ) const
 {
+	if (SomeTab.TabId.TabType != NAME_None && !TabBlacklist->PassesFilter(SomeTab.TabId.TabType))
+	{
+		return false;
+	}
+
 	// Nomad tabs being restored from layouts should not be spawned if the nomad tab is already spawned.
 	TSharedRef<FTabSpawnerEntry>* NomadSpawner = NomadTabSpawner->Find( SomeTab.TabId.TabType );
 	return ( !NomadSpawner || !NomadSpawner->Get().IsSoleTabInstanceSpawned() );
