Adding delegate to receive notifcations whenever the slate application might change input focus

diff --git a/Engine/Source/Runtime/Slate/Private/Framework/Application/SlateApplication.cpp b/Engine/Source/Runtime/Slate/Private/Framework/Application/SlateApplication.cpp
--- a/Engine/Source/Runtime/Slate/Private/Framework/Application/SlateApplication.cpp
+++ b/Engine/Source/Runtime/Slate/Private/Framework/Application/SlateApplication.cpp
@@ -2591,6 +2591,7 @@ bool FSlateApplication::SetUserFocus(FSlateUser& User, const FWidgetPath& InFocu
 #if WITH_SLATE_DEBUGGING
 	FSlateDebugging::BroadcastFocusChanging(FocusEvent, OldFocusedWidgetPath, OldFocusedWidget, NewFocusedWidgetPath, NewFocusedWidget);
 #endif
+	FocusChangingDelegate.Broadcast(FocusEvent, OldFocusedWidgetPath, OldFocusedWidget, NewFocusedWidgetPath, NewFocusedWidget);
 
 	// Notify widgets in the old focus path that focus is changing
 	if (OldFocusedWidgetPath.IsValid())
@@ -6717,4 +6718,4 @@ bool FSlateApplication::InputPreProcessorsHelper::PreProcessInput(TFunctionRef<b
 	ProcessorsPendingAddition.Reset();
 
 	return bShouldExit;
-}
\ No newline at end of file
+}
diff --git a/Engine/Source/Runtime/Slate/Public/Framework/Application/SlateApplication.h b/Engine/Source/Runtime/Slate/Public/Framework/Application/SlateApplication.h
--- a/Engine/Source/Runtime/Slate/Public/Framework/Application/SlateApplication.h
+++ b/Engine/Source/Runtime/Slate/Public/Framework/Application/SlateApplication.h
@@ -446,6 +446,10 @@ public:
 	DECLARE_EVENT_OneParam(FSlateApplication, FOnWindowBeingDestroyed, const SWindow&);
 	FOnWindowBeingDestroyed& OnWindowBeingDestroyed() { return WindowBeingDestroyedEvent; }
 
+	/** Delegate called just before possible focus change */
+	DECLARE_MULTICAST_DELEGATE_FiveParams(FOnFocusChanging, const FFocusEvent&, const FWeakWidgetPath&, const TSharedPtr<SWidget>&, const FWidgetPath&, const TSharedPtr<SWidget>&);
+	FOnFocusChanging& OnFocusChanging() { return FocusChangingDelegate; }
+
 	/** 
 	 * Removes references to FViewportRHI's.  
 	 * This has to be done explicitly instead of using the FRenderResource mechanism because FViewportRHI's are managed by the game thread.
@@ -1809,6 +1813,9 @@ private:
 	/** Delegate for slate Tick during modal dialogs */
 	FOnModalLoopTickEvent ModalLoopTickEvent;
 
+	/** Delegate for when focus might be about to change */
+	FOnFocusChanging FocusChangingDelegate;
+
 	/** Critical section to avoid multiple threads calling Slate Tick when we're synchronizing between the Slate Loading Thread and the Game Thread. */
 	FCriticalSection SlateTickCriticalSection;
 
