Adds optional padding around the label in Toolbars

diff --git a/Engine/Source/Runtime/Slate/Private/Framework/MultiBox/MultiBoxBuilder.cpp b/Engine/Source/Runtime/Slate/Private/Framework/MultiBox/MultiBoxBuilder.cpp
--- a/Engine/Source/Runtime/Slate/Private/Framework/MultiBox/MultiBoxBuilder.cpp
+++ b/Engine/Source/Runtime/Slate/Private/Framework/MultiBox/MultiBoxBuilder.cpp
@@ -459,11 +459,17 @@ void FToolBarBuilder::AddToolBarWidget( TSharedRef<SWidget> InWidget, const TAtt
 
 		+SVerticalBox::Slot()
 		.AutoHeight()
+		.Padding(GetStyleSet()->GetMargin(ISlateStyle::Join( GetStyleName(), ".Label.Padding" )))
 		.HAlign( HAlign_Center )
 		[
 			SNew( STextBlock )
-			// .Visibility_Lambda( [this] () -> EVisibility { return /*LabelVisibility.IsSet() ? LabelVisibility.GetValue() :*/ (bForceSmallIcons || FMultiBoxSettings::UseSmallToolBarIcons.Get()) ? EVisibility::Collapsed : EVisibility::Visible; } ) 
-			.Visibility_Lambda( [] () -> EVisibility { return FMultiBoxSettings::UseSmallToolBarIcons.Get() ? EVisibility::Collapsed : EVisibility::Visible; } ) 
+			.Visibility_Lambda( [ChildWidget] () -> EVisibility {
+				if (FMultiBoxSettings::UseSmallToolBarIcons.Get())
+					return EVisibility::Collapsed;
+
+				return ChildWidget->GetVisibility();
+			})
+			.IsEnabled_Lambda( [ChildWidget] () -> bool { return ChildWidget->IsEnabled(); } )
 			.Text( InLabel )
 			.TextStyle( GetStyleSet(), ISlateStyle::Join( GetStyleName(), ".Label" ) )	// Smaller font for tool tip labels
 			.ShadowOffset( FVector2D::UnitVector )
diff --git a/Engine/Source/Runtime/Slate/Private/Framework/MultiBox/SToolBarButtonBlock.cpp b/Engine/Source/Runtime/Slate/Private/Framework/MultiBox/SToolBarButtonBlock.cpp
--- a/Engine/Source/Runtime/Slate/Private/Framework/MultiBox/SToolBarButtonBlock.cpp
+++ b/Engine/Source/Runtime/Slate/Private/Framework/MultiBox/SToolBarButtonBlock.cpp
@@ -200,6 +200,7 @@ void SToolBarButtonBlock::BuildMultiBlockWidget(const ISlateStyle* StyleSet, con
 
 			// Label text
 			+ SVerticalBox::Slot().AutoHeight()
+			.Padding(StyleSet->GetMargin(ISlateStyle::Join( StyleName, ".Label.Padding" )))
 			.HAlign( HAlign_Center )	// Center the label text horizontally
 			[
 				SNew( STextBlock )
diff --git a/Engine/Source/Runtime/Slate/Private/Framework/MultiBox/SToolBarComboButtonBlock.cpp b/Engine/Source/Runtime/Slate/Private/Framework/MultiBox/SToolBarComboButtonBlock.cpp
--- a/Engine/Source/Runtime/Slate/Private/Framework/MultiBox/SToolBarComboButtonBlock.cpp
+++ b/Engine/Source/Runtime/Slate/Private/Framework/MultiBox/SToolBarComboButtonBlock.cpp
@@ -148,6 +148,7 @@ void SToolBarComboButtonBlock::BuildMultiBlockWidget(const ISlateStyle* StyleSet
 		// Label text
 		+ SVerticalBox::Slot()
 		.AutoHeight()
+		.Padding(StyleSet->GetMargin(ISlateStyle::Join( StyleName, ".Label.Padding" )))
 		.HAlign( HAlign_Center )	// Center the label text horizontally
 		[
 			SNew( STextBlock )
diff --git a/Engine/Source/Runtime/SlateCore/Private/Styling/StyleDefaults.cpp b/Engine/Source/Runtime/SlateCore/Private/Styling/StyleDefaults.cpp
--- a/Engine/Source/Runtime/SlateCore/Private/Styling/StyleDefaults.cpp
+++ b/Engine/Source/Runtime/SlateCore/Private/Styling/StyleDefaults.cpp
@@ -11,5 +11,5 @@ const FSlateFontInfo FStyleDefaults::GetFontInfo(uint16 Size)
 float FStyleDefaults::DefaultFloat;
 FVector2D FStyleDefaults::DefaultFVector2D = FVector2D( 66.0f, 66.0f );
 FLinearColor FStyleDefaults::DefaultColor = FLinearColor( 1, 1, 1 );
-FMargin FStyleDefaults::DefaultMargin = 66.0f;
+FMargin FStyleDefaults::DefaultMargin = 0.0f;
 FSlateSound FStyleDefaults::DefaultSound;
