Fixed slider widget to respect tint from style

diff --git a/Engine/Source/Runtime/Slate/Private/Widgets/Input/SSlider.cpp b/Engine/Source/Runtime/Slate/Private/Widgets/Input/SSlider.cpp
--- a/Engine/Source/Runtime/Slate/Private/Widgets/Input/SSlider.cpp
+++ b/Engine/Source/Runtime/Slate/Private/Widgets/Input/SSlider.cpp
@@ -91,13 +91,15 @@ int32 SSlider::OnPaint( const FPaintArgs& Args, const FGeometry& AllottedGeometr
 	// draw slider bar
 	auto BarTopLeft = FVector2D(SliderStartPoint.X, SliderStartPoint.Y - Style->BarThickness * 0.5f);
 	auto BarSize = FVector2D(SliderEndPoint.X - SliderStartPoint.X, Style->BarThickness);
+	auto BarImage = GetBarImage();
+	auto ThumbImage = GetThumbImage();
 	FSlateDrawElement::MakeBox(
 		OutDrawElements,
 		LayerId,
 		SliderGeometry.ToPaintGeometry(BarTopLeft, BarSize),
-		GetBarImage(),
+		BarImage,
 		DrawEffects,
-		SliderBarColor.Get().GetColor(InWidgetStyle) * InWidgetStyle.GetColorAndOpacityTint()
+		BarImage->GetTint(InWidgetStyle) * SliderBarColor.Get().GetColor(InWidgetStyle) * InWidgetStyle.GetColorAndOpacityTint()
 		);
 
 	++LayerId;
@@ -107,9 +109,9 @@ int32 SSlider::OnPaint( const FPaintArgs& Args, const FGeometry& AllottedGeometr
 		OutDrawElements,
 		LayerId,
 		SliderGeometry.ToPaintGeometry(HandleTopLeftPoint, GetThumbImage()->ImageSize),
-		GetThumbImage(),
+		ThumbImage,
 		DrawEffects,
-		SliderHandleColor.Get().GetColor(InWidgetStyle) * InWidgetStyle.GetColorAndOpacityTint()
+		ThumbImage->GetTint(InWidgetStyle) * SliderHandleColor.Get().GetColor(InWidgetStyle) * InWidgetStyle.GetColorAndOpacityTint()
 	);
 
 	return LayerId;
