Updated editable text widgets to no longer commit text on a repeat enter key event

diff --git a/Engine/Source/Runtime/Slate/Private/Widgets/Text/SlateEditableTextLayout.cpp b/Engine/Source/Runtime/Slate/Private/Widgets/Text/SlateEditableTextLayout.cpp
--- a/Engine/Source/Runtime/Slate/Private/Widgets/Text/SlateEditableTextLayout.cpp
+++ b/Engine/Source/Runtime/Slate/Private/Widgets/Text/SlateEditableTextLayout.cpp
@@ -963,7 +963,7 @@ FReply FSlateEditableTextLayout::HandleKeyDown(const FKeyEvent& InKeyEvent)
 	else if (Key == EKeys::Enter && !OwnerWidget->IsTextReadOnly())
 	{
 		FScopedEditableTextTransaction TextTransaction(*this);
-		HandleCarriageReturn();
+		HandleCarriageReturn(InKeyEvent.IsRepeat());
 		Reply = FReply::Handled();
 	}
 	else if (Key == EKeys::Tab && OwnerWidget->CanTypeCharacter(TEXT('\t')))
@@ -1488,7 +1488,7 @@ bool FSlateEditableTextLayout::HandleTypeChar(const TCHAR InChar)
 	return false;
 }
 
-bool FSlateEditableTextLayout::HandleCarriageReturn()
+bool FSlateEditableTextLayout::HandleCarriageReturn(bool isRepeat)
 {
 	if (OwnerWidget->IsTextReadOnly())
 	{
@@ -1499,6 +1499,11 @@ bool FSlateEditableTextLayout::HandleCarriageReturn()
 	{
 		InsertNewLineAtCursorImpl();
 	}
+	else if (isRepeat)
+	{
+		// Skip committing text on a repeat enter key
+		return false;
+	}
 	else
 	{
 		// Always clear the local undo chain on commit.
diff --git a/Engine/Source/Runtime/Slate/Public/Widgets/Text/SlateEditableTextLayout.h b/Engine/Source/Runtime/Slate/Public/Widgets/Text/SlateEditableTextLayout.h
--- a/Engine/Source/Runtime/Slate/Public/Widgets/Text/SlateEditableTextLayout.h
+++ b/Engine/Source/Runtime/Slate/Public/Widgets/Text/SlateEditableTextLayout.h
@@ -184,7 +184,7 @@ public:
 	bool HandleTypeChar(const TCHAR InChar);
 
 	/** Called to handle a carriage return action acting on the current selection or at the cursor position */
-	bool HandleCarriageReturn();
+	bool HandleCarriageReturn(bool isRepeat);
 
 	/** Are we able to delete the currently selected text? */
 	bool CanExecuteDelete() const;
