Added new commit type for SpinBox when setting the value via code (SetValue). This lets us account for the case where we don't want to fire OnValueComitted, but also don't want to clamp by the slider min/max values.

diff --git a/Engine/Source/Runtime/Slate/Public/Widgets/Input/SSpinBox.h b/Engine/Source/Runtime/Slate/Public/Widgets/Input/SSpinBox.h
--- a/Engine/Source/Runtime/Slate/Public/Widgets/Input/SSpinBox.h
+++ b/Engine/Source/Runtime/Slate/Public/Widgets/Input/SSpinBox.h
@@ -703,7 +703,7 @@ public:
 	void SetValue(const TAttribute<NumericType>& InValueAttribute) 
 	{
 		ValueAttribute = InValueAttribute; 
-		CommitValue(InValueAttribute.Get(), ECommitMethod::CommittedViaTypeIn, ETextCommit::Default);
+		CommitValue(InValueAttribute.Get(), ECommitMethod::CommittedViaCode, ETextCommit::Default);
 	}
 
 	/** See the MinValue attribute */
@@ -839,7 +839,8 @@ protected:
 	{
 		CommittedViaSpin,
 		CommittedViaTypeIn,
-		CommittedViaArrowKey
+		CommittedViaArrowKey,
+		CommittedViaCode
 	};
 
 	/**
