Fixed CIS errors on nightly build

#rnx

diff --git a/Engine/Source/Runtime/Slate/Private/Framework/MultiBox/Mac/MacMenu.cpp b/Engine/Source/Runtime/Slate/Private/Framework/MultiBox/Mac/MacMenu.cpp
--- a/Engine/Source/Runtime/Slate/Private/Framework/MultiBox/Mac/MacMenu.cpp
+++ b/Engine/Source/Runtime/Slate/Private/Framework/MultiBox/Mac/MacMenu.cpp
@@ -57,7 +57,7 @@ static FCriticalSection GMacMenuEntryBlockCS;;
 - (void) dealloc
 {
 	{
-		FScopeLock Lock(&GMacMenuEntryBlockCS);
+		FScopeLock MacMenuEntryBlockLock(&GMacMenuEntryBlockCS);
 		// we save the block to prevent it from being released in the main thread
 		// releasing the block in Main Thread will trigger an assertion in FSlateApplicationBase::Get()
 		GMacMenuEntryBlockDeletionQueue.Add(self.MenuEntryBlock);
@@ -84,7 +84,7 @@ static FCriticalSection GMacMenuEntryBlockCS;;
 	self = [super initWithTitle:@""];
 	[self setDelegate:self];
 	self.MenuEntryBlock = Block;
-	FScopeLock Lock(&GCachedMenuStateCS);
+	FScopeLock CachedMenuStateLock(&GCachedMenuStateCS);
 	GCachedMenuState.Add(self, TSharedPtr<TArray<FMacMenuItemState>>(new TArray<FMacMenuItemState>()));
 	return self;
 }
@@ -106,7 +106,7 @@ static FCriticalSection GMacMenuEntryBlockCS;;
 
 - (void)dealloc
 {
-	FScopeLock Lock(&GCachedMenuStateCS);
+	FScopeLock CachedMenuStateLock(&GCachedMenuStateCS);
 	GCachedMenuState.Remove(self);
 	[super dealloc];
 }
@@ -307,12 +307,12 @@ static FStartupApplicationToMacMenuBinder StaticInitializer;
 void FSlateMacMenu::CleanupOnShutdown()
 {
 	{
-		FScopeLock Lock(&GCachedMenuStateCS);
+		FScopeLock CachedMenuStateLock(&GCachedMenuStateCS);
 		GCachedMenuState.Reset();
 	}
 	{
 		GameThreadCall(^{
-			FScopeLock Lock(&GMacMenuEntryBlockCS);
+			FScopeLock MacMenuEntryBlockLock(&GMacMenuEntryBlockCS);
 			for(auto& Block : GMacMenuEntryBlockDeletionQueue)
 			{
 				Block.Reset();
@@ -529,7 +529,7 @@ void FSlateMacMenu::UpdateWithMultiBox(const TSharedPtr< FMultiBox > MultiBox)
 	FSafeMultiBoxPass* SafeMultiBoxPtr = new FSafeMultiBoxPass;
 	SafeMultiBoxPtr->MultiBox = MultiBox;
 
-	FScopeLock Lock(&GCachedMenuStateCS);
+	FScopeLock CachedMenuStateLock(&GCachedMenuStateCS);
 	GCachedMenuState.Reset();
 
 	MainThreadCall(^{
@@ -587,7 +587,7 @@ void FSlateMacMenu::UpdateWithMultiBox(const TSharedPtr< FMultiBox > MultiBox)
 	}, NSDefaultRunLoopMode, false);
 	// Now we destroy all the menu blocks in the game thread to ensure ~SWidget() is only called in Game Thread
 	GameThreadCall(^{
-		FScopeLock Lock(&GMacMenuEntryBlockCS);
+		FScopeLock MacMenuEntryBlockLock(&GMacMenuEntryBlockCS);
 		for(auto& Block : GMacMenuEntryBlockDeletionQueue)
 		{
 			Block.Reset();
@@ -599,7 +599,7 @@ void FSlateMacMenu::UpdateWithMultiBox(const TSharedPtr< FMultiBox > MultiBox)
 void FSlateMacMenu::UpdateMenu(FMacMenu* Menu)
 {
 	MainThreadCall(^{
-		FScopeLock Lock(&GCachedMenuStateCS);
+		FScopeLock CachedMenuStateLock(&GCachedMenuStateCS);
 
 		FText WindowLabel = NSLOCTEXT("MainMenu", "WindowMenu", "Window");
 		const bool bIsWindowMenu = (WindowLabel.ToString().Compare(FString([Menu title])) == 0);
@@ -755,7 +755,7 @@ void FSlateMacMenu::UpdateCachedState()
 
 	if (bShouldUpdate)
 	{
-		FScopeLock Lock(&GCachedMenuStateCS);
+		FScopeLock CachedMenuStateLock(&GCachedMenuStateCS);
 
 		for (TMap<FMacMenu*, TSharedPtr<TArray<FMacMenuItemState>>>::TIterator It(GCachedMenuState); It; ++It)
 		{
