Modified AsyncTaskNotification to add a prompt state along with an hyperlink for slate based notification
Added a function to pulse basic notification with a specified glow color while they are pending

diff --git a/Engine/Source/Runtime/Core/Private/Misc/AsyncTaskNotification.cpp b/Engine/Source/Runtime/Core/Private/Misc/AsyncTaskNotification.cpp
--- a/Engine/Source/Runtime/Core/Private/Misc/AsyncTaskNotification.cpp
+++ b/Engine/Source/Runtime/Core/Private/Misc/AsyncTaskNotification.cpp
@@ -42,6 +42,16 @@ void FAsyncTaskNotification::SetProgressText(const FText& InProgressText)
 	NotificationImpl->SetProgressText(InProgressText);
 }
 
+void FAsyncTaskNotification::SetPromptText(const FText& InPromptText)
+{
+	NotificationImpl->SetPromptText(InPromptText);
+}
+
+void FAsyncTaskNotification::SetHyperlink(const FSimpleDelegate& InHyperlink, const FText& InHyperlinkText)
+{
+	NotificationImpl->SetHyperlink(InHyperlink, InHyperlinkText);
+}
+
 void FAsyncTaskNotification::SetComplete(const bool bSuccess)
 {
 	NotificationImpl->SetComplete(bSuccess);
@@ -52,6 +62,11 @@ void FAsyncTaskNotification::SetComplete(const FText& InTitleText, const FText&
 	NotificationImpl->SetComplete(InTitleText, InProgressText, bSuccess);
 }
 
+void FAsyncTaskNotification::SetNotificationState(const FAsyncNotificationStateData& InState)
+{
+	NotificationImpl->SetNotificationState(InState);
+}
+
 void FAsyncTaskNotification::SetCanCancel(const TAttribute<bool>& InCanCancel)
 {
 	NotificationImpl->SetCanCancel(InCanCancel);
@@ -67,7 +82,7 @@ void FAsyncTaskNotification::SetKeepOpenOnFailure(const TAttribute<bool>& InKeep
 	NotificationImpl->SetKeepOpenOnFailure(InKeepOpenOnFailure);
 }
 
-bool FAsyncTaskNotification::ShouldCancel() const
+EAsyncTaskNotificationPromptAction FAsyncTaskNotification::GetPromptAction() const
 {
-	return NotificationImpl->ShouldCancel();
+	return NotificationImpl->GetPromptAction();
 }
diff --git a/Engine/Source/Runtime/Core/Private/Misc/CoreAsyncTaskNotificationImpl.cpp b/Engine/Source/Runtime/Core/Private/Misc/CoreAsyncTaskNotificationImpl.cpp
--- a/Engine/Source/Runtime/Core/Private/Misc/CoreAsyncTaskNotificationImpl.cpp
+++ b/Engine/Source/Runtime/Core/Private/Misc/CoreAsyncTaskNotificationImpl.cpp
@@ -14,7 +14,7 @@ FCoreAsyncTaskNotificationImpl::FCoreAsyncTaskNotificationImpl()
 
 FCoreAsyncTaskNotificationImpl::~FCoreAsyncTaskNotificationImpl()
 {
-	checkf(State != ENotificationState::Pending, TEXT("AsyncTaskNotification was still pending when destroyed. Missing call to SetComplete?"));
+	checkf(State != EAsyncTaskNotificationState::Pending, TEXT("AsyncTaskNotification was still pending when destroyed. Missing call to SetComplete?"));
 }
 
 void FCoreAsyncTaskNotificationImpl::Initialize(const FAsyncTaskNotificationConfig& InConfig)
@@ -46,30 +46,68 @@ void FCoreAsyncTaskNotificationImpl::SetProgressText(const FText& InProgressText
 	FScopeLock Lock(&SynchronizationObject);
 
 	ProgressText = InProgressText;
+	UpdateNotification();
+}
+
+void FCoreAsyncTaskNotificationImpl::SetPromptText(const FText& InPromptText)
+{
+	FScopeLock Lock(&SynchronizationObject);
 
+	PromptText = InPromptText;
 	UpdateNotification();
 }
 
-void FCoreAsyncTaskNotificationImpl::SetComplete(const bool bSuccess)
+void FCoreAsyncTaskNotificationImpl::SetHyperlink(const FSimpleDelegate& InHyperlink, const FText& InHyperlinkText)
 {
 	FScopeLock Lock(&SynchronizationObject);
 
-	State = bSuccess ? ENotificationState::Success : ENotificationState::Failure;
+	Hyperlink = InHyperlink;
+	HyperlinkText = InHyperlinkText;
 
 	UpdateNotification();
 }
 
-void FCoreAsyncTaskNotificationImpl::SetComplete(const FText& InTitleText, const FText& InProgressText, const bool bSuccess)
+void FCoreAsyncTaskNotificationImpl::SetComplete(const bool bSuccess)
 {
 	FScopeLock Lock(&SynchronizationObject);
 
-	TitleText = InTitleText;
-	ProgressText = InProgressText;
-	State = bSuccess ? ENotificationState::Success : ENotificationState::Failure;
+	State = bSuccess ? EAsyncTaskNotificationState::Success : EAsyncTaskNotificationState::Failure;
 
 	UpdateNotification();
 }
 
+void FCoreAsyncTaskNotificationImpl::SetComplete(const FText& InTitleText, const FText& InProgressText, const bool bSuccess)
+{
+	FAsyncNotificationStateData StateData(InTitleText, InProgressText, bSuccess ? EAsyncTaskNotificationState::Success : EAsyncTaskNotificationState::Failure);
+	StateData.PromptText = PromptText;
+	StateData.HyperlinkText = HyperlinkText;
+	StateData.Hyperlink = Hyperlink;
+	SetNotificationState(StateData);
+}
+
+
+void FCoreAsyncTaskNotificationImpl::SetNotificationState(const FAsyncNotificationStateData& InState)
+{
+	FScopeLock Lock(&SynchronizationObject);
+	bool bUpdateNotification = !TitleText.IdenticalTo(InState.TitleText)
+		|| !ProgressText.IdenticalTo(InState.ProgressText)
+		|| !PromptText.IdenticalTo(InState.PromptText)
+		|| !HyperlinkText.IdenticalTo(InState.HyperlinkText)
+		|| State != InState.State;
+
+	TitleText = InState.TitleText;
+	ProgressText = InState.ProgressText;
+	PromptText = InState.PromptText;
+	HyperlinkText = InState.HyperlinkText;
+	Hyperlink = InState.Hyperlink;
+	State = InState.State;
+
+	if (bUpdateNotification)
+	{
+		UpdateNotification();
+	}
+}
+
 void FCoreAsyncTaskNotificationImpl::SetCanCancel(const TAttribute<bool>& InCanCancel)
 {
 }
@@ -82,10 +120,9 @@ void FCoreAsyncTaskNotificationImpl::SetKeepOpenOnFailure(const TAttribute<bool>
 {
 }
 
-
-bool FCoreAsyncTaskNotificationImpl::ShouldCancel() const
+EAsyncTaskNotificationPromptAction FCoreAsyncTaskNotificationImpl::GetPromptAction() const
 {
-	return false;
+	return EAsyncTaskNotificationPromptAction::Unattended;
 }
 
 void FCoreAsyncTaskNotificationImpl::UpdateNotification()
@@ -96,7 +133,7 @@ void FCoreAsyncTaskNotificationImpl::UpdateNotification()
 void FCoreAsyncTaskNotificationImpl::LogNotification()
 {
 #if !NO_LOGGING
-	const ELogVerbosity::Type LogVerbosity = State == ENotificationState::Failure ? ELogVerbosity::Error : ELogVerbosity::Log;
+	const ELogVerbosity::Type LogVerbosity = State == EAsyncTaskNotificationState::Failure ? ELogVerbosity::Error : ELogVerbosity::Log;
 	if (LogCategory && !LogCategory->IsSuppressed(LogVerbosity))
 	{
 		FString NotificationMessage;
@@ -109,23 +146,45 @@ void FCoreAsyncTaskNotificationImpl::LogNotification()
 			}
 			NotificationMessage += ProgressText.ToString();
 		}
+		if (!HyperlinkText.IsEmpty())
+		{
+			if (!NotificationMessage.IsEmpty())
+			{
+				NotificationMessage += TEXT(" - ");
+			}
+			NotificationMessage += HyperlinkText.ToString();
+		}
+		if (!PromptText.IsEmpty())
+		{
+			if (!NotificationMessage.IsEmpty())
+			{
+				NotificationMessage += TEXT(" - ");
+			}
+			NotificationMessage += PromptText.ToString();
+		}
 
 		if (!NotificationMessage.IsEmpty())
 		{
 			static const FText PendingStateText = LOCTEXT("NotificationState_Pending", "Pending");
 			static const FText SuccessStateText = LOCTEXT("NotificationState_Success", "Success");
 			static const FText FailureStateText = LOCTEXT("NotificationState_Failure", "Failure");
+			static const FText PromptStateText = LOCTEXT("NotificationState_Prompt", "Prompt");
 
 			FText StateText = PendingStateText;
-			if (State == ENotificationState::Success)
+			switch (State)
 			{
+			case EAsyncTaskNotificationState::Success:
 				StateText = SuccessStateText;
-			}
-			else if (State == ENotificationState::Failure)
-			{
+				break;
+			case EAsyncTaskNotificationState::Failure:
 				StateText = FailureStateText;
+				break;
+			case EAsyncTaskNotificationState::Prompt:
+				StateText = PromptStateText;
+				break;
+			// default:
+			// nothing
 			}
-
 			FMsg::Logf(nullptr, 0, LogCategory->GetCategoryName(), LogVerbosity, TEXT("[%s] %s"), *StateText.ToString(), *NotificationMessage);
 		}
 	}
diff --git a/Engine/Source/Runtime/Core/Public/Misc/AsyncTaskNotification.h b/Engine/Source/Runtime/Core/Public/Misc/AsyncTaskNotification.h
--- a/Engine/Source/Runtime/Core/Public/Misc/AsyncTaskNotification.h
+++ b/Engine/Source/Runtime/Core/Public/Misc/AsyncTaskNotification.h
@@ -49,6 +49,61 @@ struct FAsyncTaskNotificationConfig
 	const FLogCategoryBase* LogCategory = nullptr;
 };
 
+enum class EAsyncTaskNotificationState : uint8
+{
+	None = 0,
+	Pending,
+	Success,
+	Failure,
+	Prompt,
+};
+
+enum class EAsyncTaskNotificationPromptAction : uint8
+{
+	/** No action taken for the ongoing task. */
+	None = 0,
+	/** Continue ongoing task after a prompt. */
+	Continue,
+	/** Cancel ongoing task. */
+	Cancel,
+	/** No action can be taken, task is unattended. */
+	Unattended
+};
+
+/**
+ * Async Notification State data
+ */
+struct FAsyncNotificationStateData
+{
+	FAsyncNotificationStateData()
+		: State(EAsyncTaskNotificationState::Pending)
+	{}
+
+	FAsyncNotificationStateData(const FText& InTitleText, const FText& InProgressText, EAsyncTaskNotificationState InState)
+		: State(InState)
+		, TitleText(InTitleText)
+		, ProgressText(InProgressText)
+	{}
+
+	/** The notification state. */
+	EAsyncTaskNotificationState State;
+
+	/** The title text displayed in the notification */
+	FText TitleText;
+
+	/** The progress text displayed in the notification (if any) */
+	FText ProgressText;
+
+	/** The prompt text displayed in the notification, follow the progress text in console notification, used on button for UI notification (if any) */
+	FText PromptText;
+
+	/** Text to display for the hyperlink message. */
+	FText HyperlinkText;
+
+	/** Hyperlink callback. If not bound the hyperlink text won't be displayed. */
+	FSimpleDelegate Hyperlink;
+};
+
 /**
  * Provides notifications for an on-going asynchronous task.
  */
@@ -87,6 +142,16 @@ public:
 	 */
 	void SetProgressText(const FText& InProgressText);
 
+	/**
+	 * Set the prompt text of this notification, if needed.
+	 */
+	void SetPromptText(const FText& InPromptText);
+
+	/**
+	 * Set the Hyperlink for this notification.
+	 */
+	void SetHyperlink(const FSimpleDelegate& InHyperlink, const FText& InHyperlinkText = FText());
+
 	/**
 	 * Set the task as complete.
 	 */
@@ -97,6 +162,11 @@ public:
 	 */
 	void SetComplete(const FText& InTitleText, const FText& InProgressText, const bool bSuccess = true);
 
+	/**
+	 * Set the notification state. Provide finer control than SetComplete by setting every field of the state.
+	 */
+	void SetNotificationState(const FAsyncNotificationStateData& InState);
+
 	/**
 	 * Set whether this task be canceled.
 	 */
@@ -113,9 +183,10 @@ public:
 	void SetKeepOpenOnFailure(const TAttribute<bool>& InKeepOpenOnFailure);
 
 	/**
-	 * True if the user has requested that the task be canceled.
+	 * Return the notification prompt action.
+	 * The action resets to `None` when the notification state changes.
 	 */
-	bool ShouldCancel() const;
+	EAsyncTaskNotificationPromptAction GetPromptAction() const;
 
 private:
 	/** Pointer to the real notification implementation */
diff --git a/Engine/Source/Runtime/Core/Public/Misc/CoreAsyncTaskNotificationImpl.h b/Engine/Source/Runtime/Core/Public/Misc/CoreAsyncTaskNotificationImpl.h
--- a/Engine/Source/Runtime/Core/Public/Misc/CoreAsyncTaskNotificationImpl.h
+++ b/Engine/Source/Runtime/Core/Public/Misc/CoreAsyncTaskNotificationImpl.h
@@ -27,6 +27,16 @@ public:
 	 */
 	virtual void SetProgressText(const FText& InProgressText) = 0;
 
+	/**
+	 * Set the prompt text of this notification.
+	 */
+	virtual void SetPromptText(const FText& InPromptText) = 0;
+
+	/**
+	 * Set the hyperlink text of this notification.
+	 */
+	virtual void SetHyperlink(const FSimpleDelegate& InHyperlink, const FText& InHyperlinkText) = 0;
+
 	/**
 	 * Set the task as complete.
 	 */
@@ -37,6 +47,11 @@ public:
 	 */
 	virtual void SetComplete(const FText& InTitleText, const FText& InProgressText, const bool bSuccess) = 0;
 
+	/**
+	 * Set the task notification state. provides finer control than SetComplete
+	 */
+	virtual void SetNotificationState(const FAsyncNotificationStateData& InState) = 0;
+
 	/**
 	 * Set whether this task be canceled.
 	 */
@@ -53,9 +68,10 @@ public:
 	virtual void SetKeepOpenOnFailure(const TAttribute<bool>& InKeepOpenOnFailure) = 0;
 
 	/**
-	 * True if the user has requested that the task be canceled
+	 * Return the notification prompt action.
+	 * The action resets to `None` when the notification state changes.
 	 */
-	virtual bool ShouldCancel() const = 0;
+	virtual EAsyncTaskNotificationPromptAction GetPromptAction() const = 0;
 };
 
 /**
@@ -71,26 +87,22 @@ public:
 	virtual void Initialize(const FAsyncTaskNotificationConfig& InConfig) override;
 	virtual void SetTitleText(const FText& InTitleText, const bool bClearProgressText) override;
 	virtual void SetProgressText(const FText& InProgressText) override;
+	virtual void SetPromptText(const FText& InPromptText) override;
+	virtual void SetHyperlink(const FSimpleDelegate& InHyperlink, const FText& InHyperlinkText) override;
 	virtual void SetComplete(const bool bSuccess) override;
 	virtual void SetComplete(const FText& InTitleText, const FText& InProgressText, const bool bSuccess) override;
+	virtual void SetNotificationState(const FAsyncNotificationStateData& InState) override;
 	virtual void SetCanCancel(const TAttribute<bool>& InCanCancel) override;
 	virtual void SetKeepOpenOnSuccess(const TAttribute<bool>& InKeepOpenOnSuccess) override;
 	virtual void SetKeepOpenOnFailure(const TAttribute<bool>& InKeepOpenOnFailure) override;
-	virtual bool ShouldCancel() const override;
+	virtual EAsyncTaskNotificationPromptAction GetPromptAction() const override;
 
 protected:
 	/** Update the notification (the critical section is held while this function is called) */
 	virtual void UpdateNotification();
 
-	enum class ENotificationState : uint8
-	{
-		Pending,
-		Success,
-		Failure,
-	};
-
 	/** The current state of this notification */
-	ENotificationState State = ENotificationState::Pending;
+	EAsyncTaskNotificationState State = EAsyncTaskNotificationState::Pending;
 
 	/** The title text displayed in the notification (if any) */
 	FText TitleText;
@@ -98,6 +110,14 @@ protected:
 	/** The progress text displayed in the notification (if any) */
 	FText ProgressText;
 
+	/** The progress text displayed in the notification (if any) */
+	FText PromptText;
+
+	/** When set this will display as a hyperlink on the right side of the notification. */
+	FSimpleDelegate Hyperlink;
+
+	/** Text to display for the hyperlink message */
+	FText HyperlinkText;
 private:
 	/** Log the current notification state (if any, and if enabled) */
 	void LogNotification();
diff --git a/Engine/Source/Runtime/Slate/Private/Framework/Notifications/SlateAsyncTaskNotificationImpl.cpp b/Engine/Source/Runtime/Slate/Private/Framework/Notifications/SlateAsyncTaskNotificationImpl.cpp
--- a/Engine/Source/Runtime/Slate/Private/Framework/Notifications/SlateAsyncTaskNotificationImpl.cpp
+++ b/Engine/Source/Runtime/Slate/Private/Framework/Notifications/SlateAsyncTaskNotificationImpl.cpp
@@ -8,6 +8,7 @@
 
 #include "Widgets/SCompoundWidget.h"
 #include "Widgets/Input/SButton.h"
+#include "Widgets/Input/SHyperlink.h"
 #include "Widgets/Text/STextBlock.h"
 #include "Widgets/Layout/SBox.h"
 #include "Widgets/Layout/SBorder.h"
@@ -38,10 +39,10 @@ public:
 	void SetOwner(TSharedPtr<SNotificationItem> InOwningNotification);
 
 	/** Update the notification state */
-	void UpdateNotification(const FText& InTitleText, const FText& InProgressText);
+	void UpdateNotification(const FText& InTitleText, const FText& InProgressText, const FText& InPromptText, const FSimpleDelegate& InHyperlink, const FText& InHyperlinkText);
 
 	/** Set the pending completion state of the notification (applied during the next Tick) and reset the external UI reference */
-	void SetPendingCompletionState(const SNotificationItem::ECompletionState InPendingCompletionState, TSharedPtr<SSlateAsyncTaskNotificationWidget>* ExternalReferenceToReset);
+	void SetPendingCompletionState(const EAsyncTaskNotificationState InPendingCompletionState, TSharedPtr<SSlateAsyncTaskNotificationWidget>* ExternalReferenceToReset);
 
 	/** Set whether this task be canceled */
 	void SetCanCancel(const TAttribute<bool>& InCanCancel);
@@ -52,9 +53,7 @@ public:
 	/** Set whether to keep this notification open on failure */
 	void SetKeepOpenOnFailure(const TAttribute<bool>& InKeepOpenOnFailure);
 
-	/** True if the user has requested that the task be canceled */
-	bool ShouldCancel() const;
-
+	EAsyncTaskNotificationPromptAction GetPromptAction() const;
 private:
 	/** Sync attribute bindings with the cached values (once per-frame from the game thread) */
 	void SyncAttributes();
@@ -77,15 +76,26 @@ private:
 	EVisibility GetCancelButtonVisibility() const;
 	FReply OnCancelButtonClicked();
 
+	/** Prompt button */
+	bool IsPromptButtonEnabled() const;
+	EVisibility GetPromptButtonVisibility() const;
+	FReply OnPromptButtonClicked();
+	FText GetPromptButtonText() const;
+
 	/** Close button */
 	EVisibility GetCloseButtonVisibility() const;
 	FReply OnCloseButtonClicked();
 
+	/** Hyperlink */
+	void OnHyperlinkClicked() const;
+	FText GetHyperlinkText() const;
+	EVisibility GetHyperlinkVisibility() const;
+
 	/** Get the current completion state from the parent notification */
 	SNotificationItem::ECompletionState GetNotificationCompletionState() const;
 
-	/** True if the user has requested that the task be canceled */
-	TAtomic<bool> bShouldCancel;
+	/** Action taken for the task, resets to none on notification state change. */
+	TAtomic<EAsyncTaskNotificationPromptAction> PromptAction;
 
 	/** Can this task be canceled? Will show a cancel button for in-progress tasks */
 	TAttribute<bool> bCanCancelAttr = false;
@@ -105,8 +115,20 @@ private:
 	/** The progress text displayed in the notification (if any) */
 	FText ProgressText;
 
+	/** The prompt text displayed on the prompt button in the notification (if any)*/
+	FText PromptText;
+
+	/** When set this will display as a hyperlink on the right side of the notification. */
+	FSimpleDelegate Hyperlink;
+
+	/** Text to display for the hyperlink message */
+	FText HyperlinkText;
+
 	/** The pending completion state of the notification (if any, applied during the next Tick) */
-	TOptional<SNotificationItem::ECompletionState> PendingCompletionState;
+	TOptional<EAsyncTaskNotificationState> PendingCompletionState;
+
+	/** The current completion state of this widget. */
+	EAsyncTaskNotificationState NotificationState = EAsyncTaskNotificationState::None;
 
 	/** Pointer to the notification item that owns this widget (this is a deliberate reference cycle as we need this object alive until we choose to expire it, at which point we release our reference to allow everything to be destroyed) */
 	TSharedPtr<SNotificationItem> OwningNotification;
@@ -120,7 +142,7 @@ private:
 
 void SSlateAsyncTaskNotificationWidget::Construct(const FArguments& InArgs, const FAsyncTaskNotificationConfig& InConfig)
 {
-	bShouldCancel = false;
+	PromptAction = FApp::IsUnattended() ? EAsyncTaskNotificationPromptAction::Unattended : EAsyncTaskNotificationPromptAction::None;
 	bCanCancelAttr = InConfig.bCanCancel;
 	bKeepOpenOnSuccessAttr = InConfig.bKeepOpenOnSuccess;
 	bKeepOpenOnFailureAttr = InConfig.bKeepOpenOnFailure;
@@ -213,15 +235,45 @@ void SSlateAsyncTaskNotificationWidget::Construct(const FArguments& InArgs, cons
 					]
 				]
 
+				// Hyperlink
+				+ SVerticalBox::Slot()
+				.AutoHeight()
+				.HAlign(HAlign_Right)
+				.VAlign(VAlign_Bottom)
+				[
+					SNew(SBox)
+					.Padding(FMargin(0.0f, 2.0f, 0.0f, 2.0f))
+					.VAlign(VAlign_Center)
+					.HAlign(HAlign_Left)
+					.Visibility(this, &SSlateAsyncTaskNotificationWidget::GetHyperlinkVisibility)
+					[
+						SNew(SHyperlink)
+						.Text(this, &SSlateAsyncTaskNotificationWidget::GetHyperlinkText)
+						.OnNavigate(this, &SSlateAsyncTaskNotificationWidget::OnHyperlinkClicked)
+					]
+				]
+
 				// Buttons
 				+SVerticalBox::Slot()
 				.AutoHeight()
-				.HAlign(HAlign_Right)
+				.HAlign(HAlign_Center)
 				.VAlign(VAlign_Bottom)
 				.Padding(FMargin(0.0f, 5.0f, 0.0f, 0.0f))
 				[
 					SNew(SHorizontalBox)
 
+					// Prompt Button
+					+ SHorizontalBox::Slot()
+					.AutoWidth()
+					.VAlign(VAlign_Center)
+					.Padding(FMargin(5.0f, 0.0f, 5.0f, 0.0f))
+					[
+						SNew(SButton)
+						.Text_Raw(this, &SSlateAsyncTaskNotificationWidget::GetPromptButtonText)
+						.Visibility(this, &SSlateAsyncTaskNotificationWidget::GetPromptButtonVisibility)
+						.OnClicked(this, &SSlateAsyncTaskNotificationWidget::OnPromptButtonClicked)
+					]
+
 					// Cancel Button
 					+SHorizontalBox::Slot()
 					.AutoWidth()
@@ -253,7 +305,7 @@ void SSlateAsyncTaskNotificationWidget::Tick(const FGeometry& AllottedGeometry,
 {
 	SyncAttributes();
 
-	SNotificationItem::ECompletionState CompletionStateToApply = SNotificationItem::CS_Pending;
+	EAsyncTaskNotificationState CompletionStateToApply = EAsyncTaskNotificationState::None;
 	{
 		FScopeLock Lock(&CompletionCS);
 
@@ -264,11 +316,29 @@ void SSlateAsyncTaskNotificationWidget::Tick(const FGeometry& AllottedGeometry,
 		}
 	}
 
-	if (CompletionStateToApply != SNotificationItem::CS_Pending)
+	if (OwningNotification && CompletionStateToApply != EAsyncTaskNotificationState::None)
 	{
-		if (OwningNotification)
+		NotificationState = CompletionStateToApply;
+		SNotificationItem::ECompletionState OwningCompletionState = SNotificationItem::CS_None;
+		switch (NotificationState)
+		{
+		case EAsyncTaskNotificationState::Pending:
+			OwningCompletionState = SNotificationItem::CS_Pending;
+			break;
+		case EAsyncTaskNotificationState::Failure:
+			OwningCompletionState = SNotificationItem::CS_Fail;
+			break;
+		case EAsyncTaskNotificationState::Success:
+			OwningCompletionState = SNotificationItem::CS_Success;
+			break;
+		case EAsyncTaskNotificationState::Prompt:
+			OwningNotification->Pulse(FLinearColor(0.f, 0.f, 1.f));
+
+			break;
+		}
+		if (OwningCompletionState != SNotificationItem::CS_None && OwningCompletionState != OwningNotification->GetCompletionState())
 		{
-			OwningNotification->SetCompletionState(CompletionStateToApply);
+			OwningNotification->SetCompletionState(OwningCompletionState);
 		}
 	}
 }
@@ -289,6 +359,9 @@ void SSlateAsyncTaskNotificationWidget::OnSetCompletionState(SNotificationItem::
 			OwningNotification.Reset();
 		}
 	}
+
+	// Reset the `PromptAction` state when changing completion state
+	PromptAction = FApp::IsUnattended() ? EAsyncTaskNotificationPromptAction::Unattended : EAsyncTaskNotificationPromptAction::None;
 }
 
 TSharedRef<SWidget> SSlateAsyncTaskNotificationWidget::AsWidget()
@@ -301,13 +374,16 @@ void SSlateAsyncTaskNotificationWidget::SetOwner(TSharedPtr<SNotificationItem> I
 	OwningNotification = InOwningNotification;
 }
 
-void SSlateAsyncTaskNotificationWidget::UpdateNotification(const FText& InTitleText, const FText& InProgressText)
+void SSlateAsyncTaskNotificationWidget::UpdateNotification(const FText& InTitleText, const FText& InProgressText, const FText& InPromptText, const FSimpleDelegate& InHyperlink, const FText& InHyperlinkText)
 {
 	TitleText = InTitleText;
 	ProgressText = InProgressText;
+	PromptText = InPromptText;
+	Hyperlink = InHyperlink;
+	HyperlinkText = InHyperlinkText;
 }
 
-void SSlateAsyncTaskNotificationWidget::SetPendingCompletionState(const SNotificationItem::ECompletionState InPendingCompletionState, TSharedPtr<SSlateAsyncTaskNotificationWidget>* ExternalReferenceToReset)
+void SSlateAsyncTaskNotificationWidget::SetPendingCompletionState(const EAsyncTaskNotificationState InPendingCompletionState, TSharedPtr<SSlateAsyncTaskNotificationWidget>* ExternalReferenceToReset)
 {
 	FScopeLock Lock(&CompletionCS);
 
@@ -340,9 +416,9 @@ void SSlateAsyncTaskNotificationWidget::SetKeepOpenOnFailure(const TAttribute<bo
 	bKeepOpenOnFailureAttr = InKeepOpenOnFailure;
 }
 
-bool SSlateAsyncTaskNotificationWidget::ShouldCancel() const
+EAsyncTaskNotificationPromptAction SSlateAsyncTaskNotificationWidget::GetPromptAction() const
 {
-	return bShouldCancel;
+	return PromptAction;
 }
 
 void SSlateAsyncTaskNotificationWidget::SyncAttributes()
@@ -380,48 +456,69 @@ FText SSlateAsyncTaskNotificationWidget::GetProgressText() const
 
 EVisibility SSlateAsyncTaskNotificationWidget::GetThrobberVisibility() const
 {
-	return (GetNotificationCompletionState() == SNotificationItem::CS_Pending)
+	return (NotificationState == EAsyncTaskNotificationState::Pending)
 		? EVisibility::Visible
 		: EVisibility::Collapsed;
 }
 
 EVisibility SSlateAsyncTaskNotificationWidget::GetStatusIconVisibility() const
 {
-	const SNotificationItem::ECompletionState NotificationState = GetNotificationCompletionState();
-	return (NotificationState == SNotificationItem::CS_Success || NotificationState == SNotificationItem::CS_Fail)
+	return (NotificationState == EAsyncTaskNotificationState::Success || NotificationState == EAsyncTaskNotificationState::Failure)
 		? EVisibility::Visible
 		: EVisibility::Collapsed;
 }
 
 const FSlateBrush* SSlateAsyncTaskNotificationWidget::GetStatusIconBrush() const
 {
-	return (GetNotificationCompletionState() == SNotificationItem::CS_Success)
+	return (NotificationState == EAsyncTaskNotificationState::Success)
 		? FCoreStyle::Get().GetBrush("NotificationList.SuccessImage")
 		: FCoreStyle::Get().GetBrush("NotificationList.FailImage");
 }
 
 bool SSlateAsyncTaskNotificationWidget::IsCancelButtonEnabled() const
 {
-	return bCanCancel && !bShouldCancel;
+	return bCanCancel && PromptAction == EAsyncTaskNotificationPromptAction::None;
 }
 
 EVisibility SSlateAsyncTaskNotificationWidget::GetCancelButtonVisibility() const
 {
-	return (bCanCancel && GetNotificationCompletionState() == SNotificationItem::CS_Pending)
+	return (bCanCancel && (NotificationState == EAsyncTaskNotificationState::Pending || NotificationState == EAsyncTaskNotificationState::Prompt))
 		? EVisibility::Visible
 		: EVisibility::Collapsed;
 }
 
 FReply SSlateAsyncTaskNotificationWidget::OnCancelButtonClicked()
 {
-	bShouldCancel = true;
+	PromptAction = EAsyncTaskNotificationPromptAction::Cancel;
+	return FReply::Handled();
+}
+
+bool SSlateAsyncTaskNotificationWidget::IsPromptButtonEnabled() const
+{
+	return PromptAction == EAsyncTaskNotificationPromptAction::None;
+}
+
+EVisibility SSlateAsyncTaskNotificationWidget::GetPromptButtonVisibility() const
+{
+	return (!FApp::IsUnattended() && NotificationState == EAsyncTaskNotificationState::Prompt)
+		? EVisibility::Visible
+		: EVisibility::Collapsed;
+}
+
+FReply SSlateAsyncTaskNotificationWidget::OnPromptButtonClicked()
+{
+	PromptAction = EAsyncTaskNotificationPromptAction::Continue;
 	return FReply::Handled();
 }
 
+FText SSlateAsyncTaskNotificationWidget::GetPromptButtonText() const
+{
+	return PromptText;
+}
+
 EVisibility SSlateAsyncTaskNotificationWidget::GetCloseButtonVisibility() const
 {
-	const SNotificationItem::ECompletionState NotificationState = GetNotificationCompletionState();
-	return (!FApp::IsUnattended() && ((bKeepOpenOnSuccess && NotificationState == SNotificationItem::CS_Success) || (bKeepOpenOnFailure && NotificationState == SNotificationItem::CS_Fail)))
+	return (!FApp::IsUnattended() && ((bKeepOpenOnSuccess && NotificationState == EAsyncTaskNotificationState::Success) || (bKeepOpenOnFailure && NotificationState == EAsyncTaskNotificationState::Failure)))
 		? EVisibility::Visible
 		: EVisibility::Collapsed;
 }
@@ -441,6 +538,21 @@ FReply SSlateAsyncTaskNotificationWidget::OnCloseButtonClicked()
 	return FReply::Handled();
 }
 
+void SSlateAsyncTaskNotificationWidget::OnHyperlinkClicked() const
+{
+	Hyperlink.ExecuteIfBound();
+}
+
+FText SSlateAsyncTaskNotificationWidget::GetHyperlinkText() const
+{
+	return HyperlinkText;
+}
+
+EVisibility SSlateAsyncTaskNotificationWidget::GetHyperlinkVisibility() const
+{
+	return Hyperlink.IsBound() ? EVisibility::Visible : EVisibility::Collapsed;
+}
+
 SNotificationItem::ECompletionState SSlateAsyncTaskNotificationWidget::GetNotificationCompletionState() const
 {
 	if (OwningNotification)
@@ -507,20 +619,17 @@ void FSlateAsyncTaskNotificationImpl::UpdateNotification()
 	if (NotificationItemWidget)
 	{
 		// Update the notification text
-		NotificationItemWidget->UpdateNotification(TitleText, ProgressText);
+		NotificationItemWidget->UpdateNotification(TitleText, ProgressText, PromptText, Hyperlink, HyperlinkText);
 
-		if (State != ENotificationState::Pending)
-		{
-			// Complete the notification and remove our references to it in a single atomic operation
-			// NotificationItemWidget will be null once this call completes
-			NotificationItemWidget->SetPendingCompletionState(State == ENotificationState::Success ? SNotificationItem::CS_Success : SNotificationItem::CS_Fail, &NotificationItemWidget);
-		}
+		// Complete the notification and remove our references to it in a single atomic operation if needed
+		// NotificationItemWidget will be null once this call completes if completion is success or failure
+		NotificationItemWidget->SetPendingCompletionState(State, State == EAsyncTaskNotificationState::Failure || State == EAsyncTaskNotificationState::Success ? &NotificationItemWidget : nullptr);
 	}
 }
 
-bool FSlateAsyncTaskNotificationImpl::ShouldCancel() const
+EAsyncTaskNotificationPromptAction FSlateAsyncTaskNotificationImpl::GetPromptAction() const
 {
-	return NotificationItemWidget && NotificationItemWidget->ShouldCancel();
+	return NotificationItemWidget ? NotificationItemWidget->GetPromptAction() : EAsyncTaskNotificationPromptAction::Unattended;
 }
 
 #undef LOCTEXT_NAMESPACE
diff --git a/Engine/Source/Runtime/Slate/Private/Framework/Notifications/SlateAsyncTaskNotificationImpl.h b/Engine/Source/Runtime/Slate/Private/Framework/Notifications/SlateAsyncTaskNotificationImpl.h
--- a/Engine/Source/Runtime/Slate/Private/Framework/Notifications/SlateAsyncTaskNotificationImpl.h
+++ b/Engine/Source/Runtime/Slate/Private/Framework/Notifications/SlateAsyncTaskNotificationImpl.h
@@ -18,8 +18,7 @@ public:
 	virtual void SetCanCancel(const TAttribute<bool>& InCanCancel) override;
 	virtual void SetKeepOpenOnSuccess(const TAttribute<bool>& InKeepOpenOnSuccess) override;
 	virtual void SetKeepOpenOnFailure(const TAttribute<bool>& InKeepOpenOnFailure) override;
-	virtual bool ShouldCancel() const override;
-
+	virtual EAsyncTaskNotificationPromptAction GetPromptAction() const override;
 private:
 	//~ FCoreAsyncTaskNotificationImpl
 	virtual void UpdateNotification() override;
diff --git a/Engine/Source/Runtime/Slate/Private/Widgets/Notifications/SNotificationList.cpp b/Engine/Source/Runtime/Slate/Private/Widgets/Notifications/SNotificationList.cpp
--- a/Engine/Source/Runtime/Slate/Private/Widgets/Notifications/SNotificationList.cpp
+++ b/Engine/Source/Runtime/Slate/Private/Widgets/Notifications/SNotificationList.cpp
@@ -105,6 +105,14 @@ public:
 		FadeAnimation.PlayReverse(this->AsShared());
 	}
 
+	virtual void Pulse(const FLinearColor& GlowColor)
+	{
+		DefaultGlowColor = GlowColor;
+		CompletionStateAnimation = FCurveSequence();
+		GlowCurve = CompletionStateAnimation.AddCurve(0.f, 0.75f);
+		CompletionStateAnimation.Play(this->AsShared());
+	}
+
 	/** Sets the ExpireDuration */
 	virtual void SetExpireDuration(float Duration) override
 	{
@@ -202,7 +210,7 @@ protected:
 		{
 		case CS_Success: return FLinearColor(0,1,0,GlowAlpha);
 		case CS_Fail: return FLinearColor(1,0,0,GlowAlpha);
-		default: return FLinearColor(1,1,1,GlowAlpha);
+		default: return DefaultGlowColor.CopyWithNewOpacity(GlowAlpha);
 		}
 	}
 
@@ -246,6 +254,9 @@ protected:
 	/** The duration before a fadeout for this element */
 	TAttribute< float > ExpireDuration;
 
+	/** The default glow color for pulse animation. */
+	FLinearColor DefaultGlowColor = FLinearColor(1.0f, 1.0f, 1.0f);
+
 	/** The text displayed in this element */
 	TSharedPtr<STextBlock> MyTextBlock;
 
diff --git a/Engine/Source/Runtime/Slate/Public/Widgets/Notifications/SNotificationList.h b/Engine/Source/Runtime/Slate/Public/Widgets/Notifications/SNotificationList.h
--- a/Engine/Source/Runtime/Slate/Public/Widgets/Notifications/SNotificationList.h
+++ b/Engine/Source/Runtime/Slate/Public/Widgets/Notifications/SNotificationList.h
@@ -29,7 +29,7 @@ public:
 		CS_None,
 		CS_Pending,
 		CS_Success,
-		CS_Fail
+		CS_Fail,
 	};
 
 	/** Sets the text for message element */
@@ -58,6 +58,9 @@ public:
 
 	/** Begin the fade out */
 	virtual void Fadeout() = 0;
+
+	/** Pulse the notification with a certain outline glow color while it is still pending. */
+	virtual void Pulse(const FLinearColor& GlowColor) = 0;
 };
 
 
