Add text verification to SEditableTextBox, similar to SInlineEditableTextBox

diff --git a/Engine/Plugins/FX/Niagara/Source/NiagaraEditor/Private/NiagaraNodeOutput.cpp b/Engine/Plugins/FX/Niagara/Source/NiagaraEditor/Private/NiagaraNodeOutput.cpp
--- a/Engine/Plugins/FX/Niagara/Source/NiagaraEditor/Private/NiagaraNodeOutput.cpp
+++ b/Engine/Plugins/FX/Niagara/Source/NiagaraEditor/Private/NiagaraNodeOutput.cpp
@@ -13,7 +13,6 @@
 #include "ToolMenus.h"
 #include "Widgets/Layout/SBox.h"
 #include "Widgets/Input/SEditableTextBox.h"
-#include "Widgets/SEditableTextBoxWithVerification.h"
 #include "NiagaraEditorUtilities.h"
 
 #define LOCTEXT_NAMESPACE "NiagaraNodeOutput"
@@ -94,7 +93,7 @@ void UNiagaraNodeOutput::GetNodeContextMenuActions(UToolMenu* Menu, UGraphNodeCo
 				.WidthOverride(100)
 				.Padding(FMargin(5, 0, 0, 0))
 				[
-					SNew(SEditableTextBoxWithVerification)
+					SNew(SEditableTextBox)
 					.Text_UObject(this, &UNiagaraNodeOutput::GetPinNameText, Pin)
 					.OnVerifyTextChanged_UObject(this, &UNiagaraNodeOutput::VerifyPinNameTextChanged, Pin)
 					.OnTextCommitted_UObject(const_cast<UNiagaraNodeOutput*>(this), &UNiagaraNodeOutput::PinNameTextCommitted, Pin)
diff --git a/Engine/Plugins/FX/Niagara/Source/NiagaraEditor/Private/NiagaraNodeWithDynamicPins.cpp b/Engine/Plugins/FX/Niagara/Source/NiagaraEditor/Private/NiagaraNodeWithDynamicPins.cpp
--- a/Engine/Plugins/FX/Niagara/Source/NiagaraEditor/Private/NiagaraNodeWithDynamicPins.cpp
+++ b/Engine/Plugins/FX/Niagara/Source/NiagaraEditor/Private/NiagaraNodeWithDynamicPins.cpp
@@ -15,7 +15,6 @@
 #include "SNiagaraParameterMapView.h"
 #include "Framework/Application/SlateApplication.h"
 #include "NiagaraNodeParameterMapBase.h"
-#include "Widgets/SEditableTextBoxWithVerification.h"
 
 #define LOCTEXT_NAMESPACE "NiagaraNodeWithDynamicPins"
 
@@ -177,7 +176,7 @@ void UNiagaraNodeWithDynamicPins::GetNodeContextMenuActions(UToolMenu* Menu, UGr
 				.WidthOverride(100)
 				.Padding(FMargin(5, 0, 0, 0))
 				[
-					SNew(SEditableTextBoxWithVerification)
+					SNew(SEditableTextBox)
 					.Text_UObject(this, &UNiagaraNodeWithDynamicPins::GetPinNameText, Pin)
 					.OnTextCommitted_UObject(const_cast<UNiagaraNodeWithDynamicPins*>(this), &UNiagaraNodeWithDynamicPins::PinNameTextCommitted, Pin)
 					.OnVerifyTextChanged_UObject(this, &UNiagaraNodeWithDynamicPins::VerifyEditablePinName, Context->Pin)
diff --git a/Engine/Plugins/FX/Niagara/Source/NiagaraEditor/Private/Widgets/SEditableTextBoxWithVerification.cpp b/Engine/Plugins/FX/Niagara/Source/NiagaraEditor/Private/Widgets/SEditableTextBoxWithVerification.cpp
deleted file mode 100644
diff --git a/Engine/Plugins/FX/Niagara/Source/NiagaraEditor/Private/Widgets/SEditableTextBoxWithVerification.h b/Engine/Plugins/FX/Niagara/Source/NiagaraEditor/Private/Widgets/SEditableTextBoxWithVerification.h
deleted file mode 100644
diff --git a/Engine/Source/Runtime/Slate/Private/Widgets/Input/SEditableTextBox.cpp b/Engine/Source/Runtime/Slate/Private/Widgets/Input/SEditableTextBox.cpp
--- a/Engine/Source/Runtime/Slate/Private/Widgets/Input/SEditableTextBox.cpp
+++ b/Engine/Source/Runtime/Slate/Private/Widgets/Input/SEditableTextBox.cpp
@@ -32,6 +32,9 @@ void SEditableTextBox::Construct( const FArguments& InArgs )
 	ForegroundColorOverride = InArgs._ForegroundColor;
 	BackgroundColorOverride = InArgs._BackgroundColor;
 	ReadOnlyForegroundColorOverride = InArgs._ReadOnlyForegroundColor;
+	OnTextChanged = InArgs._OnTextChanged;
+	OnVerifyTextChanged = InArgs._OnVerifyTextChanged;
+	OnTextCommitted = InArgs._OnTextCommitted;
 
 	SBorder::Construct( SBorder::FArguments()
 		.BorderImage( this, &SEditableTextBox::GetBorderImage )
@@ -65,8 +68,8 @@ void SEditableTextBox::Construct( const FArguments& InArgs )
 					.AllowContextMenu( InArgs._AllowContextMenu )
 					.OnContextMenuOpening( InArgs._OnContextMenuOpening )
 					.ContextMenuExtender( InArgs._ContextMenuExtender )
-					.OnTextChanged( InArgs._OnTextChanged )
-					.OnTextCommitted( InArgs._OnTextCommitted )
+					.OnTextChanged(this, &SEditableTextBox::OnEditableTextChanged)
+					.OnTextCommitted(this, &SEditableTextBox::OnEditableTextCommitted)
 					.MinDesiredWidth( InArgs._MinDesiredWidth )
 					.SelectAllTextOnCommit( InArgs._SelectAllTextOnCommit )
 					.OnKeyCharHandler( InArgs._OnKeyCharHandler )			
@@ -426,3 +429,49 @@ TOptional<FText> SEditableTextBox::GetDefaultAccessibleText(EAccessibleType Acce
 	return TOptional<FText>();
 }
 #endif
+
+void SEditableTextBox::OnEditableTextChanged(const FText& InText)
+{
+	OnTextChanged.ExecuteIfBound(InText);
+
+	if (OnVerifyTextChanged.IsBound())
+	{
+		FText OutErrorMessage;
+		if (!OnVerifyTextChanged.Execute(InText, OutErrorMessage))
+		{
+			// Display as an error.
+			SetError(OutErrorMessage);
+		}
+		else
+		{
+			SetError(FText::GetEmpty());
+		}
+	}
+}
+
+void SEditableTextBox::OnEditableTextCommitted(const FText& InText, ETextCommit::Type InCommitType)
+{
+	if (OnVerifyTextChanged.IsBound())
+	{
+		FText OutErrorMessage;
+		if (!OnVerifyTextChanged.Execute(InText, OutErrorMessage))
+		{
+           		// Display as an error.
+			if (InCommitType == ETextCommit::OnEnter)
+			{
+				SetError(OutErrorMessage);
+			}
+			return;
+		}
+		else
+		{
+			if (InCommitType == ETextCommit::OnEnter)
+			{
+				SetError(FText::GetEmpty());
+			}
+			
+		}		
+	}
+
+	OnTextCommitted.ExecuteIfBound(InText, InCommitType);
+}
diff --git a/Engine/Source/Runtime/Slate/Public/Framework/SlateDelegates.h b/Engine/Source/Runtime/Slate/Public/Framework/SlateDelegates.h
--- a/Engine/Source/Runtime/Slate/Public/Framework/SlateDelegates.h
+++ b/Engine/Source/Runtime/Slate/Public/Framework/SlateDelegates.h
@@ -239,3 +239,6 @@ DECLARE_DELEGATE_RetVal_TwoParams(FReply, FOnKeyDown, const FGeometry&, const FK
 
 /** A delegate for a callback during OnKeyChar from an SWidget */
 DECLARE_DELEGATE_RetVal_TwoParams(FReply, FOnKeyChar, const FGeometry&, const FCharacterEvent&);
+
+/** A delegate for a callback whenever text is changed programmatically or interactively by the user, and has to be verified */
+DECLARE_DELEGATE_RetVal_TwoParams(bool, FOnVerifyTextChanged, const FText&, FText&)
\ No newline at end of file
diff --git a/Engine/Source/Runtime/Slate/Public/Widgets/Input/SEditableTextBox.h b/Engine/Source/Runtime/Slate/Public/Widgets/Input/SEditableTextBox.h
--- a/Engine/Source/Runtime/Slate/Public/Widgets/Input/SEditableTextBox.h
+++ b/Engine/Source/Runtime/Slate/Public/Widgets/Input/SEditableTextBox.h
@@ -14,6 +14,7 @@
 #include "Widgets/DeclarativeSyntaxSupport.h"
 #include "Widgets/Input/SEditableText.h"
 #include "Widgets/Layout/SBorder.h"
+#include "Framework/SlateDelegates.h"
 
 class IErrorReportingWidget;
 class SBox;
@@ -113,6 +114,9 @@ public:
 		/** Called whenever the text is committed.  This happens when the user presses enter or the text box loses focus. */
 		SLATE_EVENT( FOnTextCommitted, OnTextCommitted )
 
+		/** Called whenever the text is changed programmatically or interactively by the user */
+		SLATE_EVENT( FOnVerifyTextChanged, OnVerifyTextChanged )
+
 		/** Minimum width that a text block should be */
 		SLATE_ATTRIBUTE( float, MinDesiredWidth )
 
@@ -357,6 +361,11 @@ protected:
 	virtual TSharedRef<FSlateAccessibleWidget> CreateAccessibleWidget() override;
 	virtual TOptional<FText> GetDefaultAccessibleText(EAccessibleType AccessibleType = EAccessibleType::Main) const override;
 #endif
+	/** Callback for the editable text's OnTextChanged event */
+	void OnEditableTextChanged(const FText& InText);
+
+	/** Callback when the editable text is committed. */
+	void OnEditableTextCommitted(const FText& InText, ETextCommit::Type InCommitType);
 
 	const FEditableTextBoxStyle* Style;
 
@@ -387,6 +396,15 @@ protected:
 	/** SomeWidget reporting */
 	TSharedPtr<class IErrorReportingWidget> ErrorReporting;
 
+	/** Called when the text is changed interactively */
+	FOnTextChanged OnTextChanged;
+
+	/** Called when the user commits their change to the editable text control */
+	FOnTextCommitted OnTextCommitted;	
+
+	/** Callback to verify text when changed. Will return an error message to denote problems. */
+	FOnVerifyTextChanged OnVerifyTextChanged;
+
 private:
 
 	FMargin FORCEINLINE DeterminePadding() const { check(Style);  return PaddingOverride.IsSet() ? PaddingOverride.Get() : Style->Padding; }
diff --git a/Engine/Source/Runtime/Slate/Public/Widgets/Text/SInlineEditableTextBlock.h b/Engine/Source/Runtime/Slate/Public/Widgets/Text/SInlineEditableTextBlock.h
--- a/Engine/Source/Runtime/Slate/Public/Widgets/Text/SInlineEditableTextBlock.h
+++ b/Engine/Source/Runtime/Slate/Public/Widgets/Text/SInlineEditableTextBlock.h
@@ -23,7 +23,6 @@ class SHorizontalBox;
 class SMultiLineEditableTextBox;
 class STextBlock;
 
-DECLARE_DELEGATE_RetVal_TwoParams(bool, FOnVerifyTextChanged, const FText&, FText&)
 DECLARE_DELEGATE_OneParam(FOnBeginTextEdit, const FText&)
 
 /**
