Fix static analysis warning
#rnx

diff --git a/Engine/Source/Runtime/Slate/Public/Widgets/Views/STableRow.h b/Engine/Source/Runtime/Slate/Public/Widgets/Views/STableRow.h
--- a/Engine/Source/Runtime/Slate/Public/Widgets/Views/STableRow.h
+++ b/Engine/Source/Runtime/Slate/Public/Widgets/Views/STableRow.h
@@ -661,7 +661,7 @@ public:
 			const FVector2D LocalPointerPos = MyGeometry.AbsoluteToLocal(DragDropEvent.GetScreenSpacePosition());
 			const EItemDropZone ItemHoverZone = ZoneFromPointerPosition(LocalPointerPos, MyGeometry.GetLocalSize().Y);
 
-			ItemDropZone = [&, this]()
+			ItemDropZone = [ItemHoverZone, DragDropEvent, this]()
 			{
 				TSharedRef< ITypedTableView<ItemType> > OwnerTable = OwnerTablePtr.Pin().ToSharedRef();
 				if (const ItemType* MyItemPtr = GetItemForThis(OwnerTable))
