Adding missing file

diff --git a/Engine/Source/Runtime/Slate/Private/Widgets/Layout/SDPIScaler.cpp b/Engine/Source/Runtime/Slate/Private/Widgets/Layout/SDPIScaler.cpp
--- a/Engine/Source/Runtime/Slate/Private/Widgets/Layout/SDPIScaler.cpp
+++ b/Engine/Source/Runtime/Slate/Private/Widgets/Layout/SDPIScaler.cpp
@@ -4,20 +4,21 @@
 #include "Layout/ArrangedChildren.h"
 
 SDPIScaler::SDPIScaler()
-: ChildSlot(this)
+	: ChildSlot(this)
 {
 	SetCanTick(false);
 	bCanSupportFocus = false;
+	bHasRelativeLayoutScale = true;
 }
 
 void SDPIScaler::Construct( const FArguments& InArgs )
 {
+	DPIScale = InArgs._DPIScale;
+
 	ChildSlot
 	[
 		InArgs._Content.Widget
 	];
-	
-	DPIScale = InArgs._DPIScale;
 }
 
 void SDPIScaler::OnArrangeChildren( const FGeometry& AllottedGeometry, FArrangedChildren& ArrangedChildren ) const
@@ -33,7 +34,6 @@ void SDPIScaler::OnArrangeChildren( const FGeometry& AllottedGeometry, FArranged
 			AllottedGeometry.GetLocalSize() / MyDPIScale,
 			MyDPIScale
 		));
-
 	}
 }
 	
@@ -57,10 +57,13 @@ void SDPIScaler::SetContent(TSharedRef<SWidget> InContent)
 
 void SDPIScaler::SetDPIScale(TAttribute<float> InDPIScale)
 {
-	DPIScale = InDPIScale;
+	if (SetAttribute(DPIScale, InDPIScale, EInvalidateWidgetReason::Layout))
+	{
+		InvalidatePrepass();
+	}
 }
 
-float SDPIScaler::GetRelativeLayoutScale(const FSlotBase& Child, float LayoutScaleMultiplier) const
+float SDPIScaler::GetRelativeLayoutScale(int32 ChildIndex, float LayoutScaleMultiplier) const
 {
 	return DPIScale.Get();
 }
