fix for crash when opening Locker Search bar on iOS with controller.

diff --git a/Engine/Source/Runtime/Slate/Private/Framework/Text/IOS/IOSPlatformTextField.cpp b/Engine/Source/Runtime/Slate/Private/Framework/Text/IOS/IOSPlatformTextField.cpp
--- a/Engine/Source/Runtime/Slate/Private/Framework/Text/IOS/IOSPlatformTextField.cpp
+++ b/Engine/Source/Runtime/Slate/Private/Framework/Text/IOS/IOSPlatformTextField.cpp
@@ -125,11 +125,30 @@ void FIOSPlatformTextField::ShowVirtualKeyboard(bool bShow, int32 UserIndex, TSh
 
 @implementation SlateTextField
 
+-(id)init
+{
+	self = [super init];
+	
+	if (self)
+	{
+		self->AlertController = nil;
+	}
+	
+	return self;
+}
+
 -(void)hide
 {
     if(AlertController != nil)
     {
-        [AlertController dismissViewControllerAnimated: YES completion: nil];
+		if ([AlertController respondsToSelector:@selector(dismissViewControllerAnimated: completion:)])
+		{
+        	[AlertController dismissViewControllerAnimated: YES completion: nil];
+		}
+		else
+		{
+			UE_LOG(LogTemp, Log, TEXT("AlertController didn't support needed selector"));
+		}
     }
 
     if(!TextWidget.IsValid())
@@ -156,26 +175,34 @@ void FIOSPlatformTextField::ShowVirtualKeyboard(bool bShow, int32 UserIndex, TSh
 									style:UIAlertActionStyleDefault
 									handler:^(UIAlertAction* action)
 									{
-										[AlertController dismissViewControllerAnimated : YES completion : nil];
+										if ([AlertController respondsToSelector:@selector(dismissViewControllerAnimated: completion:)])
+										{
+											[AlertController dismissViewControllerAnimated : YES completion : nil];
 
-										UITextField* AlertTextField = AlertController.textFields.firstObject;
-										TextEntry = FText::FromString(AlertTextField.text);
-										AlertController = nil;
+											UITextField* AlertTextField = AlertController.textFields.firstObject;
+											TextEntry = FText::FromString(AlertTextField.text);
+											AlertController = nil;
 										
-										FIOSAsyncTask* AsyncTask = [[FIOSAsyncTask alloc] init];
-										AsyncTask.GameThreadCallback = ^ bool(void)
-										{
-											if(TextWidget.IsValid())
+											FIOSAsyncTask* AsyncTask = [[FIOSAsyncTask alloc] init];
+											AsyncTask.GameThreadCallback = ^ bool(void)
 											{
-												TSharedPtr<IVirtualKeyboardEntry> TextEntryWidgetPin = TextWidget.Pin();
-												TextEntryWidgetPin->SetTextFromVirtualKeyboard(TextEntry, ETextEntryType::TextEntryAccepted);
-											}
+												if(TextWidget.IsValid())
+												{
+													TSharedPtr<IVirtualKeyboardEntry> TextEntryWidgetPin = TextWidget.Pin();
+													TextEntryWidgetPin->SetTextFromVirtualKeyboard(TextEntry, ETextEntryType::TextEntryAccepted);
+												}
 
-											// clear the TextWidget
+												// clear the TextWidget
+												TextWidget = nullptr;
+												return true;
+											};
+											[AsyncTask FinishedTask];
+										}
+										else
+										{
 											TextWidget = nullptr;
-											return true;
-										};
-										[AsyncTask FinishedTask];
+											UE_LOG(LogTemp, Log, TEXT("AlertController didn't support needed selector"));
+										}
 									}
 	];
 	UIAlertAction* cancelAction = [UIAlertAction
@@ -183,17 +210,25 @@ void FIOSPlatformTextField::ShowVirtualKeyboard(bool bShow, int32 UserIndex, TSh
 									style:UIAlertActionStyleDefault
 									handler:^(UIAlertAction* action)
 									{
-										[AlertController dismissViewControllerAnimated : YES completion : nil];
-										AlertController = nil;
-										
-										FIOSAsyncTask* AsyncTask = [[FIOSAsyncTask alloc] init];
-										AsyncTask.GameThreadCallback = ^ bool(void)
+										if ([AlertController respondsToSelector:@selector(dismissViewControllerAnimated: completion:)])
 										{
-											// clear the TextWidget
+											[AlertController dismissViewControllerAnimated : YES completion : nil];
+											AlertController = nil;
+										
+											FIOSAsyncTask* AsyncTask = [[FIOSAsyncTask alloc] init];
+											AsyncTask.GameThreadCallback = ^ bool(void)
+											{
+												// clear the TextWidget
+												TextWidget = nullptr;
+												return true;
+											};
+											[AsyncTask FinishedTask];
+										}
+								   		else
+								   		{
 											TextWidget = nullptr;
-											return true;
-										};
-										[AsyncTask FinishedTask];
+									   		UE_LOG(LogTemp, Log, TEXT("AlertController didn't support needed selector"));
+								   		}
 									}
 	];
 
