Refactor FontSizeMultiplie
#rnx

diff --git a/Engine/Source/Runtime/Slate/Private/Widgets/Text/SRichTextBlock.cpp b/Engine/Source/Runtime/Slate/Private/Widgets/Text/SRichTextBlock.cpp
--- a/Engine/Source/Runtime/Slate/Private/Widgets/Text/SRichTextBlock.cpp
+++ b/Engine/Source/Runtime/Slate/Private/Widgets/Text/SRichTextBlock.cpp
@@ -57,12 +57,14 @@ void SRichTextBlock::Construct( const FArguments& InArgs )
 	SetCanTick(false);
 }
 
-int32 SRichTextBlock::OnPaint( const FPaintArgs& Args, const FGeometry& AllottedGeometry, const FSlateRect& MyCullingRect, FSlateWindowElementList& OutDrawElements, int32 LayerId, const FWidgetStyle& InWidgetStyle, bool bParentEnabled ) const
+int32 SRichTextBlock::OnPaint(const FPaintArgs& Args, const FGeometry& AllottedGeometry, const FSlateRect& MyCullingRect, FSlateWindowElementList& OutDrawElements, int32 LayerId, const FWidgetStyle& InWidgetStyle, bool bParentEnabled) const
 {
 	const FVector2D LastDesiredSize = TextLayoutCache->GetDesiredSize();
 
+	const FGeometry TextBlockScaledGeometry = AllottedGeometry.MakeChild(AllottedGeometry.GetLocalSize() / TextBlockScale, FSlateLayoutTransform(TextBlockScale));
+
 	// OnPaint will also update the text layout cache if required
-	LayerId = TextLayoutCache->OnPaint(Args, AllottedGeometry, MyCullingRect, OutDrawElements, LayerId, InWidgetStyle, ShouldBeEnabled(bParentEnabled));
+	LayerId = TextLayoutCache->OnPaint(Args, TextBlockScaledGeometry, MyCullingRect, OutDrawElements, LayerId, InWidgetStyle, ShouldBeEnabled(bParentEnabled));
 
 	const FVector2D NewDesiredSize = TextLayoutCache->GetDesiredSize();
 
@@ -84,8 +86,7 @@ FVector2D SRichTextBlock::ComputeDesiredSize(float LayoutScaleMultiplier) const
 	// ComputeDesiredSize will also update the text layout cache if required
 	const FVector2D TextSize = TextLayoutCache->ComputeDesiredSize(
 		FSlateTextBlockLayout::FWidgetArgs(BoundText, HighlightText, WrapTextAt, AutoWrapText, WrappingPolicy, Margin, LineHeightPercentage, Justification),
-		LayoutScaleMultiplier, TextStyle
-		);
+		LayoutScaleMultiplier * TextBlockScale, TextStyle) * TextBlockScale;
 
 	return FVector2D(FMath::Max(TextSize.X, MinDesiredWidth.Get()), TextSize.Y);
 }
@@ -95,15 +96,18 @@ FChildren* SRichTextBlock::GetChildren()
 	return TextLayoutCache->GetChildren();
 }
 
-void SRichTextBlock::OnArrangeChildren( const FGeometry& AllottedGeometry, FArrangedChildren& ArrangedChildren ) const
+void SRichTextBlock::OnArrangeChildren(const FGeometry& AllottedGeometry, FArrangedChildren& ArrangedChildren) const
 {
-	TextLayoutCache->ArrangeChildren( AllottedGeometry, ArrangedChildren );
+	const FGeometry TextBlockScaledGeometry = AllottedGeometry.MakeChild(AllottedGeometry.GetLocalSize() / TextBlockScale, FSlateLayoutTransform(TextBlockScale));
+
+	TextLayoutCache->ArrangeChildren(TextBlockScaledGeometry, ArrangedChildren);
 }
 
 void SRichTextBlock::SetText( const TAttribute<FText>& InTextAttr )
 {
 	BoundText = InTextAttr;
 	Invalidate(EInvalidateWidget::LayoutAndVolatility);
+	InvalidatePrepass();
 }
 
 void SRichTextBlock::SetHighlightText( const TAttribute<FText>& InHighlightText )
@@ -132,6 +136,7 @@ void SRichTextBlock::SetWrapTextAt(const TAttribute<float>& InWrapTextAt)
 void SRichTextBlock::SetAutoWrapText(const TAttribute<bool>& InAutoWrapText)
 {
 	SetAttribute(AutoWrapText, InAutoWrapText, EInvalidateWidgetReason::Layout);
+	InvalidatePrepass();
 }
 
 void SRichTextBlock::SetWrappingPolicy(const TAttribute<ETextWrappingPolicy>& InWrappingPolicy)
@@ -174,6 +179,13 @@ void SRichTextBlock::SetDecoratorStyleSet(const ISlateStyle* NewDecoratorStyleSe
 	}
 }
 
+void SRichTextBlock::SetTextBlockScale(const float NewTextBlockScale)
+{
+	TextBlockScale = NewTextBlockScale;
+	Invalidate(EInvalidateWidget::Layout);
+	InvalidatePrepass();
+}
+
 void SRichTextBlock::Refresh()
 {
 	TextLayoutCache->DirtyContent();
@@ -194,13 +206,4 @@ bool SRichTextBlock::ComputeVolatility() const
 		|| MinDesiredWidth.IsBound();
 }
 
-void SRichTextBlock::SetFontSizeMultiplier(const float NewFontSizeMultiplier)
-{
-	if (Marshaller)
-	{
-		Marshaller->SetFontSizeMultiplier(NewFontSizeMultiplier);
-		Invalidate(EInvalidateWidget::Layout);
-	}
-}
-
 #endif //WITH_FANCY_TEXT
diff --git a/Engine/Source/Runtime/Slate/Private/Widgets/Text/SlateTextBlockLayout.cpp b/Engine/Source/Runtime/Slate/Private/Widgets/Text/SlateTextBlockLayout.cpp
--- a/Engine/Source/Runtime/Slate/Private/Widgets/Text/SlateTextBlockLayout.cpp
+++ b/Engine/Source/Runtime/Slate/Private/Widgets/Text/SlateTextBlockLayout.cpp
@@ -94,9 +94,14 @@ FVector2D FSlateTextBlockLayout::GetDesiredSize() const
 	return TextLayout->GetSize();
 }
 
+float FSlateTextBlockLayout::GetLayoutScale() const
+{
+	return TextLayout->GetScale();
+}
+
 int32 FSlateTextBlockLayout::OnPaint(const FPaintArgs& InPaintArgs, const FGeometry& InAllottedGeometry, const FSlateRect& InClippingRect, FSlateWindowElementList& OutDrawElements, int32 LayerId, const FWidgetStyle& InWidgetStyle, bool bParentEnabled)
 {
-	// Store a new cached size.
+	// Store a new cached size with the scale
 	CachedSize = InAllottedGeometry.GetLocalSize();
 
 	// Recompute wrapping in case the cached size changed.
diff --git a/Engine/Source/Runtime/Slate/Public/Widgets/Text/SRichTextBlock.h b/Engine/Source/Runtime/Slate/Public/Widgets/Text/SRichTextBlock.h
--- a/Engine/Source/Runtime/Slate/Public/Widgets/Text/SRichTextBlock.h
+++ b/Engine/Source/Runtime/Slate/Public/Widgets/Text/SRichTextBlock.h
@@ -216,9 +216,8 @@ public:
 	 */
 	void Refresh();
 
-	/** set the multiplier affected at the font in the text */
-	void SetFontSizeMultiplier(const float NewFontSizeMultiplier);
-
+	/** set the scale value at the TextLayout*/
+	void SetTextBlockScale(const float NewTextBlockScale);
 
 protected:
 	//~ SWidget interface
@@ -261,6 +260,9 @@ private:
 	/** Prevents the text block from being smaller than desired in certain cases (e.g. when it is empty) */
 	TAttribute<float> MinDesiredWidth;
 
+	/** Use to Scale the entire Text Block*/
+	float TextBlockScale = 1.0f;
+
 	/**  */
 	TSharedPtr<FRichTextLayoutMarshaller> Marshaller;
 };
diff --git a/Engine/Source/Runtime/Slate/Public/Widgets/Text/SlateTextBlockLayout.h b/Engine/Source/Runtime/Slate/Public/Widgets/Text/SlateTextBlockLayout.h
--- a/Engine/Source/Runtime/Slate/Public/Widgets/Text/SlateTextBlockLayout.h
+++ b/Engine/Source/Runtime/Slate/Public/Widgets/Text/SlateTextBlockLayout.h
@@ -68,6 +68,13 @@ public:
 	 */
 	FVector2D GetDesiredSize() const;
 
+
+	/**
+	 * Get the TextLayout scale.
+	 */
+	float GetLayoutScale() const;
+
+
 	/**
 	 * Paint this layout, updating the internal cache as required
 	 */
