Fix bAllowsRightClickDragScrolling not being initialized
#rnx

diff --git a/Engine/Source/Runtime/Slate/Private/Widgets/Layout/SScrollBox.cpp b/Engine/Source/Runtime/Slate/Private/Widgets/Layout/SScrollBox.cpp
--- a/Engine/Source/Runtime/Slate/Private/Widgets/Layout/SScrollBox.cpp
+++ b/Engine/Source/Runtime/Slate/Private/Widgets/Layout/SScrollBox.cpp
@@ -124,6 +124,7 @@ void SScrollBox::Construct( const FArguments& InArgs )
 	Orientation = InArgs._Orientation;
 	bScrollToEnd = false;
 	bIsScrollingActiveTimerRegistered = false;
+	bAllowsRightClickDragScrolling = false;
 	ConsumeMouseWheel = InArgs._ConsumeMouseWheel;
 	TickScrollDelta = 0;
 	AllowOverscroll = InArgs._AllowOverscroll;
