PlatformTextField hololens and lumin build fixes.

diff --git a/Engine/Source/Runtime/Slate/Public/Framework/Text/Hololens/HololensPlatformTextField.h b/Engine/Source/Runtime/Slate/Public/Framework/Text/Hololens/HololensPlatformTextField.h
new file mode 100644
--- /dev/null
+++ b/Engine/Source/Runtime/Slate/Public/Framework/Text/Hololens/HololensPlatformTextField.h
@@ -0,0 +1,5 @@
+// Copyright 1998-2019 Epic Games, Inc. All Rights Reserved.
+
+#pragma once
+
+#include "Framework/Text/GenericPlatformTextField.h"
diff --git a/Engine/Source/Runtime/Slate/Public/Framework/Text/Lumin/LuminPlatformTextField.h b/Engine/Source/Runtime/Slate/Public/Framework/Text/Lumin/LuminPlatformTextField.h
new file mode 100644
--- /dev/null
+++ b/Engine/Source/Runtime/Slate/Public/Framework/Text/Lumin/LuminPlatformTextField.h
@@ -0,0 +1,5 @@
+// Copyright 1998-2019 Epic Games, Inc. All Rights Reserved.
+
+#pragma once
+
+#include "Framework/Text/GenericPlatformTextField.h"
