Fix issue with non-unity Android builds

#ue4
#android

diff --git a/Engine/Source/Runtime/Slate/Public/Framework/Text/Android/AndroidPlatformTextField.h b/Engine/Source/Runtime/Slate/Public/Framework/Text/Android/AndroidPlatformTextField.h
--- a/Engine/Source/Runtime/Slate/Public/Framework/Text/Android/AndroidPlatformTextField.h
+++ b/Engine/Source/Runtime/Slate/Public/Framework/Text/Android/AndroidPlatformTextField.h
@@ -2,6 +2,11 @@
 
 #pragma once
 
+// some android platforms disable this but not defined otherwise if AndroidPlatform.h not included first
+#ifndef USE_ANDROID_JNI
+	#define USE_ANDROID_JNI							1
+#endif
+
 #if USE_ANDROID_JNI
 
 #include "CoreMinimal.h"
