Fixed p4 ignored some files in the merge from Main.

diff --git a/Engine/Source/Developer/ToolMenus/Public/ToolMenu.h b/Engine/Source/Developer/ToolMenus/Public/ToolMenu.h
--- a/Engine/Source/Developer/ToolMenus/Public/ToolMenu.h
+++ b/Engine/Source/Developer/ToolMenus/Public/ToolMenu.h
@@ -62,6 +62,9 @@ public:
 	/** returns "SubMenuC.SubMenuD" for menu "ModuleA.MenuB.SubMenuC.SubMenuD" */
 	FString GetSubMenuNamePath() const;
 
+	/* Set support for extenders */
+	void SetExtendersEnabled(bool bEnabled);
+
 	//~ Begin UToolMenuBase Interface
 	virtual bool IsEditing() const override;
 	virtual FName GetSectionName(const FName InEntryName) const override;
@@ -164,6 +167,7 @@ private:
 
 	bool bRegistered;
 	bool bIsRegistering;
+	bool bExtendersEnabled;
 
 	const ISlateStyle* StyleSet;
 
diff --git a/Engine/Source/Runtime/Slate/Public/Framework/MultiBox/MultiBoxBuilder.h b/Engine/Source/Runtime/Slate/Public/Framework/MultiBox/MultiBoxBuilder.h
--- a/Engine/Source/Runtime/Slate/Public/Framework/MultiBox/MultiBoxBuilder.h
+++ b/Engine/Source/Runtime/Slate/Public/Framework/MultiBox/MultiBoxBuilder.h
@@ -119,6 +119,12 @@ public:
 	/** @return  The customization settings for the box being built */
 	FMultiBoxCustomization GetCustomization() const;
 
+	/** Sets extender support */
+	void SetExtendersEnabled(bool bEnabled) { bExtendersEnabled = bEnabled; }
+
+	/** @return True if extenders are enabled */
+	bool ExtendersEnabled() const { return bExtendersEnabled; }
+
 protected:
 	/** Applies any potential extension hooks at the current place */
 	virtual void ApplyHook(FName InExtensionHook, EExtensionHook::Position HookPosition) = 0;
@@ -143,6 +149,9 @@ protected:
 
 	/** Name of the menu */
 	FName MenuName;
+
+	/** If extenders are enabled */
+	bool bExtendersEnabled;
 };
 
 
