Pass over Use Mouse For Touch and cleaning up some redundant code.

diff --git a/Engine/Source/Editor/LevelEditor/Private/SLevelViewport.cpp b/Engine/Source/Editor/LevelEditor/Private/SLevelViewport.cpp
--- a/Engine/Source/Editor/LevelEditor/Private/SLevelViewport.cpp
+++ b/Engine/Source/Editor/LevelEditor/Private/SLevelViewport.cpp
@@ -4014,8 +4014,6 @@ void SLevelViewport::StartPlayInEditorSession(UGameViewportClient* PlayClient, c
 	// Focus will be set when the game viewport is registered
 	FSlateApplication::Get().ClearKeyboardFocus(EFocusCause::SetDirectly);
 
-	PlayClient->SetPlayInEditorUseMouseForTouch(GetDefault<ULevelEditorPlaySettings>()->UseMouseForTouch);
-
 	// Attach global play world actions widget to view port
 	ActiveViewport = MakeShareable( new FSceneViewport( PlayClient, ViewportWidget) );
 	ActiveViewport->SetPlayInEditorViewport( true );
diff --git a/Engine/Source/Runtime/Engine/Classes/Engine/GameViewportClient.h b/Engine/Source/Runtime/Engine/Classes/Engine/GameViewportClient.h
--- a/Engine/Source/Runtime/Engine/Classes/Engine/GameViewportClient.h
+++ b/Engine/Source/Runtime/Engine/Classes/Engine/GameViewportClient.h
@@ -606,13 +606,9 @@ public:
 	/** FViewport interface */
 	virtual bool ShouldDPIScaleSceneCanvas() const override { return false; }
 
-#if WITH_EDITOR
-	void SetPlayInEditorUseMouseForTouch(bool bUseMouseForTouch);
-#endif
+	bool GetUseMouseForTouch() const;
 
 protected:
-
-	bool GetUseMouseForTouch() const;
 	void SetCurrentBufferVisualizationMode(FName NewBufferVisualizationMode) { CurrentBufferVisualizationMode = NewBufferVisualizationMode; }
 	FName GetCurrentBufferVisualizationMode() const { return CurrentBufferVisualizationMode; }
 	bool HasAudioFocus() const { return bHasAudioFocus; }
diff --git a/Engine/Source/Runtime/Engine/Private/GameViewportClient.cpp b/Engine/Source/Runtime/Engine/Private/GameViewportClient.cpp
--- a/Engine/Source/Runtime/Engine/Private/GameViewportClient.cpp
+++ b/Engine/Source/Runtime/Engine/Private/GameViewportClient.cpp
@@ -3820,17 +3820,10 @@ bool UGameViewportClient::IsSimulateInEditorViewport() const
 	return GameViewport ? GameViewport->GetPlayInEditorIsSimulate() : false;
 }
 
-#if WITH_EDITOR
-void UGameViewportClient::SetPlayInEditorUseMouseForTouch(bool bInUseMouseForTouch)
-{
-	bUseMouseForTouchInEditor = bInUseMouseForTouch;
-}
-#endif
-
 bool UGameViewportClient::GetUseMouseForTouch() const
 {
 #if WITH_EDITOR
-	return bUseMouseForTouchInEditor || GetDefault<UInputSettings>()->bUseMouseForTouch;
+	return GetDefault<ULevelEditorPlaySettings>()->UseMouseForTouch || GetDefault<UInputSettings>()->bUseMouseForTouch;
 #else
 	return GetDefault<UInputSettings>()->bUseMouseForTouch;
 #endif
diff --git a/Engine/Source/Runtime/Slate/Private/Framework/Application/SlateApplication.cpp b/Engine/Source/Runtime/Slate/Private/Framework/Application/SlateApplication.cpp
--- a/Engine/Source/Runtime/Slate/Private/Framework/Application/SlateApplication.cpp
+++ b/Engine/Source/Runtime/Slate/Private/Framework/Application/SlateApplication.cpp
@@ -4400,7 +4400,7 @@ bool FSlateApplication::OnMouseDown(const TSharedPtr< FGenericWindow >& Platform
 bool FSlateApplication::OnMouseDown( const TSharedPtr< FGenericWindow >& PlatformWindow, const EMouseButtons::Type Button, const FVector2D CursorPos )
 {
 	// convert a left mouse button click to touch event if we are faking it
-	if ((bIsFakingTouch || bIsGameFakingTouch) && Button == EMouseButtons::Left)
+	if (IsFakingTouchEvents() && Button == EMouseButtons::Left)
 	{
 		bIsFakingTouched = true;
 		return OnTouchStarted( PlatformWindow, PlatformApplication->Cursor->GetPosition(), 1.0f, 0, 0 );
@@ -5057,7 +5057,7 @@ bool FSlateApplication::OnMouseDoubleClick( const TSharedPtr< FGenericWindow >&
 
 bool FSlateApplication::OnMouseDoubleClick( const TSharedPtr< FGenericWindow >& PlatformWindow, const EMouseButtons::Type Button, const FVector2D CursorPos )
 {
-	if (bIsFakingTouch || bIsGameFakingTouch)
+	if (IsFakingTouchEvents())
 	{
 		bIsFakingTouched = true;
 		return OnTouchStarted(PlatformWindow, PlatformApplication->Cursor->GetPosition(), 1.0f, 0, 0);
@@ -5144,7 +5144,7 @@ bool FSlateApplication::OnMouseUp( const EMouseButtons::Type Button )
 bool FSlateApplication::OnMouseUp( const EMouseButtons::Type Button, const FVector2D CursorPos )
 {
 	// convert left mouse click to touch event if we are faking it	
-	if ((bIsFakingTouch || bIsGameFakingTouch) && Button == EMouseButtons::Left)
+	if (IsFakingTouchEvents() && Button == EMouseButtons::Left)
 	{
 		bIsFakingTouched = false;
 		return OnTouchEnded(PlatformApplication->Cursor->GetPosition(), 0, 0);
@@ -5318,7 +5318,7 @@ FReply FSlateApplication::RouteMouseWheelOrGestureEvent(const FWidgetPath& Widge
 bool FSlateApplication::OnMouseMove()
 {
 	// If the left button is pressed we fake 
-	if ((bIsFakingTouched || bIsGameFakingTouch) && GetPressedMouseButtons().Contains(EKeys::LeftMouseButton))
+	if (IsFakingTouchEvents() && (GetPressedMouseButtons().Num() == 0 || GetPressedMouseButtons().Contains(EKeys::LeftMouseButton)))
 	{
 		// convert to touch event if we are faking it
 		if (bIsFakingTouched)
@@ -5362,7 +5362,7 @@ bool FSlateApplication::OnMouseMove()
 bool FSlateApplication::OnRawMouseMove( const int32 X, const int32 Y )
 {
     // We fake a move only if the left mous button is down
-	if ((bIsFakingTouch || bIsGameFakingTouch) && GetPressedMouseButtons().Contains(EKeys::LeftMouseButton))
+	if (IsFakingTouchEvents() && (GetPressedMouseButtons().Num() == 0 || GetPressedMouseButtons().Contains(EKeys::LeftMouseButton)))
 	{
 		// convert to touch event if we are faking it
 		if (bIsFakingTouched)
