potential fix for crash while trying to hide the system text field

diff --git a/Engine/Source/Runtime/Slate/Private/Framework/Text/IOS/IOSPlatformTextField.cpp b/Engine/Source/Runtime/Slate/Private/Framework/Text/IOS/IOSPlatformTextField.cpp
--- a/Engine/Source/Runtime/Slate/Private/Framework/Text/IOS/IOSPlatformTextField.cpp
+++ b/Engine/Source/Runtime/Slate/Private/Framework/Text/IOS/IOSPlatformTextField.cpp
@@ -57,7 +57,10 @@ FIOSPlatformTextField::~FIOSPlatformTextField()
 		dispatch_async(dispatch_get_main_queue(), ^{
             NSLog(@"Finally releasing text field %@", LocalTextField);
 #if !PLATFORM_TVOS
-            [LocalTextField hide];
+			if (LocalTextField != nullptr && [LocalTextField respondsToSelector:@selector(hide:)])
+			{
+            	[LocalTextField hide];
+			}
 #endif
 		});
 	}
