Allowing per-user nav configs to be established at the Slate level
- FSlateUser has an optional TSharedPtr<FNavigationConfig> member now that, when set, is used instead of the global one that lives on FSlateApplication

diff --git a/Engine/Source/Runtime/Slate/Private/Framework/Application/SlateApplication.cpp b/Engine/Source/Runtime/Slate/Private/Framework/Application/SlateApplication.cpp
--- a/Engine/Source/Runtime/Slate/Private/Framework/Application/SlateApplication.cpp
+++ b/Engine/Source/Runtime/Slate/Private/Framework/Application/SlateApplication.cpp
@@ -1747,38 +1747,26 @@ static bool IsFocusInViewport(const TSet<TWeakPtr<SViewport>> Viewports, const F
 
 EUINavigation FSlateApplication::GetNavigationDirectionFromKey(const FKeyEvent& InKeyEvent) const
 {
-#if WITH_EDITOR
-	// Check if the focused widget is an editor widget or a PIE widget so we know which config to use.
-	TSharedPtr<const FSlateUser> User = GetUser(GetUserIndexForKeyboard());
-	if (User && !IsFocusInViewport(AllGameViewports, User->GetWeakFocusPath()))
-	{
-		return EditorNavigationConfig->GetNavigationDirectionFromKey(InKeyEvent);
-	}
-#endif
-	return NavigationConfig->GetNavigationDirectionFromKey(InKeyEvent);
+	TSharedRef<FNavigationConfig> RelevantNavConfig = GetRelevantNavConfig(InKeyEvent.GetUserIndex());
+	return RelevantNavConfig->GetNavigationDirectionFromKey(InKeyEvent);
 }
 
 EUINavigation FSlateApplication::GetNavigationDirectionFromAnalog(const FAnalogInputEvent& InAnalogEvent)
 {
-#if WITH_EDITOR
-	// Check if the focused widget is an editor widget or a PIE widget so we know which config to use.
-	TSharedPtr<const FSlateUser> User = GetUser(GetUserIndexForKeyboard());
-	if (User && !IsFocusInViewport(AllGameViewports, User->GetWeakFocusPath()))
-	{
-		return EditorNavigationConfig->GetNavigationDirectionFromKey(InAnalogEvent);
-	}
-#endif
-	return NavigationConfig->GetNavigationDirectionFromAnalog(InAnalogEvent);
+	TSharedRef<FNavigationConfig> RelevantNavConfig = GetRelevantNavConfig(InAnalogEvent.GetUserIndex());
+	return RelevantNavConfig->GetNavigationDirectionFromAnalog(InAnalogEvent);
 }
 
 EUINavigationAction FSlateApplication::GetNavigationActionFromKey(const FKeyEvent& InKeyEvent) const
 {
-	return NavigationConfig->GetNavigationActionFromKey(InKeyEvent);
+	TSharedRef<FNavigationConfig> RelevantNavConfig = GetRelevantNavConfig(InKeyEvent.GetUserIndex());
+	return RelevantNavConfig->GetNavigationActionFromKey(InKeyEvent);
 }
 
 EUINavigationAction FSlateApplication::GetNavigationActionForKey(const FKey& InKey) const
 {
 PRAGMA_DISABLE_DEPRECATION_WARNINGS
+	// Not enough info to pick the best config, so this can only use the default
 	return NavigationConfig->GetNavigationActionForKey(InKey);
 PRAGMA_ENABLE_DEPRECATION_WARNINGS
 }
@@ -2705,7 +2693,7 @@ bool FSlateApplication::SetUserFocus(FSlateUser& User, const FWidgetPath& InFocu
 			ProcessReply(InFocusPath, Reply, nullptr, nullptr, User.GetUserIndex());
 		}
 
-		NavigationConfig->OnNavigationChangedFocus(OldFocusedWidget, NewFocusedWidget, FocusEvent);
+		GetRelevantNavConfig(User.GetUserIndex())->OnNavigationChangedFocus(OldFocusedWidget, NewFocusedWidget, FocusEvent);
 
 #if WITH_ACCESSIBILITY
 		GetAccessibleMessageHandler()->OnWidgetEventRaised(NewFocusedWidget.ToSharedRef(), EAccessibleEvent::FocusChange, false, true);
@@ -3981,6 +3969,29 @@ bool FSlateApplication::ShowUserFocus(const TSharedPtr<const SWidget> Widget) co
 	return false;
 }
 
+TSharedRef<FNavigationConfig> FSlateApplication::GetRelevantNavConfig(int32 UserIndex) const
+{
+	TSharedPtr<FNavigationConfig> RelevantNavConfig = NavigationConfig;
+	if (TSharedPtr<const FSlateUser> User = GetUser(UserIndex))
+	{
+#if WITH_EDITOR
+		// Check if the focused widget is an editor widget or a PIE widget so we know which config to use.
+		if (UserIndex == GetUserIndexForKeyboard() && !IsFocusInViewport(AllGameViewports, User->GetWeakFocusPath()))
+		{
+			RelevantNavConfig = EditorNavigationConfig;
+		}
+		else
+#endif
+		if (TSharedPtr<FNavigationConfig> UserNavConfig = User->GetUserNavigationConfig())
+		{
+			// Use the user's personal config if it has one assigned
+			RelevantNavConfig = UserNavConfig;
+		}
+	}
+
+	return RelevantNavConfig.ToSharedRef();
+}
+
 bool FSlateApplication::HasUserFocusedDescendants(const TSharedRef< const SWidget >& Widget, int32 UserIndex) const
 {
 	TSharedPtr<const FSlateUser> User = GetUser(UserIndex);
diff --git a/Engine/Source/Runtime/Slate/Private/Framework/Application/SlateUser.cpp b/Engine/Source/Runtime/Slate/Private/Framework/Application/SlateUser.cpp
--- a/Engine/Source/Runtime/Slate/Private/Framework/Application/SlateUser.cpp
+++ b/Engine/Source/Runtime/Slate/Private/Framework/Application/SlateUser.cpp
@@ -2,6 +2,7 @@
 
 #include "Framework/Application/SlateUser.h"
 #include "Framework/Application/SlateApplication.h"
+#include "Framework/Application/NavigationConfig.h"
 #include "Widgets/SWindow.h"
 #include "Widgets/SWeakWidget.h"
 
@@ -532,6 +533,21 @@ void FSlateUser::CloseTooltip()
 	}
 }
 
+void FSlateUser::SetUserNavigationConfig(TSharedPtr<FNavigationConfig> InNavigationConfig)
+{
+	if (UserNavigationConfig)
+	{
+		UserNavigationConfig->OnUnregister();
+	}
+
+	UserNavigationConfig = InNavigationConfig;
+	
+	if (InNavigationConfig)
+	{
+		InNavigationConfig->OnRegister();
+	}
+}
+
 bool FSlateUser::IsTouchPointerActive(int32 TouchPointerIndex) const
 {
 	return TouchPointerIndex < (int32)ETouchIndex::CursorPointerIndex && PointerPositionsByIndex.Contains(TouchPointerIndex);
diff --git a/Engine/Source/Runtime/Slate/Public/Framework/Application/SlateApplication.h b/Engine/Source/Runtime/Slate/Public/Framework/Application/SlateApplication.h
--- a/Engine/Source/Runtime/Slate/Public/Framework/Application/SlateApplication.h
+++ b/Engine/Source/Runtime/Slate/Public/Framework/Application/SlateApplication.h
@@ -955,6 +955,8 @@ protected:
 	virtual bool IsWidgetDirectlyHovered(const TSharedPtr<const SWidget> Widget) const override;
 	virtual bool ShowUserFocus(const TSharedPtr<const SWidget> Widget) const override;
 
+	TSharedRef<FNavigationConfig> GetRelevantNavConfig(int32 UserIndex) const;
+
 	/**
 	 * Pumps and ticks the platform.
 	 */
diff --git a/Engine/Source/Runtime/Slate/Public/Framework/Application/SlateUser.h b/Engine/Source/Runtime/Slate/Public/Framework/Application/SlateUser.h
--- a/Engine/Source/Runtime/Slate/Public/Framework/Application/SlateUser.h
+++ b/Engine/Source/Runtime/Slate/Public/Framework/Application/SlateUser.h
@@ -11,6 +11,7 @@ class SWindow;
 class ICursor;
 class FDragDropOperation;
 class FMenuStack;
+class FNavigationConfig;
 
 /** Handle to a virtual user of slate. */
 class SLATE_API FSlateVirtualUserHandle
@@ -119,6 +120,9 @@ public:
 
 	const FGestureDetector& GetGestureDetector() const { return GestureDetector; }
 
+	void SetUserNavigationConfig(TSharedPtr<FNavigationConfig> InNavigationConfig);
+	TSharedPtr<FNavigationConfig> GetUserNavigationConfig() const { return UserNavigationConfig; }
+
 SLATE_SCOPE:
 	static TSharedRef<FSlateUser> Create(int32 InUserIndex, TSharedPtr<ICursor> InCursor);
 	
@@ -198,6 +202,9 @@ private:
 	/** The cursor this user is in control of. Guaranteed to be valid for all real users, absence implies this is a virtual user. */
 	TSharedPtr<ICursor> Cursor;
 
+	/** SlateUsers can optionally be individually assigned a navigation config to use. This overrides the global nav config that lives on FSlateApplication when valid. */
+	TSharedPtr<FNavigationConfig> UserNavigationConfig;
+
 	/** Whether this user is currently drawing the cursor onscreen each frame */
 	bool bCanDrawCursor = true;
 
