Slate - Adding stats around expensive OS calls in slate, assume we're over a slate window for processing mouse input as long as the application is activated.  This will avoid making OS calls that could potentially hitch when the app/game is actively being used.

diff --git a/Engine/Source/Runtime/ApplicationCore/Private/Windows/WindowsApplication.cpp b/Engine/Source/Runtime/ApplicationCore/Private/Windows/WindowsApplication.cpp
--- a/Engine/Source/Runtime/ApplicationCore/Private/Windows/WindowsApplication.cpp
+++ b/Engine/Source/Runtime/ApplicationCore/Private/Windows/WindowsApplication.cpp
@@ -392,6 +392,8 @@ static TSharedPtr< FWindowsWindow > FindWindowByHWND(const TArray< TSharedRef< F
 
 bool FWindowsApplication::IsCursorDirectlyOverSlateWindow() const
 {
+	QUICK_SCOPE_CYCLE_COUNTER(STAT_STAT_IsCursorDirectlyOverSlateWindow);
+
 	POINT CursorPos;
 	BOOL bGotPoint = ::GetCursorPos(&CursorPos);
 	if (bGotPoint)
diff --git a/Engine/Source/Runtime/Slate/Private/Framework/Application/SlateApplication.cpp b/Engine/Source/Runtime/Slate/Private/Framework/Application/SlateApplication.cpp
--- a/Engine/Source/Runtime/Slate/Private/Framework/Application/SlateApplication.cpp
+++ b/Engine/Source/Runtime/Slate/Private/Framework/Application/SlateApplication.cpp
@@ -5426,7 +5426,7 @@ bool FSlateApplication::ProcessMouseMoveEvent( const FPointerEvent& MouseEvent,
 	// When the event came from the OS, we are guaranteed to be over a slate window.
 	// Otherwise, we are synthesizing a MouseMove ourselves, and must verify that the
 	// cursor is indeed over a Slate window.
-	const bool bOverSlateWindow = !bIsSynthetic || PlatformApplication->IsCursorDirectlyOverSlateWindow();
+	const bool bOverSlateWindow = !bIsSynthetic || IsActive() || PlatformApplication->IsCursorDirectlyOverSlateWindow();
 	
 	FWidgetPath WidgetsUnderCursor = bOverSlateWindow
 		? LocateWindowUnderMouse(MouseEvent.GetScreenSpacePosition(), GetInteractiveTopLevelWindows(), false, MouseEvent.GetUserIndex())
diff --git a/Engine/Source/Runtime/Slate/Private/Framework/Application/SlateUser.cpp b/Engine/Source/Runtime/Slate/Private/Framework/Application/SlateUser.cpp
--- a/Engine/Source/Runtime/Slate/Private/Framework/Application/SlateUser.cpp
+++ b/Engine/Source/Runtime/Slate/Private/Framework/Application/SlateUser.cpp
@@ -1121,10 +1121,12 @@ void FSlateUser::UpdateTooltip(const FMenuStack& MenuStack, bool bCanSpawnNewToo
 		IsInGameThread() &&					// We should never allow the slate loading thread to create new windows or interact with the hittest grid
 		!SlateApp.IsUsingHighPrecisionMouseMovment() && // If we are using HighPrecision movement then we can't rely on the OS cursor to be accurate
 		!IsDragDropping() &&				// We must not currently be in the middle of a drag-drop action
-		
-		//@todo DanH: We need to check if OUR cursor is over a slate window, not just the platform cursor. 
-		//		See about adding FPlatformApplication::GetSlateWindowUnderPoint(FVector2D) or something.
-		SlateApp.GetPlatformApplication()->IsCursorDirectlyOverSlateWindow(); // The cursor must be over a Slate window
+		(
+			SlateApp.IsActive() || // Assume we need update if app is active
+			//@todo DanH: We need to check if OUR cursor is over a slate window, not just the platform cursor. 
+			//		See about adding FPlatformApplication::GetSlateWindowUnderPoint(FVector2D) or something.
+			SlateApp.GetPlatformApplication()->IsCursorDirectlyOverSlateWindow() // The cursor must be over a Slate window
+		);
 
 	if (bCheckForTooltipChanges)
 	{
