[Slate] Fixed TabManager to accept multiple calls to InvokeTab(). Fixes UE-89183 (crash in Unreal Insights when opening the Networking Insights tab).

diff --git a/Engine/Source/Runtime/Slate/Private/Framework/Docking/TabManager.cpp b/Engine/Source/Runtime/Slate/Private/Framework/Docking/TabManager.cpp
--- a/Engine/Source/Runtime/Slate/Private/Framework/Docking/TabManager.cpp
+++ b/Engine/Source/Runtime/Slate/Private/Framework/Docking/TabManager.cpp
@@ -1658,7 +1658,7 @@ TSharedPtr<SDockTab> FTabManager::SpawnTab(const FTabId& TabId, const TSharedPtr
 			bSpawningAllowedBySpawner = Spawner->CanSpawnTab.Execute(FSpawnTabArgs(ParentWindow, TabId));
 		}
 
-		if (bSpawningAllowedBySpawner && !Spawner->SpawnedTabPtr.IsValid())
+		if (bSpawningAllowedBySpawner && (!Spawner->SpawnedTabPtr.IsValid() || Spawner->OnFindTabToReuse.IsBound()))
 		{
 			NewTabWidget = Spawner->OnSpawnTab.Execute(FSpawnTabArgs(ParentWindow, TabId));
 			NewTabWidget->SetLayoutIdentifier(TabId);
