Disallow focus out events on editable text if GIsTransacting

diff --git a/Engine/Source/Runtime/Slate/Private/Widgets/Text/SlateEditableTextLayout.cpp b/Engine/Source/Runtime/Slate/Private/Widgets/Text/SlateEditableTextLayout.cpp
--- a/Engine/Source/Runtime/Slate/Private/Widgets/Text/SlateEditableTextLayout.cpp
+++ b/Engine/Source/Runtime/Slate/Private/Widgets/Text/SlateEditableTextLayout.cpp
@@ -754,7 +754,9 @@ bool FSlateEditableTextLayout::HandleFocusReceived(const FFocusEvent& InFocusEve
 
 bool FSlateEditableTextLayout::HandleFocusLost(const FFocusEvent& InFocusEvent)
 {
-	if (ActiveContextMenu.IsValid())
+	// GIsTransacting will be true if the user does Undo while the local undo stack is empty, since
+	// CanExecuteUndo() will fail and the event will bubble up to the Editor's undo stack.
+	if (ActiveContextMenu.IsValid() || GIsTransacting)
 	{
 		return false;
 	}
