- Fix non-editor compile error
#rnx

diff --git a/Engine/Source/Runtime/Slate/Private/Widgets/Text/SlateEditableTextLayout.cpp b/Engine/Source/Runtime/Slate/Private/Widgets/Text/SlateEditableTextLayout.cpp
--- a/Engine/Source/Runtime/Slate/Private/Widgets/Text/SlateEditableTextLayout.cpp
+++ b/Engine/Source/Runtime/Slate/Private/Widgets/Text/SlateEditableTextLayout.cpp
@@ -756,7 +756,14 @@ bool FSlateEditableTextLayout::HandleFocusLost(const FFocusEvent& InFocusEvent)
 {
 	// GIsTransacting will be true if the user does Undo while the local undo stack is empty, since
 	// CanExecuteUndo() will fail and the event will bubble up to the Editor's undo stack.
-	if (ActiveContextMenu.IsValid() || GIsTransacting)
+#if WITH_EDITORONLY_DATA
+	if (GIsTransacting)
+	{
+		return false;
+	}
+#endif
+
+	if (ActiveContextMenu.IsValid())
 	{
 		return false;
 	}
