Undo is now always possible in editable text fields.

diff --git a/Engine/Source/Runtime/Slate/Private/Widgets/Text/SlateEditableTextLayout.cpp b/Engine/Source/Runtime/Slate/Private/Widgets/Text/SlateEditableTextLayout.cpp
--- a/Engine/Source/Runtime/Slate/Private/Widgets/Text/SlateEditableTextLayout.cpp
+++ b/Engine/Source/Runtime/Slate/Private/Widgets/Text/SlateEditableTextLayout.cpp
@@ -754,15 +754,6 @@ bool FSlateEditableTextLayout::HandleFocusReceived(const FFocusEvent& InFocusEve
 
 bool FSlateEditableTextLayout::HandleFocusLost(const FFocusEvent& InFocusEvent)
 {
-	// GIsTransacting will be true if the user does Undo while the local undo stack is empty, since
-	// CanExecuteUndo() will fail and the event will bubble up to the Editor's undo stack.
-#if WITH_EDITORONLY_DATA
-	if (GIsTransacting)
-	{
-		return false;
-	}
-#endif
-
 	if (ActiveContextMenu.IsValid())
 	{
 		return false;
@@ -3073,7 +3064,12 @@ void FSlateEditableTextLayout::MakeUndoState(SlateEditableTextTypes::FUndoState&
 
 bool FSlateEditableTextLayout::CanExecuteUndo() const
 {
-	return !OwnerWidget->IsTextReadOnly() && UndoStates.Num() > 0 && !TextInputMethodContext->IsComposing();
+	// Previously, if UndoStates was empty then the event would bubble up back to the Editor and trigger an undo.
+	// Now, the editable text always catches the undo event so that it never bubbles up. This prevents bugs such
+	// as undo-ing in a search box triggering an undo, and various issues related to undo-ing property changes
+	// while focused in to the property widget being undone.
+	// Note that these cases are all still checked in the actual Undo method.
+	return !OwnerWidget->IsTextReadOnly()/* && UndoStates.Num() > 0 && !TextInputMethodContext->IsComposing()*/;
 }
 
 void FSlateEditableTextLayout::Undo()
