Always allow Redo in EditableText to match new Undo behavior

diff --git a/Engine/Source/Runtime/Slate/Private/Widgets/Text/SlateEditableTextLayout.cpp b/Engine/Source/Runtime/Slate/Private/Widgets/Text/SlateEditableTextLayout.cpp
--- a/Engine/Source/Runtime/Slate/Private/Widgets/Text/SlateEditableTextLayout.cpp
+++ b/Engine/Source/Runtime/Slate/Private/Widgets/Text/SlateEditableTextLayout.cpp
@@ -3069,7 +3069,7 @@ bool FSlateEditableTextLayout::CanExecuteUndo() const
 	// as undo-ing in a search box triggering an undo, and various issues related to undo-ing property changes
 	// while focused in to the property widget being undone.
 	// Note that these cases are all still checked in the actual Undo method.
-	return !OwnerWidget->IsTextReadOnly()/* && UndoStates.Num() > 0 && !TextInputMethodContext->IsComposing()*/;
+	return !OwnerWidget->IsTextReadOnly()/* && UndoStates.Num() > 0*/ && !TextInputMethodContext->IsComposing();
 }
 
 void FSlateEditableTextLayout::Undo()
@@ -3124,7 +3124,8 @@ void FSlateEditableTextLayout::Undo()
 
 bool FSlateEditableTextLayout::CanExecuteRedo() const
 {
-	return !OwnerWidget->IsTextReadOnly() && CurrentUndoLevel != INDEX_NONE && !TextInputMethodContext->IsComposing();
+	// See comment in CanExecuteUndo
+	return !OwnerWidget->IsTextReadOnly() /*&& CurrentUndoLevel != INDEX_NONE*/ && !TextInputMethodContext->IsComposing();
 }
 
 void FSlateEditableTextLayout::Redo()
