[NO NON-PII COMMIT MESSAGE]

diff --git a/Engine/Source/Runtime/Slate/Private/Framework/Application/AnalogCursor.cpp b/Engine/Source/Runtime/Slate/Private/Framework/Application/AnalogCursor.cpp
--- a/Engine/Source/Runtime/Slate/Private/Framework/Application/AnalogCursor.cpp
+++ b/Engine/Source/Runtime/Slate/Private/Framework/Application/AnalogCursor.cpp
@@ -277,6 +277,21 @@ bool FAnalogCursor::IsRelevantInput(const FInputEvent& InputEvent) const
 	return GetOwnerUserIndex() == InputEvent.GetUserIndex();
 }
 
+bool FAnalogCursor::IsRelevantInput(const FKeyEvent& KeyEvent) const
+{
+	return IsRelevantInput(static_cast<FInputEvent>(KeyEvent));
+}
+
+bool FAnalogCursor::IsRelevantInput(const FAnalogInputEvent& AnalogInputEvent) const
+{
+	return IsRelevantInput(static_cast<FInputEvent>(AnalogInputEvent));
+}
+
+bool FAnalogCursor::IsRelevantInput(const FPointerEvent& MouseEvent) const
+{
+	return IsRelevantInput(static_cast<FInputEvent>(MouseEvent));
+}
+
 void FAnalogCursor::ClearAnalogValues()
 {
 	AnalogValues[static_cast<uint8>(EAnalogStick::Left)] = FVector2D::ZeroVector;
diff --git a/Engine/Source/Runtime/Slate/Private/Framework/Application/SlateApplication.cpp b/Engine/Source/Runtime/Slate/Private/Framework/Application/SlateApplication.cpp
--- a/Engine/Source/Runtime/Slate/Private/Framework/Application/SlateApplication.cpp
+++ b/Engine/Source/Runtime/Slate/Private/Framework/Application/SlateApplication.cpp
@@ -1465,12 +1465,7 @@ void FSlateApplication::TickApplication(ESlateTickType TickType, float DeltaTime
 	
 		const bool bIsUserIdle = (TimeSinceInput > SleepThreshold) && (TimeSinceMouseMove > SleepThreshold);
 		const bool bAnyActiveTimersPending = AnyActiveTimersArePending();
-		if (bAnyActiveTimersPending)
-		{
-			// Some UI might slide under the cursor. To a widget, this is as if the cursor moved over it.
-			ForEachUser([](FSlateUser& User) { User.QueueSyntheticCursorMove(); });
-		}
-
+		
 		// Generate any simulated gestures that we've detected.
 		ForEachUser([this] (FSlateUser& User) {
 			User.GetGestureDetector().GenerateGestures(*this, SimulateGestures);
@@ -1484,6 +1479,8 @@ void FSlateApplication::TickApplication(ESlateTickType TickType, float DeltaTime
 		bIsSlateAsleep = true;
 		if	(!AllowSlateToSleep.GetValueOnGameThread() || bAnyActiveTimersPending || !bIsUserIdle || bSynthesizedCursorMove || FApp::UseVRFocus())
 		{
+			ForEachUser([](FSlateUser& User) { User.QueueSyntheticCursorMove(); });
+
 			bIsSlateAsleep = false; // if we get here, then Slate is not sleeping
 
 			// Update any notifications - this needs to be done after windows have updated themselves 
@@ -3357,11 +3354,6 @@ bool FSlateApplication::IsWindowInDestroyQueue(TSharedRef<SWindow> Window) const
 	return WindowDestroyQueue.Contains(Window);
 }
 
-void FSlateApplication::QueueSynthesizedMouseMove(const FInputEvent& SourceEvent)
-{
-	GetOrCreateUser(SourceEvent)->QueueSyntheticCursorMove();
-}
-
 void FSlateApplication::SetUnhandledKeyDownEventHandler( const FOnKeyEvent& NewHandler )
 {
 	UnhandledKeyDownEventHandler = NewHandler;
@@ -4138,8 +4130,6 @@ bool FSlateApplication::ProcessKeyDownEvent( const FKeyEvent& InKeyEvent )
 	}
 #endif //WITH_EDITOR
 
-	QueueSynthesizedMouseMove(InKeyEvent);
-
 	// Analog cursor gets first chance at the input
 	if (InputPreProcessors.HandleKeyDownEvent(*this, InKeyEvent))
 	{
@@ -4251,8 +4241,6 @@ bool FSlateApplication::ProcessKeyUpEvent( const FKeyEvent& InKeyEvent )
 
 	TScopeCounter<int32> BeginInput(ProcessingInput);
 
-	QueueSynthesizedMouseMove(InKeyEvent);
-
 	// Analog cursor gets first chance at the input
 	if (InputPreProcessors.HandleKeyUpEvent(*this, InKeyEvent))
 	{
@@ -4306,8 +4294,6 @@ bool FSlateApplication::ProcessAnalogInputEvent(const FAnalogInputEvent& InAnalo
 
 	TScopeCounter<int32> BeginInput(ProcessingInput);
 
-	QueueSynthesizedMouseMove(InAnalogInputEvent);
-
 	FReply Reply = FReply::Unhandled();
 
 	// Analog cursor gets first chance at the input
@@ -4340,8 +4326,6 @@ bool FSlateApplication::ProcessAnalogInputEvent(const FAnalogInputEvent& InAnalo
 
 				return FReply::Unhandled();
 			});
-
-		QueueSynthesizedMouseMove(ModifiedEvent);
 	}
 
 	// If no one handled this, it was probably motion in the deadzone.  Don't treat it as activity.
@@ -4447,7 +4431,6 @@ bool FSlateApplication::ProcessMouseButtonDownEvent( const TSharedPtr< FGenericW
 	}
 #endif //WITH_EDITOR
 
-	QueueSynthesizedMouseMove(MouseEvent);
 	SetLastUserInteractionTime(this->GetCurrentTime());
 	LastUserInteractionTimeForThrottling = LastUserInteractionTime;
 	
@@ -5094,7 +5077,6 @@ bool FSlateApplication::ProcessMouseButtonDoubleClickEvent( const TSharedPtr< FG
 {
 	SCOPE_CYCLE_COUNTER(STAT_ProcessMouseButtonDoubleClick);
 
-	QueueSynthesizedMouseMove(InMouseEvent);
 	SetLastUserInteractionTime(this->GetCurrentTime());
 	LastUserInteractionTimeForThrottling = LastUserInteractionTime;
 
@@ -5188,7 +5170,6 @@ bool FSlateApplication::ProcessMouseButtonUpEvent( const FPointerEvent& MouseEve
 		FSlateThrottleManager::Get().LeaveResponsiveMode(MouseButtonDownResponsivnessThrottle);
 	}
 
-	QueueSynthesizedMouseMove(MouseEvent);
 	SetLastUserInteractionTime(this->GetCurrentTime());
 	LastUserInteractionTimeForThrottling = LastUserInteractionTime;
 
@@ -5241,8 +5222,6 @@ bool FSlateApplication::ProcessMouseWheelOrGestureEvent( const FPointerEvent& In
 {
 	SCOPE_CYCLE_COUNTER(STAT_ProcessMouseWheelGesture);
 
-	QueueSynthesizedMouseMove(InWheelEvent);
-
 	bool bShouldProcessEvent = false;
 
 	if ( InGestureEvent )
@@ -5415,8 +5394,7 @@ bool FSlateApplication::ProcessMouseMoveEvent( const FPointerEvent& MouseEvent,
 	if ( !bIsSynthetic )
 	{
 		QUICK_SCOPE_CYCLE_COUNTER(STAT_ProcessMouseMove_Tooltip);
-
-		QueueSynthesizedMouseMove(MouseEvent);
+		
 		GetOrCreateUser(MouseEvent)->UpdateTooltip(MenuStack, /*bCanSpawnNewTooltip =*/true);
 		
 		// Guard against synthesized mouse moves and only track user interaction if the cursor pos changed
@@ -5830,7 +5808,6 @@ bool FSlateApplication::OnMotionDetected(const FVector& Tilt, const FVector& Rot
 
 void FSlateApplication::ProcessMotionDetectedEvent( const FMotionEvent& MotionEvent )
 {
-	QueueSynthesizedMouseMove(MotionEvent);
 	SetLastUserInteractionTime(this->GetCurrentTime());
 	
 	if (!InputPreProcessors.HandleMotionDetectedEvent(*this, MotionEvent))
@@ -6039,10 +6016,6 @@ bool FSlateApplication::ProcessWindowActivatedEvent( const FWindowActivateEvent&
 			SetLastUserInteractionTime(this->GetCurrentTime());
 		}
 		
-		// Widgets that happen to be under the mouse need to update if activation changes
-		// This also serves as a force redraw which is needed when restoring a window that was previously inactive.
-		ForEachUser([](FSlateUser& User) { User.QueueSyntheticCursorMove(); });
-
 		// NOTE: The window is brought to front even when a modal window is active and this is not the modal window one of its children 
 		// The reason for this is so that the Slate window order is in sync with the OS window order when a modal window is open.  This is important so that when the modal window closes the proper window receives input from Slate.
 		// If you change this be sure to test windows are activated properly and receive input when they are opened when a modal dialog is open.
@@ -6184,14 +6157,7 @@ void FSlateApplication::ProcessApplicationActivationEvent(bool InAppActivated)
 		// Clear the pressed buttons when we deactivate the application, the button state can no longer be trusted.
 		PressedMouseButtons.Reset();
 	}
-	else
-	{
-		//Ensure that slate ticks/renders next frame
-		ForEachUser([](FSlateUser& User) { User.QueueSyntheticCursorMove(); });
-
-		//TODO Requery pushed button state?
-	}
-
+	
 	OnApplicationActivationStateChanged().Broadcast(InAppActivated);
 }
 
diff --git a/Engine/Source/Runtime/Slate/Public/Framework/Application/AnalogCursor.h b/Engine/Source/Runtime/Slate/Public/Framework/Application/AnalogCursor.h
--- a/Engine/Source/Runtime/Slate/Public/Framework/Application/AnalogCursor.h
+++ b/Engine/Source/Runtime/Slate/Public/Framework/Application/AnalogCursor.h
@@ -59,6 +59,9 @@ public:
 protected:
 
 	virtual bool IsRelevantInput(const FInputEvent& InputEvent) const;
+	virtual bool IsRelevantInput(const FKeyEvent& KeyEvent) const;
+	virtual bool IsRelevantInput(const FAnalogInputEvent& AnalogInputEvent) const;
+	virtual bool IsRelevantInput(const FPointerEvent& MouseEvent) const;
 
 	/** Getter */
 	FORCEINLINE const FVector2D& GetAnalogValues( EAnalogStick Stick = EAnalogStick::Left ) const
diff --git a/Engine/Source/Runtime/Slate/Public/Framework/Application/SlateApplication.h b/Engine/Source/Runtime/Slate/Public/Framework/Application/SlateApplication.h
--- a/Engine/Source/Runtime/Slate/Public/Framework/Application/SlateApplication.h
+++ b/Engine/Source/Runtime/Slate/Public/Framework/Application/SlateApplication.h
@@ -1577,9 +1577,6 @@ private:
 	bool SetUserFocus(FSlateUser& User, const FWidgetPath& InFocusPath, const EFocusCause InCause);
 
 private:
-	/** Signal that a synthesized mouse move will be required after this operation. */
-	void QueueSynthesizedMouseMove(const FInputEvent& SourceEvent);
-
 	/**
 	 * Will be invoked when the size of the geometry of the virtual
 	 * desktop changes (e.g. resolution change or monitors re-arranged)
