Slate: don't commit the same text on focus lost (UE-89605).

- This resutls in unnecessary material recompiles when clicking through material paramaters.

diff --git a/Engine/Source/Runtime/Slate/Private/Widgets/Text/SlateEditableTextLayout.cpp b/Engine/Source/Runtime/Slate/Private/Widgets/Text/SlateEditableTextLayout.cpp
--- a/Engine/Source/Runtime/Slate/Private/Widgets/Text/SlateEditableTextLayout.cpp
+++ b/Engine/Source/Runtime/Slate/Private/Widgets/Text/SlateEditableTextLayout.cpp
@@ -802,11 +802,15 @@ bool FSlateEditableTextLayout::HandleFocusLost(const FFocusEvent& InFocusEvent)
 
 	const FText EditedText = GetEditableText();
 
-	OwnerWidget->OnTextCommitted(EditedText, TextAction);
+	// do not commit if the edited text is the same as original
+	if (!EditedText.ToString().Equals(OriginalText.Text.ToString(), ESearchCase::CaseSensitive))
+	{
+		OwnerWidget->OnTextCommitted(EditedText, TextAction);
 
-	// Reload underlying value now it is committed  (commit may alter the value) 
-	// so it can be re-displayed in the edit box
-	LoadText();
+		// Reload underlying value now it is committed  (commit may alter the value) 
+		// so it can be re-displayed in the edit box
+		LoadText();
+	}
 
 	UpdateCursorHighlight();
 
