Menu animations are now disabled by default. They cause flickering of popups in high DPI unaware applications.

diff --git a/Engine/Source/Runtime/Slate/Private/Framework/Application/SlateApplication.cpp b/Engine/Source/Runtime/Slate/Private/Framework/Application/SlateApplication.cpp
--- a/Engine/Source/Runtime/Slate/Private/Framework/Application/SlateApplication.cpp
+++ b/Engine/Source/Runtime/Slate/Private/Framework/Application/SlateApplication.cpp
@@ -715,7 +715,7 @@ FSlateApplication::FSlateApplication()
 	, bIsFakingTouched( false )
 	, bTouchFallbackToMouse( true )
 	, bSoftwareCursorAvailable( false )	
-	, bMenuAnimationsEnabled( true )
+	, bMenuAnimationsEnabled( false )
 	, AppIcon( FCoreStyle::Get().GetBrush("DefaultAppIcon") )
 	, VirtualDesktopRect( 0,0,0,0 )
 	, NavigationConfig(MakeShared<FNavigationConfig>())
