Fix mac cis

#rnx

diff --git a/Engine/Source/Runtime/Slate/Private/Framework/MultiBox/Mac/MacMenu.cpp b/Engine/Source/Runtime/Slate/Private/Framework/MultiBox/Mac/MacMenu.cpp
--- a/Engine/Source/Runtime/Slate/Private/Framework/MultiBox/Mac/MacMenu.cpp
+++ b/Engine/Source/Runtime/Slate/Private/Framework/MultiBox/Mac/MacMenu.cpp
@@ -691,7 +691,7 @@ void FSlateMacMenu::UpdateMenu(FMacMenu* Menu)
 						[MenuItem setState:MenuItemState.State];
 					}
 				}
-				else if (MenuItemState.Type == EMultiBlockType::MenuSeparator)
+				else if (MenuItemState.Type == EMultiBlockType::Separator)
 				{
 					if (MenuItem && ![MenuItem isSeparatorItem])
 					{
