Slate - Updating the LogState debug calls on the WidgetSwitcher to log more readable information.

diff --git a/Engine/Source/Runtime/Slate/Private/Widgets/Layout/SWidgetSwitcher.cpp b/Engine/Source/Runtime/Slate/Private/Widgets/Layout/SWidgetSwitcher.cpp
--- a/Engine/Source/Runtime/Slate/Private/Widgets/Layout/SWidgetSwitcher.cpp
+++ b/Engine/Source/Runtime/Slate/Private/Widgets/Layout/SWidgetSwitcher.cpp
@@ -2,6 +2,8 @@
 
 #include "Widgets/Layout/SWidgetSwitcher.h"
 #include "Layout/LayoutUtils.h"
+#include "Types/ReflectionMetadata.h"
+
 #if WITH_ACCESSIBILITY
 #include "Framework/Application/SlateApplication.h"
 #include "Widgets/Accessibility/SlateAccessibleMessageHandler.h"
@@ -116,7 +118,9 @@ int32 SWidgetSwitcher::RemoveSlot( TSharedRef<SWidget> WidgetToRemove )
 
 void SWidgetSwitcher::SetActiveWidgetIndex( int32 Index )
 {
-	if (WidgetIndex.Get() != Index)
+	const int32 OldIndex = WidgetIndex.Get();
+
+	if (OldIndex != Index)
 	{
 		Invalidate(EInvalidateWidget::ChildOrder);
 
@@ -127,13 +131,21 @@ void SWidgetSwitcher::SetActiveWidgetIndex( int32 Index )
 		{
 			SWidget& OldActiveWidget = ActiveSlot->GetWidget().Get();
 			InvalidateChildRemovedFromTree(OldActiveWidget);
-		}
-
-		WidgetIndex = Index;
 
 #if WITH_SLATE_DEBUGGING
-		UE_LOG(LogSlate, Log, TEXT("WidgetSwitcher ('%s') active index changed to %d"), *GetTag().ToString(), Index);
+			UE_LOG(LogSlate, Verbose, TEXT("WidgetSwitcher ('%s') Active Slot Changed: %d(%s) -FROM- %d(%s)"), *FReflectionMetaData::GetWidgetDebugInfo(this), 
+				Index, *FReflectionMetaData::GetWidgetDebugInfo(GetWidget(Index).Get()),
+				OldIndex, *FReflectionMetaData::GetWidgetDebugInfo(OldActiveWidget));
+#endif
+		}
+		else
+		{
+#if WITH_SLATE_DEBUGGING
+			UE_LOG(LogSlate, Verbose, TEXT("WidgetSwitcher ('%s') Active Slot Changed: %d(%s)"), *FReflectionMetaData::GetWidgetDebugInfo(this), Index, *FReflectionMetaData::GetWidgetDebugInfo(GetWidget(Index).Get()));
 #endif
+		}
+
+		WidgetIndex = Index;
 	}
 }
 
