Slate - If the child of a Scalebox is collapsed, and the Scalebox is not, don't do any work calculating scale for it.

diff --git a/Engine/Source/Runtime/Slate/Private/Widgets/Layout/SScaleBox.cpp b/Engine/Source/Runtime/Slate/Private/Widgets/Layout/SScaleBox.cpp
--- a/Engine/Source/Runtime/Slate/Private/Widgets/Layout/SScaleBox.cpp
+++ b/Engine/Source/Runtime/Slate/Private/Widgets/Layout/SScaleBox.cpp
@@ -49,6 +49,13 @@ bool SScaleBox::CustomPrepass(float LayoutScaleMultiplier)
 {
 	SWidget& ChildSlotWidget = ChildSlot.GetWidget().Get();
 
+	// If the child is collapsed don't bother calculating any scale for it.
+	const EVisibility ChildVisibility = ChildSlotWidget.GetVisibility();
+	if (ChildVisibility == EVisibility::Collapsed)
+	{
+		return false;
+	}
+
 	const bool bNeedsNormalizingPrepassOrLocalGeometry = DoesScaleRequireNormalizingPrepassOrLocalGeometry();
 
 	// If we need a normalizing prepass, or we've yet to give the child a chance to generate a desired
@@ -149,24 +156,30 @@ float SScaleBox::ComputeContentScale(const FGeometry& PaintGeometry) const
 	{
 		switch (CurrentStretch)
 		{
-		case EStretch::ScaleToFit:
-		{
-			//FVector2D LocalSnappedArea = PaintGeometry.LocalToRoundedLocal(PaintGeometry.GetLocalSize()) - PaintGeometry.LocalToRoundedLocal(FVector2D(0, 0));
-			//FinalScale = FMath::Min(LocalSnappedArea.X / ChildDesiredSize.X, LocalSnappedArea.Y / ChildDesiredSize.Y);
-			FinalScale = FMath::Min(PaintGeometry.GetLocalSize().X / ChildDesiredSize.X, PaintGeometry.GetLocalSize().Y / ChildDesiredSize.Y);
-			break;
-		}
-		case EStretch::ScaleToFitX:
-			FinalScale = PaintGeometry.GetLocalSize().X / ChildDesiredSize.X;
-			break;
-		case EStretch::ScaleToFitY:
-			FinalScale = PaintGeometry.GetLocalSize().Y / ChildDesiredSize.Y;
-			break;
-		case EStretch::Fill:
-			break;
-		case EStretch::ScaleToFill:
-			FinalScale = FMath::Max(PaintGeometry.GetLocalSize().X / ChildDesiredSize.X, PaintGeometry.GetLocalSize().Y / ChildDesiredSize.Y);
-			break;
+			case EStretch::ScaleToFit:
+			{
+				FinalScale = FMath::Min(PaintGeometry.GetLocalSize().X / ChildDesiredSize.X, PaintGeometry.GetLocalSize().Y / ChildDesiredSize.Y);
+				break;
+			}
+			case EStretch::ScaleToFitX:
+			{
+				FinalScale = PaintGeometry.GetLocalSize().X / ChildDesiredSize.X;
+				break;
+			}
+			case EStretch::ScaleToFitY:
+			{
+				FinalScale = PaintGeometry.GetLocalSize().Y / ChildDesiredSize.Y;
+				break;
+			}
+			case EStretch::Fill:
+			{
+				break;
+			}
+			case EStretch::ScaleToFill:
+			{
+				FinalScale = FMath::Max(PaintGeometry.GetLocalSize().X / ChildDesiredSize.X, PaintGeometry.GetLocalSize().Y / ChildDesiredSize.Y);
+				break;
+			}
 		}
 
 		switch (CurrentStretchDirection)
