Slate - When using demo.ClientRecordAsyncEndOfFrame 1 to run replay recording concurrent with slate tick, there tend to still be wrting performed on replicated objects due to player input.  Because of this, Slate Tick is now seperable dividing up the initial call to just Platform and Input ticking, which performs message pumping as well as some additional added slate input tasks - tasks which inhertently tend to have networking side effects if the user is performing some action like clicking a button, which might trigger something networking related.  Then we run replay recording if that option is enabled alongside slate tick where we only advance time, and then tick and paint widgets, which is assumed to not perform any replication altering behavior, this fixes several networking related crashes.

diff --git a/Engine/Source/Runtime/Launch/Private/LaunchEngineLoop.cpp b/Engine/Source/Runtime/Launch/Private/LaunchEngineLoop.cpp
--- a/Engine/Source/Runtime/Launch/Private/LaunchEngineLoop.cpp
+++ b/Engine/Source/Runtime/Launch/Private/LaunchEngineLoop.cpp
@@ -4940,6 +4940,21 @@ void FEngineLoop::Tick()
 		}
 #endif
 
+		// Tick the platform and input portion of Slate application, we need to do this before we run things
+		// concurrent with networking.
+		if (FSlateApplication::IsInitialized() && !bIdleMode)
+		{
+			{
+				QUICK_SCOPE_CYCLE_COUNTER(STAT_FEngineLoop_ProcessPlayerControllersSlateOperations);
+				check(!IsRunningDedicatedServer());
+
+				// Process slate operations accumulated in the world ticks.
+				ProcessLocalPlayerSlateOperations();
+			}
+
+			FSlateApplication::Get().Tick(ESlateTickType::PlatformAndInput);
+		}
+
 #if WITH_ENGINE
 		// process concurrent Slate tasks
 		FGraphEventRef ConcurrentTask;
@@ -4980,18 +4995,11 @@ void FEngineLoop::Tick()
 		}
 #endif
 
-		// tick Slate application
+		// Tick(Advance) Time for the application and then tick and paint slate application widgets.
+		// We split separate this action from the one above to permit running network replication concurrent with slate widget ticking and painting.
 		if (FSlateApplication::IsInitialized() && !bIdleMode)
 		{
-			{
-				QUICK_SCOPE_CYCLE_COUNTER(STAT_FEngineLoop_ProcessPlayerControllersSlateOperations);
-				check(!IsRunningDedicatedServer());
-
-				// Process slate operations accumulated in the world ticks.
-				ProcessLocalPlayerSlateOperations();
-			}
-
-			FSlateApplication::Get().Tick();
+			FSlateApplication::Get().Tick(ESlateTickType::TimeAndWidgets);
 		}
 
 #if WITH_ENGINE
diff --git a/Engine/Source/Runtime/MoviePlayer/Private/DefaultGameMoviePlayer.cpp b/Engine/Source/Runtime/MoviePlayer/Private/DefaultGameMoviePlayer.cpp
--- a/Engine/Source/Runtime/MoviePlayer/Private/DefaultGameMoviePlayer.cpp
+++ b/Engine/Source/Runtime/MoviePlayer/Private/DefaultGameMoviePlayer.cpp
@@ -843,7 +843,7 @@ void FMoviePlayerWidgetRenderer::DrawWindow(float DeltaTime)
 
 	FVector2D DrawSize = VirtualRenderWindow->GetClientSizeInScreen();
 
-	FSlateApplication::Get().Tick(ESlateTickType::TimeOnly);
+	FSlateApplication::Get().Tick(ESlateTickType::Time);
 
 	const float Scale = 1.0f;
 	FGeometry WindowGeometry = FGeometry::MakeRoot(DrawSize, FSlateLayoutTransform(Scale));
diff --git a/Engine/Source/Runtime/PreLoadScreen/Private/PreLoadSlateThreading.cpp b/Engine/Source/Runtime/PreLoadScreen/Private/PreLoadSlateThreading.cpp
--- a/Engine/Source/Runtime/PreLoadScreen/Private/PreLoadSlateThreading.cpp
+++ b/Engine/Source/Runtime/PreLoadScreen/Private/PreLoadSlateThreading.cpp
@@ -65,9 +65,9 @@ void FPreLoadSlateWidgetRenderer::DrawWindow(float DeltaTime)
         return;
     }
 
-    FVector2D DrawSize = VirtualRenderWindow->GetClientSizeInScreen();
+    const FVector2D DrawSize = VirtualRenderWindow->GetClientSizeInScreen();
 
-    FSlateApplication::Get().Tick(ESlateTickType::TimeOnly);
+    FSlateApplication::Get().Tick(ESlateTickType::Time);
 
     const float Scale = 1.0f;
     FGeometry WindowGeometry = FGeometry::MakeRoot(DrawSize, FSlateLayoutTransform(Scale));
diff --git a/Engine/Source/Runtime/Slate/Private/Framework/Application/SlateApplication.cpp b/Engine/Source/Runtime/Slate/Private/Framework/Application/SlateApplication.cpp
--- a/Engine/Source/Runtime/Slate/Private/Framework/Application/SlateApplication.cpp
+++ b/Engine/Source/Runtime/Slate/Private/Framework/Application/SlateApplication.cpp
@@ -427,6 +427,7 @@ struct FEventRouter::Translate<FPointerEvent>
 };
 
 DECLARE_CYCLE_STAT( TEXT("Message Tick Time"), STAT_SlateMessageTick, STATGROUP_Slate );
+DECLARE_CYCLE_STAT( TEXT("Slate App Input"), STAT_SlateApplicationInput, STATGROUP_Slate );
 DECLARE_CYCLE_STAT( TEXT("Total Slate Tick Time"), STAT_SlateTickTime, STATGROUP_Slate );
 DECLARE_CYCLE_STAT( TEXT("Draw Window And Children Time"), STAT_SlateDrawWindowTime, STATGROUP_Slate );
 DECLARE_CYCLE_STAT( TEXT("TickRegisteredWidgets"), STAT_SlateTickRegisteredWidgets, STATGROUP_Slate );
@@ -1327,6 +1328,24 @@ void FSlateApplication::FinishedInputThisFrame()
 	ForEachUser([] (FSlateUser& User) { User.FinishFrame(); });
 }
 
+static const TCHAR* LexToString(ESlateTickType TickType)
+{
+	switch (TickType)
+	{
+	case ESlateTickType::Time:
+		return TEXT("Time");
+	case ESlateTickType::PlatformAndInput:
+		return TEXT("Platform and Input");
+	case ESlateTickType::Widgets:
+		return TEXT("Widgets");
+	case ESlateTickType::TimeAndWidgets:
+		return TEXT("Time and Widgets");
+	case ESlateTickType::All:
+	default:
+		return TEXT("All");
+	}
+}
+
 void FSlateApplication::Tick(ESlateTickType TickType)
 {
 	LLM_SCOPE(ELLMTag::UI);
@@ -1335,11 +1354,11 @@ void FSlateApplication::Tick(ESlateTickType TickType)
 	CSV_SCOPED_TIMING_STAT_EXCLUSIVE(UI);
 
 	// It is not valid to tick Slate on any other thread but the game thread unless we are only updating time
-	check(IsInGameThread() || TickType == ESlateTickType::TimeOnly);
+	check(IsInGameThread() || TickType == ESlateTickType::Time);
 
 	FScopeLock SlateTickAccess(&SlateTickCriticalSection);
 
-	SCOPED_NAMED_EVENT_TEXT("Slate::Tick", FColor::Magenta);
+	SCOPED_NAMED_EVENT_F(TEXT("Slate::Tick (%s)"), FColor::Magenta, LexToString(TickType));
 	CSV_SCOPED_TIMING_STAT(Slate, Tick);
 
 	{
@@ -1347,18 +1366,33 @@ void FSlateApplication::Tick(ESlateTickType TickType)
 
 		const float DeltaTime = GetDeltaTime();
 
-		if (TickType == ESlateTickType::All)
+		if (EnumHasAnyFlags(TickType, ESlateTickType::PlatformAndInput))
 		{
-#if WITH_ACCESSIBILITY
-			// We ensure to only call this in TickType::All to avoid the movie thread also calling this unnecessarily 
-			GetAccessibleMessageHandler()->ProcessAccessibleTasks();
-#endif
 			TickPlatform(DeltaTime);
 		}
-		TickApplication(TickType, DeltaTime);
-#if WITH_ACCESSIBILITY
-		if (TickType == ESlateTickType::All)
+
+		if (EnumHasAnyFlags(TickType, ESlateTickType::Time))
 		{
+			TickTime();
+		}
+
+		if (EnumHasAnyFlags(TickType, ESlateTickType::Widgets))
+		{
+			TickAndDrawWidgets(DeltaTime);
+		}
+	}
+}
+
+void FSlateApplication::TickTime()
+{
+	LastTickTime = CurrentTime;
+	CurrentTime = FPlatformTime::Seconds();
+
+	// Handle large quantums
+	const double MaxQuantumBeforeClamp = 1.0 / 8.0;		// 8 FPS
+	if (GetDeltaTime() > MaxQuantumBeforeClamp)
+	{
+		LastTickTime = CurrentTime - MaxQuantumBeforeClamp;
 			// we call this again to improve the responsiveness of accessibility navigation and announcements 
 			GetAccessibleMessageHandler()->ProcessAccessibleTasks();
 		}
@@ -1370,6 +1404,13 @@ void FSlateApplication::TickPlatform(float DeltaTime)
 {
 	SCOPED_NAMED_EVENT_TEXT("Slate::TickPlatform", FColor::Magenta);
 
+#if WITH_ACCESSIBILITY
+	{
+		// We ensure to only call this in TickType::All to avoid the movie thread also calling this unnecessarily 
+		GetAccessibleMessageHandler()->ProcessAccessibleTasks();
+	}
+#endif
+
 	{
 		SCOPE_CYCLE_COUNTER(STAT_SlateMessageTick);
 
@@ -1388,9 +1429,29 @@ void FSlateApplication::TickPlatform(float DeltaTime)
 		PlatformApplication->Tick(DeltaTime);
 		PlatformApplication->ProcessDeferredEvents(DeltaTime);
 	}
+
+	{
+		SCOPE_CYCLE_COUNTER(STAT_SlateApplicationInput);
+
+		ForEachUser([this](FSlateUser& User) {
+			User.UpdateCursor();
+			User.UpdateTooltip(MenuStack, false);
+		});
+
+		bool bSynthesizedCursorMoveThisFrame = false;
+		ForEachUser([&bSynthesizedCursorMoveThisFrame](FSlateUser& User) {
+			bSynthesizedCursorMoveThisFrame |= User.SynthesizeCursorMoveIfNeeded();
+		});
+		bSynthesizedCursorMove = bSynthesizedCursorMoveThisFrame;
+
+		// Generate any simulated gestures that we've detected.
+		ForEachUser([this](FSlateUser& User) {
+			User.GetGestureDetector().GenerateGestures(*this, SimulateGestures);
+		});
+	}
 }
 
-void FSlateApplication::TickApplication(ESlateTickType TickType, float DeltaTime)
+void FSlateApplication::TickAndDrawWidgets(float DeltaTime)
 {
 	if (Renderer.IsValid())
 	{
@@ -1400,29 +1461,14 @@ void FSlateApplication::TickApplication(ESlateTickType TickType, float DeltaTime
 		Renderer->ReleaseAccessedResources(/* Flush State */ false);
 	}
 
-	if (TickType == ESlateTickType::All)
 	{
-		FPlatformMisc::BeginNamedEvent(FColor::Purple, TEXT("Slate Pre-Paint Processing"));
-
-		{
-			SCOPE_CYCLE_COUNTER(STAT_SlatePreTickEvent);
-			PreTickEvent.Broadcast(DeltaTime);
-		}
-
-		ForEachUser([this](FSlateUser& User) {
-			User.UpdateCursor();
-			User.UpdateTooltip(MenuStack, false);
-		});
+		SCOPE_CYCLE_COUNTER(STAT_SlatePreTickEvent);
+		PreTickEvent.Broadcast(DeltaTime);
 	}
 
-	// Advance time
-	LastTickTime = CurrentTime;
-	CurrentTime = FPlatformTime::Seconds();
-
 	// Update average time between ticks.  This is used to monitor how responsive the application "feels".
 	// Note that we calculate this before we apply the max quantum clamping below, because we want to store
 	// the actual frame rate, even if it is very low.
-	if (TickType == ESlateTickType::All)
 	{
 		// Scalar percent of new delta time that contributes to running average.  Use a lower value to add more smoothing
 		// to the average frame rate.  A value of 1.0 will disable smoothing.
@@ -1442,20 +1488,7 @@ void FSlateApplication::TickApplication(ESlateTickType TickType, float DeltaTime
 		}
 	}
 
-	// Handle large quantums
-	const double MaxQuantumBeforeClamp = 1.0 / 8.0;		// 8 FPS
-	if( GetDeltaTime() > MaxQuantumBeforeClamp )
-	{
-		LastTickTime = CurrentTime - MaxQuantumBeforeClamp;
-	}
-
-	if (TickType == ESlateTickType::All)
-	{
-		bool bSynthesizedCursorMove = false;
-		ForEachUser([&bSynthesizedCursorMove](FSlateUser& User) {
-			bSynthesizedCursorMove |= User.SynthesizeCursorMoveIfNeeded();
-		});
-		
+	{	
 		// Update auto-throttling based on elapsed time since user interaction
 		ThrottleApplicationBasedOnMouseMovement();
 
@@ -1467,13 +1500,7 @@ void FSlateApplication::TickApplication(ESlateTickType TickType, float DeltaTime
 	
 		const bool bIsUserIdle = (TimeSinceInput > SleepThreshold) && (TimeSinceMouseMove > SleepThreshold);
 		const bool bAnyActiveTimersPending = AnyActiveTimersArePending();
-		
-		// Generate any simulated gestures that we've detected.
-		ForEachUser([this] (FSlateUser& User) {
-			User.GetGestureDetector().GenerateGestures(*this, SimulateGestures);
-		});
 
-		FPlatformMisc::EndNamedEvent();
 		// skip tick/draw if we are idle and there are no active timers registered that we need to drive slate for.
 		// This effectively means the slate application is totally idle and we don't need to update the UI.
 		// This relies on Widgets properly registering for Active timer when they need something to happen even
@@ -1501,11 +1528,20 @@ void FSlateApplication::TickApplication(ESlateTickType TickType, float DeltaTime
 			AccessibleMessageHandler->Tick();
 #endif
 		}
+	}
 
+	{
+		// SCOPE_CYCLE_COUNTER(STAT_SlatePostTickEvent);
 		PostTickEvent.Broadcast(DeltaTime);
 	}
-}
 
+#if WITH_ACCESSIBILITY
+	{
+		// we call this again to improve the responsiveness of accessibility navigation and announcements 
+		GetAccessibleMessageHandler()->ProcessAccessibleTasks();
+	}
+#endif
+}
 
 void FSlateApplication::PumpMessages()
 {
@@ -1904,8 +1940,11 @@ void FSlateApplication::AddModalWindow( TSharedRef<SWindow> InSlateWindow, const
 				// Slate's Tick does not issue Begin/EndFrame so we'll do it ourselves.
 				Renderer->BeginFrame();
 
+				// Advance time for the application
+				TickTime();
+
 				// Tick and render Slate
-				TickApplication(ESlateTickType::All, DeltaTime);
+				TickAndDrawWidgets(DeltaTime);
 				
 				Renderer->EndFrame();
 			}
diff --git a/Engine/Source/Runtime/Slate/Public/Framework/Application/SlateApplication.h b/Engine/Source/Runtime/Slate/Public/Framework/Application/SlateApplication.h
--- a/Engine/Source/Runtime/Slate/Public/Framework/Application/SlateApplication.h
+++ b/Engine/Source/Runtime/Slate/Public/Framework/Application/SlateApplication.h
@@ -133,12 +133,23 @@ public:
 enum class ESlateTickType : uint8
 {
 	/** Tick time only */
-	TimeOnly,
+	Time = 1 << 0,
+
+	/** Only process input for the platform, and additional input tasks by Slate. */
+	PlatformAndInput = 1 << 1,
+
+	/** Only Tick and Paint Widgets */
+	Widgets = 1 << 2,
+
+	/** Time and Widgets */
+	TimeAndWidgets = Time | Widgets,
 
 	/** Update time, tick and paint widgets, and process input */
-	All,
+	All = Time | PlatformAndInput | Widgets,
 };
 
+ENUM_CLASS_FLAGS(ESlateTickType);
+
 class SLATE_API FSlateApplication
 	: public FSlateApplicationBase
 	, public FGenericApplicationMessageHandler
@@ -957,6 +968,9 @@ protected:
 
 	TSharedRef<FNavigationConfig> GetRelevantNavConfig(int32 UserIndex) const;
 
+	/** Advances time for the application. */
+	void TickTime();
+
 	/**
 	 * Pumps and ticks the platform.
 	 */
@@ -965,7 +979,7 @@ protected:
 	/**
 	 * Ticks and paints the actual Slate portion of the application.
 	 */
-	void TickApplication(ESlateTickType TickType, float DeltaTime);
+	void TickAndDrawWidgets(float DeltaTime);
 
 	/** Draws Slate windows. Should only be called by the application's main loop or renderer. */
 	void DrawWindows();
@@ -1822,6 +1836,8 @@ private:
 	/** Are we currently processing input in slate?  If so this value will be greater than 0. */
 	int32 ProcessingInput;
 
+	/** Did we synthesize cursor input this frame? */
+	bool bSynthesizedCursorMove = false;
 		
 	/**
 	 * A helper class to wrap the list of input pre-processors. 
