Slate - Adding a note to Slate Tick.

diff --git a/Engine/Source/Runtime/Slate/Private/Framework/Application/SlateApplication.cpp b/Engine/Source/Runtime/Slate/Private/Framework/Application/SlateApplication.cpp
--- a/Engine/Source/Runtime/Slate/Private/Framework/Application/SlateApplication.cpp
+++ b/Engine/Source/Runtime/Slate/Private/Framework/Application/SlateApplication.cpp
@@ -1366,6 +1366,12 @@ void FSlateApplication::Tick(ESlateTickType TickType)
 
 		const float DeltaTime = GetDeltaTime();
 
+		// IMPORTANT
+		// Do not add code to these different if-statements, if you need to add additional logic to
+		// ticking the platform, do it inside of TickPlatform, for example.  These functions are sometimes
+		// called directly inside of Slate Application, so unless they're embedded in those calls, they wont
+		// get run.
+
 		if (EnumHasAnyFlags(TickType, ESlateTickType::PlatformAndInput))
 		{
 			TickPlatform(DeltaTime);
