Slate - Fixing bad merge.

diff --git a/Engine/Source/Runtime/Slate/Private/Framework/Application/SlateApplication.cpp b/Engine/Source/Runtime/Slate/Private/Framework/Application/SlateApplication.cpp
--- a/Engine/Source/Runtime/Slate/Private/Framework/Application/SlateApplication.cpp
+++ b/Engine/Source/Runtime/Slate/Private/Framework/Application/SlateApplication.cpp
@@ -1399,10 +1399,6 @@ void FSlateApplication::TickTime()
 	if (GetDeltaTime() > MaxQuantumBeforeClamp)
 	{
 		LastTickTime = CurrentTime - MaxQuantumBeforeClamp;
-			// we call this again to improve the responsiveness of accessibility navigation and announcements 
-			GetAccessibleMessageHandler()->ProcessAccessibleTasks();
-		}
-#endif
 	}
 }
 
