[Engine] - Add the possibility to arrange the slots top-to-bottom in the DynamicEntryBox and the WrapBox widgets instead of only supporting left-to-right.
#RNX

diff --git a/Engine/Plugins/2D/Paper2D/Source/Paper2DEditor/Private/GroupedSprites/GroupedSpriteDetailsCustomization.cpp b/Engine/Plugins/2D/Paper2D/Source/Paper2DEditor/Private/GroupedSprites/GroupedSpriteDetailsCustomization.cpp
--- a/Engine/Plugins/2D/Paper2D/Source/Paper2DEditor/Private/GroupedSprites/GroupedSpriteDetailsCustomization.cpp
+++ b/Engine/Plugins/2D/Paper2D/Source/Paper2DEditor/Private/GroupedSprites/GroupedSpriteDetailsCustomization.cpp
@@ -41,7 +41,7 @@ void FGroupedSpriteComponentDetailsCustomization::CustomizeDetails(IDetailLayout
 	DetailBuilder.GetObjectsBeingCustomized(/*out*/ ObjectsBeingCustomized);
 
 
-	TSharedRef<SWrapBox> ButtonBox = SNew(SWrapBox).UseAllottedWidth(true);
+	TSharedRef<SWrapBox> ButtonBox = SNew(SWrapBox).UseAllottedSize(true);
 
 	const float MinButtonSize = 100.0f;
 	const FMargin ButtonPadding(0.0f, 2.0f, 2.0f, 0.0f);
diff --git a/Engine/Plugins/2D/Paper2D/Source/Paper2DEditor/Private/TileMapEditing/PaperTileMapDetailsCustomization.cpp b/Engine/Plugins/2D/Paper2D/Source/Paper2DEditor/Private/TileMapEditing/PaperTileMapDetailsCustomization.cpp
--- a/Engine/Plugins/2D/Paper2D/Source/Paper2DEditor/Private/TileMapEditing/PaperTileMapDetailsCustomization.cpp
+++ b/Engine/Plugins/2D/Paper2D/Source/Paper2DEditor/Private/TileMapEditing/PaperTileMapDetailsCustomization.cpp
@@ -105,7 +105,7 @@ void FPaperTileMapDetailsCustomization::CustomizeDetails(IDetailLayoutBuilder& D
 
 	TAttribute<EVisibility> InternalInstanceVis = TAttribute<EVisibility>::Create(TAttribute<EVisibility>::FGetter::CreateSP(this, &FPaperTileMapDetailsCustomization::GetVisibilityForInstancedOnlyProperties));
 
-	TSharedRef<SWrapBox> ButtonBox = SNew(SWrapBox).UseAllottedWidth(true);
+	TSharedRef<SWrapBox> ButtonBox = SNew(SWrapBox).UseAllottedSize(true);
 
 	const float MinButtonSize = 120.0f;
 	const FMargin ButtonPadding(0.0f, 2.0f, 2.0f, 0.0f);
diff --git a/Engine/Plugins/Cameras/CameraShakePreviewer/Source/CameraShakePreviewer/Private/SCameraShakePreviewer.cpp b/Engine/Plugins/Cameras/CameraShakePreviewer/Source/CameraShakePreviewer/Private/SCameraShakePreviewer.cpp
--- a/Engine/Plugins/Cameras/CameraShakePreviewer/Source/CameraShakePreviewer/Private/SCameraShakePreviewer.cpp
+++ b/Engine/Plugins/Cameras/CameraShakePreviewer/Source/CameraShakePreviewer/Private/SCameraShakePreviewer.cpp
@@ -197,7 +197,7 @@ void SCameraShakePreviewer::Construct(const FArguments& InArgs)
 			.Padding(1.0f)
 			[
 				SNew(SWrapBox)
-				.UseAllottedWidth(true)
+				.UseAllottedSize(true)
 				+SWrapBox::Slot()
 				.Padding(FMargin(2.0f))
 				[
diff --git a/Engine/Plugins/FX/Niagara/Source/NiagaraEditor/Private/SNewSystemDialog.cpp b/Engine/Plugins/FX/Niagara/Source/NiagaraEditor/Private/SNewSystemDialog.cpp
--- a/Engine/Plugins/FX/Niagara/Source/NiagaraEditor/Private/SNewSystemDialog.cpp
+++ b/Engine/Plugins/FX/Niagara/Source/NiagaraEditor/Private/SNewSystemDialog.cpp
@@ -83,7 +83,7 @@ void SNewSystemDialog::Construct(const FArguments& InArgs)
 				.AutoHeight()
 				[
 					SAssignNew(SelectedEmitterBox, SWrapBox)
-					.UseAllottedWidth(true)
+					.UseAllottedSize(true)
 				]),
 			SNiagaraNewAssetDialog::FNiagaraNewAssetDialogOption(
 				LOCTEXT("CreateFromTemplateLabel", "New system from a template"),
diff --git a/Engine/Plugins/FX/Niagara/Source/NiagaraEditorWidgets/Private/SNiagaraStack.cpp b/Engine/Plugins/FX/Niagara/Source/NiagaraEditorWidgets/Private/SNiagaraStack.cpp
--- a/Engine/Plugins/FX/Niagara/Source/NiagaraEditorWidgets/Private/SNiagaraStack.cpp
+++ b/Engine/Plugins/FX/Niagara/Source/NiagaraEditorWidgets/Private/SNiagaraStack.cpp
@@ -108,7 +108,7 @@ public:
 					// Name and Source Emitter Name
 					SNew(SWrapBox)
 					.Clipping(EWidgetClipping::ClipToBoundsAlways) 
-					.UseAllottedWidth(true)
+					.UseAllottedSize(true)
 					+ SWrapBox::Slot()
 					[
 				
diff --git a/Engine/Plugins/FX/Niagara/Source/NiagaraEditorWidgets/Private/SNiagaraStackEntryWidget.cpp b/Engine/Plugins/FX/Niagara/Source/NiagaraEditorWidgets/Private/SNiagaraStackEntryWidget.cpp
--- a/Engine/Plugins/FX/Niagara/Source/NiagaraEditorWidgets/Private/SNiagaraStackEntryWidget.cpp
+++ b/Engine/Plugins/FX/Niagara/Source/NiagaraEditorWidgets/Private/SNiagaraStackEntryWidget.cpp
@@ -107,7 +107,7 @@ TSharedRef<SWidget> SNiagaraStackDisplayName::ConstructChildren()
 	if(NameWidgets.Num() > 1)
 	{
 		TSharedRef<SWrapBox> NamesWrapBox = SNew(SWrapBox)
-			.UseAllottedWidth(true);
+			.UseAllottedSize(true);
 		for (TSharedRef<SWidget> NameWidget : NameWidgets)
 		{
 			NamesWrapBox->AddSlot()
diff --git a/Engine/Source/Developer/LogVisualizer/Private/SVisualLoggerFilters.cpp b/Engine/Source/Developer/LogVisualizer/Private/SVisualLoggerFilters.cpp
--- a/Engine/Source/Developer/LogVisualizer/Private/SVisualLoggerFilters.cpp
+++ b/Engine/Source/Developer/LogVisualizer/Private/SVisualLoggerFilters.cpp
@@ -28,7 +28,7 @@ void SVisualLoggerFilters::Construct(const FArguments& InArgs, const TSharedRef<
 	ChildSlot
 		[
 			SAssignNew(FilterBox, SWrapBox)
-			.UseAllottedWidth(true)
+			.UseAllottedSize(true)
 		];
 
 	GraphsFilterCombo =
diff --git a/Engine/Source/Editor/ContentBrowser/Private/AssetViewWidgets.cpp b/Engine/Source/Editor/ContentBrowser/Private/AssetViewWidgets.cpp
--- a/Engine/Source/Editor/ContentBrowser/Private/AssetViewWidgets.cpp
+++ b/Engine/Source/Editor/ContentBrowser/Private/AssetViewWidgets.cpp
@@ -943,7 +943,7 @@ TSharedRef<SWidget> SAssetViewItem::CreateToolTipWidget() const
 				if (CollectionsContainingObject.Num() > 0)
 				{
 					TSharedRef<SWrapBox> CollectionPipsWrapBox = SNew(SWrapBox)
-						.PreferredWidth(700.0f);
+						.PreferredSize(700.0f);
 
 					for (const FCollectionNameType& CollectionContainingObject : CollectionsContainingObject)
 					{
diff --git a/Engine/Source/Editor/ContentBrowser/Private/SContentBrowser.cpp b/Engine/Source/Editor/ContentBrowser/Private/SContentBrowser.cpp
--- a/Engine/Source/Editor/ContentBrowser/Private/SContentBrowser.cpp
+++ b/Engine/Source/Editor/ContentBrowser/Private/SContentBrowser.cpp
@@ -154,7 +154,7 @@ void SContentBrowser::Construct( const FArguments& InArgs, const FName& InInstan
 		.Padding( 0, 0, 0, 0 )
 		[
 			SNew( SWrapBox )
-			.UseAllottedWidth( true )
+			.UseAllottedSize( true )
 			.InnerSlotPadding( FVector2D( 5, 2 ) )
 
 			+ SWrapBox::Slot()
diff --git a/Engine/Source/Editor/ContentBrowser/Private/SFilterList.cpp b/Engine/Source/Editor/ContentBrowser/Private/SFilterList.cpp
--- a/Engine/Source/Editor/ContentBrowser/Private/SFilterList.cpp
+++ b/Engine/Source/Editor/ContentBrowser/Private/SFilterList.cpp
@@ -559,7 +559,7 @@ void SFilterList::Construct( const FArguments& InArgs )
 	}
 
 	FilterBox = SNew(SWrapBox)
-		.UseAllottedWidth(true);
+		.UseAllottedSize(true);
 
 	ChildSlot
 	[
diff --git a/Engine/Source/Editor/DetailCustomizations/Private/DialogueStructsCustomizations.cpp b/Engine/Source/Editor/DetailCustomizations/Private/DialogueStructsCustomizations.cpp
--- a/Engine/Source/Editor/DetailCustomizations/Private/DialogueStructsCustomizations.cpp
+++ b/Engine/Source/Editor/DetailCustomizations/Private/DialogueStructsCustomizations.cpp
@@ -487,7 +487,7 @@ TSharedRef<SWidget> STargetSetDropDown::MakeComboButtonItemWidget(TSharedPtr<FTa
 	{
 		const TSharedRef<SWrapBox> WrapBox =
 			SNew( SWrapBox )
-			.PreferredWidth( this, &STargetSetDropDown::GetPreferredWidthForWrapping );
+			.PreferredSize( this, &STargetSetDropDown::GetPreferredWidthForWrapping );
 
 		for(int32 i = 0; i < TargetSet->Num(); ++i)
 		{
diff --git a/Engine/Source/Editor/DetailCustomizations/Private/DialogueWaveWidgets.cpp b/Engine/Source/Editor/DetailCustomizations/Private/DialogueWaveWidgets.cpp
--- a/Engine/Source/Editor/DetailCustomizations/Private/DialogueWaveWidgets.cpp
+++ b/Engine/Source/Editor/DetailCustomizations/Private/DialogueWaveWidgets.cpp
@@ -655,7 +655,7 @@ void STargetsSummaryWidget::GenerateContent()
 		{
 			const TSharedRef<SWrapBox> WrapBox =
 				SNew( SWrapBox )
-				.PreferredWidth(this, &STargetsSummaryWidget::GetPreferredWidthForWrapping);
+				.PreferredSize(this, &STargetsSummaryWidget::GetPreferredWidthForWrapping);
 
 			// Show tiles only.
 			for(uint32 i = 0; i < TargetCount; ++i)
diff --git a/Engine/Source/Editor/DetailCustomizations/Private/ObjectDetails.cpp b/Engine/Source/Editor/DetailCustomizations/Private/ObjectDetails.cpp
--- a/Engine/Source/Editor/DetailCustomizations/Private/ObjectDetails.cpp
+++ b/Engine/Source/Editor/DetailCustomizations/Private/ObjectDetails.cpp
@@ -167,7 +167,7 @@ void FObjectDetails::AddCallInEditorMethods(IDetailLayoutBuilder& DetailBuilder)
 			FCategoryEntry(FName InCategoryName)
 				: CategoryName(InCategoryName)
 			{
-				WrapBox = SNew(SWrapBox).UseAllottedWidth(true);
+				WrapBox = SNew(SWrapBox).UseAllottedSize(true);
 			}
 		};
 
diff --git a/Engine/Source/Editor/FoliageEdit/Private/SFoliageEdit.cpp b/Engine/Source/Editor/FoliageEdit/Private/SFoliageEdit.cpp
--- a/Engine/Source/Editor/FoliageEdit/Private/SFoliageEdit.cpp
+++ b/Engine/Source/Editor/FoliageEdit/Private/SFoliageEdit.cpp
@@ -289,7 +289,7 @@ void SFoliageEdit::Construct(const FArguments& InArgs)
 						.Padding(StandardPadding)
 						[
 							SNew(SWrapBox)
-							.UseAllottedWidth(true)
+							.UseAllottedSize(true)
 							.InnerSlotPadding({6, 5})
 
 							+ SWrapBox::Slot()
@@ -338,7 +338,7 @@ void SFoliageEdit::Construct(const FArguments& InArgs)
 						.Padding(StandardPadding)
 						[
 							SNew(SWrapBox)
-							.UseAllottedWidth(true)
+							.UseAllottedSize(true)
 							.InnerSlotPadding({6, 5})
 
 							+ SWrapBox::Slot()
@@ -454,7 +454,7 @@ void SFoliageEdit::Construct(const FArguments& InArgs)
 					.AutoHeight()
 					[
 						SNew(SWrapBox)
-						.UseAllottedWidth(true)
+						.UseAllottedSize(true)
 						.Visibility(this, &SFoliageEdit::GetVisibility_SelectionOptions)
 
 						// Select all instances
diff --git a/Engine/Source/Editor/GraphEditor/Private/KismetNodes/SGraphNodeK2Copy.cpp b/Engine/Source/Editor/GraphEditor/Private/KismetNodes/SGraphNodeK2Copy.cpp
--- a/Engine/Source/Editor/GraphEditor/Private/KismetNodes/SGraphNodeK2Copy.cpp
+++ b/Engine/Source/Editor/GraphEditor/Private/KismetNodes/SGraphNodeK2Copy.cpp
@@ -75,7 +75,7 @@ public:
 		// Create the widget used for the pin body (status indicator, label, and value)
 		TSharedRef<SWrapBox> LocalLabelAndValue =
 			SNew(SWrapBox)
-			.PreferredWidth(150.f);
+			.PreferredSize(150.f);
 
 		LocalLabelAndValue->AddSlot()
 			.VAlign(VAlign_Center)
diff --git a/Engine/Source/Editor/GraphEditor/Private/SGraphPin.cpp b/Engine/Source/Editor/GraphEditor/Private/SGraphPin.cpp
--- a/Engine/Source/Editor/GraphEditor/Private/SGraphPin.cpp
+++ b/Engine/Source/Editor/GraphEditor/Private/SGraphPin.cpp
@@ -198,7 +198,7 @@ void SGraphPin::Construct(const FArguments& InArgs, UEdGraphPin* InPin)
 	// Create the widget used for the pin body (status indicator, label, and value)
 	LabelAndValue =
 		SNew(SWrapBox)
-		.PreferredWidth(150.f);
+		.PreferredSize(150.f);
 
 	if (!bIsInput)
 	{
diff --git a/Engine/Source/Editor/HierarchicalLODOutliner/Private/HLODOutliner.cpp b/Engine/Source/Editor/HierarchicalLODOutliner/Private/HLODOutliner.cpp
--- a/Engine/Source/Editor/HierarchicalLODOutliner/Private/HLODOutliner.cpp
+++ b/Engine/Source/Editor/HierarchicalLODOutliner/Private/HLODOutliner.cpp
@@ -318,7 +318,7 @@ namespace HLODOutliner
 			.Padding(FMargin(0.0f, 2.0f))
 			[
 				SNew(SWrapBox)
-				.UseAllottedWidth(true)
+				.UseAllottedSize(true)
 
 				+ SWrapBox::Slot()
 				.Padding(FMargin(2.0f))
diff --git a/Engine/Source/Editor/Kismet/Private/BlueprintEditorTabFactories.cpp b/Engine/Source/Editor/Kismet/Private/BlueprintEditorTabFactories.cpp
--- a/Engine/Source/Editor/Kismet/Private/BlueprintEditorTabFactories.cpp
+++ b/Engine/Source/Editor/Kismet/Private/BlueprintEditorTabFactories.cpp
@@ -247,7 +247,7 @@ TSharedRef<SWidget> FDefaultsEditorSummoner::CreateOptionalDataOnlyMessage() con
 				.BorderImage(FEditorStyle::GetBrush("ToolPanel.GroupBorder"))
 				[
 					SNew(SWrapBox)
-					.UseAllottedWidth(true)
+					.UseAllottedSize(true)
 
 					+ SWrapBox::Slot()
 					[
diff --git a/Engine/Source/Editor/Layers/Private/SActorLayerCloud.h b/Engine/Source/Editor/Layers/Private/SActorLayerCloud.h
--- a/Engine/Source/Editor/Layers/Private/SActorLayerCloud.h
+++ b/Engine/Source/Editor/Layers/Private/SActorLayerCloud.h
@@ -55,7 +55,7 @@ public:
 			+SOverlay::Slot()
 			[
 				SAssignNew( TagArea, SWrapBox )
-				.UseAllottedWidth( true )
+				.UseAllottedSize( true )
 			]
 
 			+SOverlay::Slot()
diff --git a/Engine/Source/Editor/NewLevelDialog/Private/NewLevelDialogModule.cpp b/Engine/Source/Editor/NewLevelDialog/Private/NewLevelDialogModule.cpp
--- a/Engine/Source/Editor/NewLevelDialog/Private/NewLevelDialogModule.cpp
+++ b/Engine/Source/Editor/NewLevelDialog/Private/NewLevelDialogModule.cpp
@@ -140,7 +140,7 @@ public:
 					.Padding(15)
 					[
 						SAssignNew(TemplatesWrapBox, SWrapBox)
-						.PreferredWidth(DEFAULT_WINDOW_SIZE.X - 35.0)   // apparently no way to auto size the width of wrap boxes
+						.PreferredSize(DEFAULT_WINDOW_SIZE.X - 35.0)   // apparently no way to auto size the width of wrap boxes
 					]
 				]
 				+SVerticalBox::Slot()
diff --git a/Engine/Source/Editor/Persona/Private/PersonaMeshDetails.cpp b/Engine/Source/Editor/Persona/Private/PersonaMeshDetails.cpp
--- a/Engine/Source/Editor/Persona/Private/PersonaMeshDetails.cpp
+++ b/Engine/Source/Editor/Persona/Private/PersonaMeshDetails.cpp
@@ -342,7 +342,7 @@ void SSkeletalLODActions::Construct(const FArguments& InArgs)
 			+ SVerticalBox::Slot()
 			[
 				SAssignNew(WrapBox, SWrapBox)
-				.UseAllottedWidth(true)
+				.UseAllottedSize(true)
 			]
 		];
 
diff --git a/Engine/Source/Editor/PinnedCommandList/Private/SPinnedCommandList.cpp b/Engine/Source/Editor/PinnedCommandList/Private/SPinnedCommandList.cpp
--- a/Engine/Source/Editor/PinnedCommandList/Private/SPinnedCommandList.cpp
+++ b/Engine/Source/Editor/PinnedCommandList/Private/SPinnedCommandList.cpp
@@ -327,7 +327,7 @@ void SPinnedCommandList::Construct( const FArguments& InArgs, const FName& InCon
 	ChildSlot
 	[
 		SAssignNew(CommandBox, SWrapBox)
-		.UseAllottedWidth(true)
+		.UseAllottedSize(true)
 		.InnerSlotPadding(FVector2D(0.0f, 0.0f))
 	];
 }
diff --git a/Engine/Source/Editor/PropertyEditor/Private/SPerPlatformPropertiesWidget.cpp b/Engine/Source/Editor/PropertyEditor/Private/SPerPlatformPropertiesWidget.cpp
--- a/Engine/Source/Editor/PropertyEditor/Private/SPerPlatformPropertiesWidget.cpp
+++ b/Engine/Source/Editor/PropertyEditor/Private/SPerPlatformPropertiesWidget.cpp
@@ -40,7 +40,7 @@ void SPerPlatformPropertiesWidget::ConstructChildren()
 	.HAlign(HAlign_Fill)
 	[
 		SAssignNew(WrapBox, SWrapBox)
-		.UseAllottedWidth(true)
+		.UseAllottedSize(true)
 	];
 
 	if (OnGenerateWidget.IsBound())
diff --git a/Engine/Source/Editor/Sequencer/Private/SSequencer.cpp b/Engine/Source/Editor/Sequencer/Private/SSequencer.cpp
--- a/Engine/Source/Editor/Sequencer/Private/SSequencer.cpp
+++ b/Engine/Source/Editor/Sequencer/Private/SSequencer.cpp
@@ -576,7 +576,7 @@ void SSequencer::Construct(const FArguments& InArgs, TSharedRef<FSequencer> InSe
 						.Padding(FMargin(CommonPadding, 0.f))
 						[
 							SNew(SWrapBox)
-							.UseAllottedWidth(true)
+							.UseAllottedSize(true)
 							.InnerSlotPadding(FVector2D(5, 2))
 
 							+ SWrapBox::Slot()
diff --git a/Engine/Source/Editor/SourceControlWindows/Private/SSourceControlSubmit.cpp b/Engine/Source/Editor/SourceControlWindows/Private/SSourceControlSubmit.cpp
--- a/Engine/Source/Editor/SourceControlWindows/Private/SSourceControlSubmit.cpp
+++ b/Engine/Source/Editor/SourceControlWindows/Private/SSourceControlSubmit.cpp
@@ -175,7 +175,7 @@ void SSourceControlSubmitWidget::Construct(const FArguments& InArgs)
 			.Padding(5)
 			[
 				SNew(SWrapBox)
-				.UseAllottedWidth(true)
+				.UseAllottedSize(true)
 				+SWrapBox::Slot()
 				.Padding(0.0f, 0.0f, 16.0f, 0.0f)
 				[
diff --git a/Engine/Source/Editor/VirtualTexturingEditor/Private/RuntimeVirtualTextureDetailsCustomization.cpp b/Engine/Source/Editor/VirtualTexturingEditor/Private/RuntimeVirtualTextureDetailsCustomization.cpp
--- a/Engine/Source/Editor/VirtualTexturingEditor/Private/RuntimeVirtualTextureDetailsCustomization.cpp
+++ b/Engine/Source/Editor/VirtualTexturingEditor/Private/RuntimeVirtualTextureDetailsCustomization.cpp
@@ -51,7 +51,7 @@ namespace
 			.Padding(4.0f)
 			[
 				SNew(SWrapBox)
-				.UseAllottedWidth(true)
+				.UseAllottedSize(true)
 
 				+ SWrapBox::Slot()
 				.Padding(FMargin(0.0f, 2.0f, 2.0f, 0.0f))
diff --git a/Engine/Source/Editor/WorldBrowser/Private/Tiles/SWorldComposition.cpp b/Engine/Source/Editor/WorldBrowser/Private/Tiles/SWorldComposition.cpp
--- a/Engine/Source/Editor/WorldBrowser/Private/Tiles/SWorldComposition.cpp
+++ b/Engine/Source/Editor/WorldBrowser/Private/Tiles/SWorldComposition.cpp
@@ -1009,7 +1009,7 @@ TSharedRef<SWidget> SWorldComposition::ConstructContentWidget()
 		.AutoHeight()
 		[
 			SAssignNew(LayersListWrapBox, SWrapBox)
-			.UseAllottedWidth(true)
+			.UseAllottedSize(true)
 		]
 				
 		+SVerticalBox::Slot()
diff --git a/Engine/Source/Runtime/Slate/Private/Widgets/Layout/SWrapBox.cpp b/Engine/Source/Runtime/Slate/Private/Widgets/Layout/SWrapBox.cpp
--- a/Engine/Source/Runtime/Slate/Private/Widgets/Layout/SWrapBox.cpp
+++ b/Engine/Source/Runtime/Slate/Private/Widgets/Layout/SWrapBox.cpp
@@ -36,9 +36,17 @@ int32 SWrapBox::RemoveSlot( const TSharedRef<SWidget>& SlotWidget )
 
 void SWrapBox::Construct( const FArguments& InArgs )
 {
-	PreferredWidth = InArgs._PreferredWidth;
+	PreferredSize = InArgs._PreferredSize;
+
+	// Handle deprecation of PreferredWidth
+	if (!PreferredSize.IsSet() && !PreferredSize.IsBound())
+	{
+		PreferredSize = InArgs._PreferredWidth;
+	}
+
 	InnerSlotPadding = InArgs._InnerSlotPadding;
-	bUseAllottedWidth = InArgs._UseAllottedWidth;
+	bUseAllottedSize = InArgs._UseAllottedSize;
+	Orientation = InArgs._Orientation;
 
 	// Copy the children from the declaration to the widget
 	for ( int32 ChildIndex=0; ChildIndex < InArgs.Slots.Num(); ++ChildIndex )
@@ -49,9 +57,9 @@ void SWrapBox::Construct( const FArguments& InArgs )
 
 void SWrapBox::Tick( const FGeometry& AllottedGeometry, const double InCurrentTime, const float InDeltaTime )
 {
-	if (bUseAllottedWidth)
+	if (bUseAllottedSize)
 	{
-		PreferredWidth = AllottedGeometry.GetLocalSize().X;
+		PreferredSize = Orientation == EOrientation::Orient_Vertical ? AllottedGeometry.GetLocalSize().Y : AllottedGeometry.GetLocalSize().X;
 	}
 }
 
@@ -81,7 +89,7 @@ private:
 	const SWrapBox& WrapBox;
 	const FOnSlotArranged& OnSlotArranged;
 	FVector2D Offset;
-	float MaximumHeightInCurrentLine;
+	float MaximumSizeInCurrentLine;
 	int32 IndexOfFirstChildInCurrentLine;
 	TMap<int32, FArrangementData> OngoingArrangementDataMap;
 };
@@ -91,7 +99,7 @@ SWrapBox::FChildArranger::FChildArranger(const SWrapBox& InWrapBox, const FOnSlo
 	: WrapBox(InWrapBox)
 	, OnSlotArranged(InOnSlotArranged)
 	, Offset(FVector2D::ZeroVector)
-	, MaximumHeightInCurrentLine(0.0f)
+	, MaximumSizeInCurrentLine(0.0f)
 	, IndexOfFirstChildInCurrentLine(INDEX_NONE)
 {
 	OngoingArrangementDataMap.Reserve(WrapBox.Slots.Num());
@@ -142,10 +150,11 @@ void SWrapBox::FChildArranger::Arrange()
 			ArrangementData.SlotOffset.Y = Offset.Y;
 		};
 
-		// Rule: If this child is not the first child in the line, "inner slot padding" needs to be injected left of it.
+		// Rule: If this child is not the first child in the line, "inner slot padding" needs to be injected left or top of it, dependently of the orientation.
 		if (!IsFirstChildInCurrentLine())
 		{
-			Offset.X += WrapBox.InnerSlotPadding.X;
+			Offset.Y += ((WrapBox.Orientation == EOrientation::Orient_Vertical) * WrapBox.InnerSlotPadding.Y);
+			Offset.X += ((WrapBox.Orientation == EOrientation::Orient_Horizontal) * WrapBox.InnerSlotPadding.X);
 		}
 
 		const FVector2D DesiredSizeOfSlot = Slot.SlotPadding.Get().GetDesiredSize() + Widget->GetDesiredSize();
@@ -156,35 +165,70 @@ void SWrapBox::FChildArranger::Arrange()
 		ArrangementData.SlotSize.X = DesiredSizeOfSlot.X;
 		ArrangementData.SlotSize.Y = DesiredSizeOfSlot.Y;
 
-		const float RightBoundOfChild = ArrangementData.SlotOffset.X + ArrangementData.SlotSize.X;
-
-		// Rule: If required due to a wrapping width under specified threshold, start a new line and allocate all of it to this child.
-		if (Slot.SlotFillLineWhenWidthLessThan.IsSet() && WrapBox.PreferredWidth.Get() < Slot.SlotFillLineWhenWidthLessThan.GetValue())
+		if (WrapBox.Orientation == EOrientation::Orient_Vertical)
 		{
-			// Begin a new line if the current one isn't empty, because we demand a whole line to ourselves.
-			if (!IsFirstChildInCurrentLine())
+			const float BottomBoundOfChild = ArrangementData.SlotOffset.Y + ArrangementData.SlotSize.Y;
+
+			// Rule: If required due to a wrapping height under specified threshold, start a new line and allocate all of it to this child.
+			if (Slot.SlotFillLineWhenSizeLessThan.IsSet() && WrapBox.PreferredSize.Get() < Slot.SlotFillLineWhenSizeLessThan.GetValue())
+			{
+				// Begin a new line if the current one isn't empty, because we demand a whole line to ourselves.
+				if (!IsFirstChildInCurrentLine())
+				{
+					BeginNewLine();
+				}
+
+				// Fill height of rest of wrap box.
+				ArrangementData.SlotSize.Y = WrapBox.PreferredSize.Get() - Offset.Y;
+			}
+			// Rule: If the end of a child would go beyond the width to wrap at, it should move to a new line.
+			else if (BottomBoundOfChild > WrapBox.PreferredSize.Get())
 			{
-				BeginNewLine();
+				// Begin a new line if the current one isn't empty, because we demand a new line.
+				if (!IsFirstChildInCurrentLine())
+				{
+					BeginNewLine();
+				}
 			}
 
-			// Fill width of rest of wrap box.
-			ArrangementData.SlotSize.X = WrapBox.PreferredWidth.Get() - Offset.X;
+			// Update current line maximum size.
+			MaximumSizeInCurrentLine = FMath::Max(MaximumSizeInCurrentLine, ArrangementData.SlotSize.X);
+
+			// Update offset to bottom bound of child.
+			Offset.Y = ArrangementData.SlotOffset.Y + ArrangementData.SlotSize.Y;
 		}
-		// Rule: If the end of a child would go beyond the width to wrap at, it should move to a new line.
-		else if (RightBoundOfChild > WrapBox.PreferredWidth.Get())
+		else
 		{
-			// Begin a new line if the current one isn't empty, because we demand a new line.
-			if (!IsFirstChildInCurrentLine())
+			const float RightBoundOfChild = ArrangementData.SlotOffset.X + ArrangementData.SlotSize.X;
+
+			// Rule: If required due to a wrapping width under specified threshold, start a new line and allocate all of it to this child.
+			if (Slot.SlotFillLineWhenSizeLessThan.IsSet() && WrapBox.PreferredSize.Get() < Slot.SlotFillLineWhenSizeLessThan.GetValue())
 			{
-				BeginNewLine();
+				// Begin a new line if the current one isn't empty, because we demand a whole line to ourselves.
+				if (!IsFirstChildInCurrentLine())
+				{
+					BeginNewLine();
+				}
+
+				// Fill width of rest of wrap box.
+				ArrangementData.SlotSize.X = WrapBox.PreferredSize.Get() - Offset.X;
+			}
+			// Rule: If the end of a child would go beyond the width to wrap at, it should move to a new line.
+			else if (RightBoundOfChild > WrapBox.PreferredSize.Get())
+			{
+				// Begin a new line if the current one isn't empty, because we demand a new line.
+				if (!IsFirstChildInCurrentLine())
+				{
+					BeginNewLine();
+				}
 			}
-		}
 
-		// Update current line maximum height.
-		MaximumHeightInCurrentLine = FMath::Max(MaximumHeightInCurrentLine, ArrangementData.SlotSize.Y);
+			// Update current line maximum size.
+			MaximumSizeInCurrentLine = FMath::Max(MaximumSizeInCurrentLine, ArrangementData.SlotSize.Y);
 
-		// Update offset to right bound of child.
-		Offset.X = ArrangementData.SlotOffset.X + ArrangementData.SlotSize.X;
+			// Update offset to right bound of child.
+			Offset.X = ArrangementData.SlotOffset.X + ArrangementData.SlotSize.X;
+		}
 	}
 
 	// Attempt to finalize the final line if there are any children in it.
@@ -222,7 +266,14 @@ void SWrapBox::FChildArranger::FinalizeLine(int32 IndexOfLastChildInCurrentLine)
 		// Rule: The last uncollapsed child in a line may request to fill the remaining empty space in the line.
 		if (ChildIndex == IndexOfLastChildInCurrentLine && Slot.bSlotFillEmptySpace)
 		{
-			ArrangementData.SlotSize.X = WrapBox.PreferredWidth.Get() - ArrangementData.SlotOffset.X;
+			if (WrapBox.Orientation == EOrientation::Orient_Vertical)
+			{
+				ArrangementData.SlotSize.Y = WrapBox.PreferredSize.Get() - ArrangementData.SlotOffset.Y;
+			}
+			else
+			{
+				ArrangementData.SlotSize.X = WrapBox.PreferredSize.Get() - ArrangementData.SlotOffset.X;
+			}
 		}
 		
 		// All slots on this line should now match to the tallest element's height, which they can then use to do their alignment in OnSlotArranged below (eg. center within that)
@@ -233,11 +284,24 @@ void SWrapBox::FChildArranger::FinalizeLine(int32 IndexOfLastChildInCurrentLine)
 		OnSlotArranged(Slot, ArrangementData);
 	}
 
-	// Set initial state for new line.
-	Offset.X = 0.0f;
-	// Since this is the initial state for a new line, this only happens after the first line, so the inner slot vertical padding should always be added.
-	Offset.Y += MaximumHeightInCurrentLine + WrapBox.InnerSlotPadding.Y;
-	MaximumHeightInCurrentLine = 0.0f;
+	if (WrapBox.Orientation == EOrientation::Orient_Vertical)
+	{
+		// Set initial state for new vertical line.
+		Offset.Y = 0.0f;
+
+		// Since this is the initial state for a new vertical line, this only happens after the first line, so the inner slot horizontal padding should always be added.
+		Offset.X += MaximumSizeInCurrentLine + WrapBox.InnerSlotPadding.X;
+	}
+	else
+	{
+		// Set initial state for horizontal new line.
+		Offset.X = 0.0f;
+
+		// Since this is the initial state for a new horizontal line, this only happens after the first line, so the inner slot vertical padding should always be added.
+		Offset.Y += MaximumSizeInCurrentLine + WrapBox.InnerSlotPadding.Y;
+	}
+
+	MaximumSizeInCurrentLine = 0.0f;
 	IndexOfFirstChildInCurrentLine = INDEX_NONE;
 }
 
@@ -248,8 +312,6 @@ void SWrapBox::FChildArranger::Arrange(const SWrapBox& WrapBox, const FOnSlotArr
 
 void SWrapBox::OnArrangeChildren(const FGeometry& AllottedGeometry, FArrangedChildren& ArrangedChildren) const
 {
-	const float WidthToWrapAt = PreferredWidth.Get();
-
 	FChildArranger::Arrange(*this, [&](const FSlot& Slot, const FChildArranger::FArrangementData& ArrangementData)
 	{
 		// Calculate offset and size in slot using alignment.
@@ -295,10 +357,25 @@ void SWrapBox::SetInnerSlotPadding(FVector2D InInnerSlotPadding)
 
 void SWrapBox::SetWrapWidth(const TAttribute<float>& InWrapWidth)
 {
-	PreferredWidth = InWrapWidth;
+	PreferredSize = InWrapWidth;
+}
+
+void SWrapBox::SetWrapSize(const TAttribute<float>& InWrapSize)
+{
+	PreferredSize = InWrapSize;
 }
 
 void SWrapBox::SetUseAllottedWidth(bool bInUseAllottedWidth)
 {
-	bUseAllottedWidth = bInUseAllottedWidth;
+	bUseAllottedSize = bInUseAllottedWidth;
+}
+
+void SWrapBox::SetUseAllottedSize(bool bInUseAllottedSize)
+{
+	bUseAllottedSize = bInUseAllottedSize;
 }
+
+void SWrapBox::SetOrientation(EOrientation InOrientation)
+{
+	Orientation = InOrientation;
+}
\ No newline at end of file
diff --git a/Engine/Source/Runtime/Slate/Public/Widgets/Layout/SWrapBox.h b/Engine/Source/Runtime/Slate/Public/Widgets/Layout/SWrapBox.h
--- a/Engine/Source/Runtime/Slate/Public/Widgets/Layout/SWrapBox.h
+++ b/Engine/Source/Runtime/Slate/Public/Widgets/Layout/SWrapBox.h
@@ -14,18 +14,39 @@
 class FArrangedChildren;
 
 /**
+ * With EOrientation::Orient_Horizontal
  * Arranges widgets left-to-right.
- * When the widgets exceed the PreferredWidth
+ * When the widgets exceed the PreferredSize
  * the SWrapBox will place widgets on the next line.
  *
  * Illustration:
- *                      +-----Preferred Width
+ *                      +-----Preferred Size
  *                      |
  *       [-----------][-|-]
  *       [--][------[--]|
  *       [--------------|]
  *       [---]          |
  */
+
+ /**
+  * With EOrientation::Orient_Vertical
+  * Arranges widgets top-to-bottom.
+  * When the widgets exceed the PreferredSize
+  * the SVerticalWrapBox will place widgets on the next line.
+  *
+  * Illustration:
+  *
+  *      [___]  [___]
+  *      [-1-]  [-3-]
+  *
+  *		 [___]  [___]
+  *      [-2-]  [-4-]
+  *
+  *      [___]
+  *==============================>--------Preferred Size
+  *		 [-3-]
+  */
+
 class SLATE_API SWrapBox : public SPanel
 {
 public:
@@ -37,34 +58,36 @@ public:
 		FSlot()
 			: TSlotBase<FSlot>()
 			, TSupportsContentAlignmentMixin<FSlot>(HAlign_Fill, VAlign_Fill)
-			, SlotFillLineWhenWidthLessThan()
+			, SlotFillLineWhenSizeLessThan()
 			, bSlotFillEmptySpace(false)
 		{
 		}
 
 		/** If the total available space in the wrap panel drops below this threshold, this slot will attempt to fill an entire line. */
-		FSlot& FillLineWhenWidthLessThan(TOptional<float> InFillLineWhenWidthLessThan)
+		FSlot& FillLineWhenWidthLessThan(TOptional<float> InFillLineWhenSizeLessThan)
 		{
-			SlotFillLineWhenWidthLessThan = InFillLineWhenWidthLessThan;
-			return *( static_cast<FSlot*>( this ) );
+			SlotFillLineWhenSizeLessThan = InFillLineWhenSizeLessThan;
+			return *(static_cast<FSlot*>(this));
 		}
 
 		/** Should this slot fill the remaining space on the line? */
 		FSlot& FillEmptySpace(bool bInFillEmptySpace)
 		{
 			bSlotFillEmptySpace = bInFillEmptySpace;
-			return *( static_cast<FSlot*>( this ) );
+			return *(static_cast<FSlot*>(this));
 		}
 
-		TOptional<float> SlotFillLineWhenWidthLessThan;
+		TOptional<float> SlotFillLineWhenSizeLessThan;
 		bool bSlotFillEmptySpace;
 	};
 
 
 	SLATE_BEGIN_ARGS(SWrapBox)
-		: _PreferredWidth( 100.f )
-		, _InnerSlotPadding( FVector2D::ZeroVector )
-		, _UseAllottedWidth( false )
+		: _PreferredWidth(100.f)
+		, _PreferredSize(100.f)
+		, _InnerSlotPadding(FVector2D::ZeroVector)
+		, _UseAllottedSize(false)
+		, _Orientation(EOrientation::Orient_Horizontal)
 		{
 			_Visibility = EVisibility::SelfHitTestInvisible;
 		}
@@ -75,11 +98,17 @@ public:
 		/** The preferred width, if not set will fill the space */
 		SLATE_ATTRIBUTE( float, PreferredWidth )
 
+		/** The preferred size, if not set will fill the space */
+		SLATE_ATTRIBUTE( float, PreferredSize )
+
 		/** The inner slot padding goes between slots sharing borders */
 		SLATE_ARGUMENT( FVector2D, InnerSlotPadding )
 
-		/** if true, the PreferredWidth will always match the room available to the SWrapBox  */
-		SLATE_ARGUMENT( bool, UseAllottedWidth )
+		/** if true, the PreferredSize will always match the room available to the SWrapBox  */
+		SLATE_ARGUMENT( bool, UseAllottedSize )
+
+		/** Determines if the wrap box needs to arrange the slots left-to-right or top-to-bottom.*/
+		SLATE_ARGUMENT_DEFAULT(EOrientation, Orientation) { EOrientation::Orient_Horizontal };
 	SLATE_END_ARGS()
 
 	SWrapBox();
@@ -111,16 +140,31 @@ public:
 	void SetInnerSlotPadding(FVector2D InInnerSlotPadding);
 
 	/** Set the width at which the wrap panel should wrap its content. */
-	void SetWrapWidth( const TAttribute<float>& InWrapWidth );
+	UE_DEPRECATED(4.26, "Deprecated, please use SetWrapSize() instead")
+	void SetWrapWidth(const TAttribute<float>& InWrapWidth);
+
+	/** Set the size at which the wrap panel should wrap its content. */
+	void SetWrapSize( const TAttribute<float>& InWrapSize );
 
 	/** When true, use the WrapWidth property to determine where to wrap to the next line. */
+	UE_DEPRECATED(4.26, "Deprecated, please use SetUseAllottedSize() instead")
 	void SetUseAllottedWidth(bool bInUseAllottedWidth);
 
+	/** When true, use the WrapSize property to determine where to wrap to the next line. */
+	void SetUseAllottedSize(bool bInUseAllottedSize);
+
+	/** Set the Orientation to determine if the wrap box needs to arrange the slots left-to-right or top-to-bottom */
+	void SetOrientation(EOrientation InOrientation);
+
 private:
 
 	/** How wide this panel should appear to be. Any widgets past this line will be wrapped onto the next line. */
+	UE_DEPRECATED(4.26, "Deprecated, please use PreferredSize instead")
 	TAttribute<float> PreferredWidth;
 
+	/** How wide or long, dependently of the orientation, this panel should appear to be. Any widgets past this line will be wrapped onto the next line. */
+	TAttribute<float> PreferredSize;
+
 	/** The slots that contain this panel's children. */
 	TPanelChildren<FSlot> Slots;
 
@@ -128,9 +172,15 @@ private:
 	FVector2D InnerSlotPadding;
 
 	/** If true the box will have a preferred width equal to its alloted width  */
+	UE_DEPRECATED(4.26, "Deprecated, please use bUseAllottedSize instead")
 	bool bUseAllottedWidth;
 
+	/** If true the box will have a preferred size equal to its alloted size  */
+	bool bUseAllottedSize;
+
+	/** Determines if the wrap box needs to arrange the slots left-to-right or top-to-bottom.*/
+	EOrientation Orientation = EOrientation::Orient_Horizontal;
+
 	class FChildArranger;
 	friend class SWrapBox::FChildArranger;
-};
-
+};
\ No newline at end of file
diff --git a/Engine/Source/Runtime/UMG/Private/Components/DynamicEntryBoxBase.cpp b/Engine/Source/Runtime/UMG/Private/Components/DynamicEntryBoxBase.cpp
--- a/Engine/Source/Runtime/UMG/Private/Components/DynamicEntryBoxBase.cpp
+++ b/Engine/Source/Runtime/UMG/Private/Components/DynamicEntryBoxBase.cpp
@@ -39,6 +39,7 @@ void UDynamicEntryBoxBase::ResetInternal(bool bDeleteWidgets)
 			StaticCastSharedPtr<SBoxPanel>(MyPanelWidget)->ClearChildren();
 			break;
 		case EDynamicBoxType::Wrap:
+		case EDynamicBoxType::VerticalWrap:
 			StaticCastSharedPtr<SWrapBox>(MyPanelWidget)->ClearChildren();
 			break;
 		case EDynamicBoxType::Overlay:
@@ -74,6 +75,7 @@ void UDynamicEntryBoxBase::RemoveEntryInternal(UUserWidget* EntryWidget)
 					StaticCastSharedPtr<SBoxPanel>(MyPanelWidget)->RemoveSlot(CachedEntryWidget.ToSharedRef());
 					break;
 				case EDynamicBoxType::Wrap:
+				case EDynamicBoxType::VerticalWrap:
 					StaticCastSharedPtr<SWrapBox>(MyPanelWidget)->RemoveSlot(CachedEntryWidget.ToSharedRef());
 					break;
 				case EDynamicBoxType::Overlay:
@@ -92,7 +94,7 @@ void UDynamicEntryBoxBase::SetEntrySpacing(const FVector2D& InEntrySpacing)
 
 	if (MyPanelWidget.IsValid())
 	{
-		if (EntryBoxType == EDynamicBoxType::Wrap)
+		if (EntryBoxType == EDynamicBoxType::Wrap || EntryBoxType == EDynamicBoxType::VerticalWrap)
 		{
 			// Wrap boxes can change their widget spacing on the fly
 			StaticCastSharedPtr<SWrapBox>(MyPanelWidget)->SetInnerSlotPadding(EntrySpacing);
@@ -197,7 +199,14 @@ TSharedRef<SWidget> UDynamicEntryBoxBase::RebuildWidget()
 		break;
 	case EDynamicBoxType::Wrap:
 		EntryBoxWidget = SAssignNew(MyPanelWidget, SWrapBox)
-			.UseAllottedWidth(true)
+			.UseAllottedSize(true)
+			.Orientation(EOrientation::Orient_Horizontal)
+			.InnerSlotPadding(EntrySpacing);
+		break;
+	case EDynamicBoxType::VerticalWrap:
+		EntryBoxWidget = SAssignNew(MyPanelWidget, SWrapBox)
+			.UseAllottedSize(true)
+			.Orientation(EOrientation::Orient_Vertical)
 			.InnerSlotPadding(EntrySpacing);
 		break;
 	case EDynamicBoxType::Overlay:
@@ -302,7 +311,7 @@ FMargin UDynamicEntryBoxBase::BuildEntryPadding(const FVector2D& DesiredSpacing)
 void UDynamicEntryBoxBase::AddEntryChild(UUserWidget& ChildWidget)
 {
 	FSlotBase* NewSlot = nullptr;
-	if (EntryBoxType == EDynamicBoxType::Wrap)
+	if (EntryBoxType == EDynamicBoxType::Wrap || EntryBoxType == EDynamicBoxType::VerticalWrap)
 	{
 		NewSlot = &StaticCastSharedPtr<SWrapBox>(MyPanelWidget)->AddSlot()
 			.FillEmptySpace(false)
diff --git a/Engine/Source/Runtime/UMG/Private/Components/WrapBox.cpp b/Engine/Source/Runtime/UMG/Private/Components/WrapBox.cpp
--- a/Engine/Source/Runtime/UMG/Private/Components/WrapBox.cpp
+++ b/Engine/Source/Runtime/UMG/Private/Components/WrapBox.cpp
@@ -13,8 +13,26 @@ UWrapBox::UWrapBox(const FObjectInitializer& ObjectInitializer)
 {
 	bIsVariable = false;
 	Visibility = ESlateVisibility::SelfHitTestInvisible;
-	WrapWidth = 500;
-	bExplicitWrapWidth = false;
+	WrapSize = 500;
+	bExplicitWrapSize = false;
+	Orientation = EOrientation::Orient_Horizontal;
+}
+
+void UWrapBox::PostLoad()
+{
+	Super::PostLoad();
+
+	if (WrapWidth_DEPRECATED != 0.0f)
+	{
+		WrapSize = WrapWidth_DEPRECATED;
+		WrapWidth_DEPRECATED = 0.0f;
+	}
+
+	if (bExplicitWrapWidth_DEPRECATED)
+	{
+		bExplicitWrapSize = bExplicitWrapWidth_DEPRECATED;
+		bExplicitWrapWidth_DEPRECATED = false;
+	}
 }
 
 void UWrapBox::ReleaseSlateResources(bool bReleaseChildren)
@@ -64,8 +82,9 @@ UWrapBoxSlot* UWrapBox::AddChildToWrapBox(UWidget* Content)
 TSharedRef<SWidget> UWrapBox::RebuildWidget()
 {
 	MyWrapBox = SNew(SWrapBox)
-		.UseAllottedWidth(!bExplicitWrapWidth)
-		.PreferredWidth(WrapWidth);
+		.UseAllottedSize(!bExplicitWrapSize)
+		.PreferredSize(WrapSize)
+		.Orientation(Orientation);
 
 	for ( UPanelSlot* PanelSlot : Slots )
 	{
@@ -84,8 +103,9 @@ void UWrapBox::SynchronizeProperties()
 	Super::SynchronizeProperties();
 
 	MyWrapBox->SetInnerSlotPadding(InnerSlotPadding);
-	MyWrapBox->SetUseAllottedWidth(!bExplicitWrapWidth);
-	MyWrapBox->SetWrapWidth(WrapWidth);
+	MyWrapBox->SetUseAllottedSize(!bExplicitWrapSize);
+	MyWrapBox->SetWrapSize(WrapSize);
+	MyWrapBox->SetOrientation(Orientation);
 }
 
 void UWrapBox::SetInnerSlotPadding(FVector2D InPadding)
diff --git a/Engine/Source/Runtime/UMG/Public/Components/DynamicEntryBoxBase.h b/Engine/Source/Runtime/UMG/Public/Components/DynamicEntryBoxBase.h
--- a/Engine/Source/Runtime/UMG/Public/Components/DynamicEntryBoxBase.h
+++ b/Engine/Source/Runtime/UMG/Public/Components/DynamicEntryBoxBase.h
@@ -15,6 +15,7 @@ enum class EDynamicBoxType : uint8
 	Horizontal,
 	Vertical,
 	Wrap,
+	VerticalWrap,
 	Overlay
 };
 
diff --git a/Engine/Source/Runtime/UMG/Public/Components/WrapBox.h b/Engine/Source/Runtime/UMG/Public/Components/WrapBox.h
--- a/Engine/Source/Runtime/UMG/Public/Components/WrapBox.h
+++ b/Engine/Source/Runtime/UMG/Public/Components/WrapBox.h
@@ -12,7 +12,7 @@
 class UWrapBoxSlot;
 
 /**
- * Arranges widgets left-to-right.  When the widgets exceed the Width it will place widgets on the next line.
+ * Arranges widgets left-to-right or top-to-bottom dependently of the orientation.  When the widgets exceed the wrapSize it will place widgets on the next line.
  * 
  * * Many Children
  * * Flows
@@ -28,19 +28,30 @@ public:
 	UPROPERTY(EditAnywhere, BlueprintReadOnly, Category="Content Layout")
 	FVector2D InnerSlotPadding;
 
-	/** When this width is exceeded, elements will start appearing on the next line. */
-	UPROPERTY(EditAnywhere, BlueprintReadOnly, Category = "Content Layout", meta=(EditCondition = "bExplicitWrapWidth"))
-	float WrapWidth;
+	/** DEPRECATED value replaced by WrapSize, When this width is exceeded, elements will start appearing on the next line. */
+	UPROPERTY()
+	float WrapWidth_DEPRECATED;
 
-	/** Use explicit wrap width whenever possible. It greatly simplifies layout calculations and reduces likelihood of "wiggling UI" */
+	/** When this size is exceeded, elements will start appearing on the next line. */
+	UPROPERTY(EditAnywhere, BlueprintReadOnly, Category = "Content Layout", meta=(EditCondition = "bExplicitWrapSize"))
+	float WrapSize;
+
+	/** DEPRECATED value replaced by bExplicitWrapSize, Use explicit wrap width whenever possible. It greatly simplifies layout calculations and reduces likelihood of "wiggling UI" */
+	UPROPERTY()
+	bool bExplicitWrapWidth_DEPRECATED;
+
+	/** Use explicit wrap size whenever possible. It greatly simplifies layout calculations and reduces likelihood of "wiggling UI" */
+	UPROPERTY(EditAnywhere, BlueprintReadOnly, Category = "Content Layout")
+	bool bExplicitWrapSize;
+
+	/** Determines if the Wrap Box should arranges the widgets left-to-right or top-to-bottom */
 	UPROPERTY(EditAnywhere, BlueprintReadOnly, Category = "Content Layout")
-	bool bExplicitWrapWidth;
+	TEnumAsByte<EOrientation> Orientation = EOrientation::Orient_Horizontal;
 
 	/** Sets the inner slot padding goes between slots sharing borders */
 	UFUNCTION(BlueprintCallable, Category="Content Layout")
 	void SetInnerSlotPadding(FVector2D InPadding);
 
-
 public:
 
 	UE_DEPRECATED(4.22, "Deprecated, please use AddChildToWrapBox() instead")
@@ -55,6 +66,8 @@ public:
 	// End UWidget interface
 #endif
 
+	virtual void PostLoad() override;
+
 protected:
 
 	// UPanelWidget
