Wrap Box: a few updates got lost in the merge for adding horizontal/vertical option.

#rnx

diff --git a/Engine/Plugins/FX/Niagara/Source/NiagaraEditor/Private/Customizations/NiagaraPlatformSetCustomization.cpp b/Engine/Plugins/FX/Niagara/Source/NiagaraEditor/Private/Customizations/NiagaraPlatformSetCustomization.cpp
--- a/Engine/Plugins/FX/Niagara/Source/NiagaraEditor/Private/Customizations/NiagaraPlatformSetCustomization.cpp
+++ b/Engine/Plugins/FX/Niagara/Source/NiagaraEditor/Private/Customizations/NiagaraPlatformSetCustomization.cpp
@@ -98,7 +98,7 @@ void FNiagaraPlatformSetCustomization::CustomizeHeader(TSharedRef<IPropertyHandl
 		.WholeRowContent()
 		[ 	
 			SAssignNew(QualityLevelWidgetBox, SWrapBox)
-			.UseAllottedWidth(true)
+			.UseAllottedSize(true)
 		];
 
 	GenerateQualityLevelSelectionWidgets();
diff --git a/Engine/Source/Editor/Persona/Private/SAnimMontageSectionsPanel.cpp b/Engine/Source/Editor/Persona/Private/SAnimMontageSectionsPanel.cpp
--- a/Engine/Source/Editor/Persona/Private/SAnimMontageSectionsPanel.cpp
+++ b/Engine/Source/Editor/Persona/Private/SAnimMontageSectionsPanel.cpp
@@ -275,7 +275,7 @@ void SAnimMontageSectionsPanel::Update()
 					.BorderImage(FCoreStyle::Get().GetBrush("ToolPanel.GroupBorder"))
 					[
 						SAssignNew(WrapBox, SWrapBox)
-						.UseAllottedWidth(true)
+						.UseAllottedSize(true)
 					]
 				];
 			
diff --git a/Engine/Source/Runtime/Slate/Private/Widgets/Layout/SWrapBox.cpp b/Engine/Source/Runtime/Slate/Private/Widgets/Layout/SWrapBox.cpp
--- a/Engine/Source/Runtime/Slate/Private/Widgets/Layout/SWrapBox.cpp
+++ b/Engine/Source/Runtime/Slate/Private/Widgets/Layout/SWrapBox.cpp
@@ -279,7 +279,7 @@ void SWrapBox::FChildArranger::FinalizeLine(int32 IndexOfLastChildInCurrentLine)
 		// All slots on this line should now match to the tallest element's height, which they can then use to do their alignment in OnSlotArranged below (eg. center within that)
 		// If we left their height as is, then their slots would just be whatever their child's desired height was, and so a vertical alignment of "center" would actually 
 		// leave the widget at the top of the line, since you couldn't calculate how much to offset by to actually reach the center of the "container"
-		ArrangementData.SlotSize.Y = MaximumHeightInCurrentLine;
+		ArrangementData.SlotSize.Y = MaximumSizeInCurrentLine;
 
 		OnSlotArranged(Slot, ArrangementData);
 	}
