SWrapBox deprecation fix.

diff --git a/Engine/Source/Runtime/Slate/Private/Widgets/Layout/SWrapBox.cpp b/Engine/Source/Runtime/Slate/Private/Widgets/Layout/SWrapBox.cpp
--- a/Engine/Source/Runtime/Slate/Private/Widgets/Layout/SWrapBox.cpp
+++ b/Engine/Source/Runtime/Slate/Private/Widgets/Layout/SWrapBox.cpp
@@ -45,7 +45,7 @@ void SWrapBox::Construct( const FArguments& InArgs )
 	}
 
 	InnerSlotPadding = InArgs._InnerSlotPadding;
-	bUseAllottedSize = InArgs._UseAllottedSize;
+	bUseAllottedSize = InArgs._UseAllottedSize || InArgs._UseAllottedWidth;
 	Orientation = InArgs._Orientation;
 
 	// Copy the children from the declaration to the widget
diff --git a/Engine/Source/Runtime/Slate/Public/Widgets/Layout/SWrapBox.h b/Engine/Source/Runtime/Slate/Public/Widgets/Layout/SWrapBox.h
--- a/Engine/Source/Runtime/Slate/Public/Widgets/Layout/SWrapBox.h
+++ b/Engine/Source/Runtime/Slate/Public/Widgets/Layout/SWrapBox.h
@@ -86,6 +86,7 @@ public:
 		: _PreferredWidth(100.f)
 		, _PreferredSize(100.f)
 		, _InnerSlotPadding(FVector2D::ZeroVector)
+		, _UseAllottedWidth(false)
 		, _UseAllottedSize(false)
 		, _Orientation(EOrientation::Orient_Horizontal)
 		{
@@ -104,6 +105,9 @@ public:
 		/** The inner slot padding goes between slots sharing borders */
 		SLATE_ARGUMENT( FVector2D, InnerSlotPadding )
 
+		/** if true, the PreferredWidth will always match the room available to the SWrapBox  */
+		SLATE_ARGUMENT( bool, UseAllottedWidth )
+
 		/** if true, the PreferredSize will always match the room available to the SWrapBox  */
 		SLATE_ARGUMENT( bool, UseAllottedSize )
 
