Fix Engine change by readding old values in Slot class of SWrapBox
#RNX

diff --git a/Engine/Source/Editor/ContentBrowser/Private/SContentBrowser.cpp b/Engine/Source/Editor/ContentBrowser/Private/SContentBrowser.cpp
--- a/Engine/Source/Editor/ContentBrowser/Private/SContentBrowser.cpp
+++ b/Engine/Source/Editor/ContentBrowser/Private/SContentBrowser.cpp
@@ -158,7 +158,7 @@ void SContentBrowser::Construct( const FArguments& InArgs, const FName& InInstan
 			.InnerSlotPadding( FVector2D( 5, 2 ) )
 
 			+ SWrapBox::Slot()
-			.FillLineWhenWidthLessThan( 600 )
+			.FillLineWhenSizeLessThan( 600 )
 			.FillEmptySpace( true )
 			[
 				SNew( SHorizontalBox )
diff --git a/Engine/Source/Editor/Sequencer/Private/SSequencer.cpp b/Engine/Source/Editor/Sequencer/Private/SSequencer.cpp
--- a/Engine/Source/Editor/Sequencer/Private/SSequencer.cpp
+++ b/Engine/Source/Editor/Sequencer/Private/SSequencer.cpp
@@ -581,7 +581,7 @@ void SSequencer::Construct(const FArguments& InArgs, TSharedRef<FSequencer> InSe
 
 							+ SWrapBox::Slot()
 							.FillEmptySpace(true)
-							.FillLineWhenWidthLessThan(600)
+							.FillLineWhenSizeLessThan(600)
 							[
 								SAssignNew(ToolbarContainer, SBox)
 							]
diff --git a/Engine/Source/Runtime/Slate/Private/Widgets/Layout/SWrapBox.cpp b/Engine/Source/Runtime/Slate/Private/Widgets/Layout/SWrapBox.cpp
--- a/Engine/Source/Runtime/Slate/Private/Widgets/Layout/SWrapBox.cpp
+++ b/Engine/Source/Runtime/Slate/Private/Widgets/Layout/SWrapBox.cpp
@@ -169,9 +169,12 @@ void SWrapBox::FChildArranger::Arrange()
 		{
 			const float BottomBoundOfChild = ArrangementData.SlotOffset.Y + ArrangementData.SlotSize.Y;
 
+			PRAGMA_DISABLE_DEPRECATION_WARNINGS
 			// Rule: If required due to a wrapping height under specified threshold, start a new line and allocate all of it to this child.
-			if (Slot.SlotFillLineWhenSizeLessThan.IsSet() && WrapBox.PreferredSize.Get() < Slot.SlotFillLineWhenSizeLessThan.GetValue())
+			if (Slot.SlotFillLineWhenSizeLessThan.IsSet() && WrapBox.PreferredSize.Get() < Slot.SlotFillLineWhenSizeLessThan.GetValue()
+				|| (!Slot.SlotFillLineWhenSizeLessThan.IsSet() && Slot.SlotFillLineWhenWidthLessThan.IsSet() && WrapBox.PreferredSize.Get() < Slot.SlotFillLineWhenWidthLessThan.GetValue()))
 			{
+			PRAGMA_ENABLE_DEPRECATION_WARNINGS
 				// Begin a new line if the current one isn't empty, because we demand a whole line to ourselves.
 				if (!IsFirstChildInCurrentLine())
 				{
@@ -201,9 +204,12 @@ void SWrapBox::FChildArranger::Arrange()
 		{
 			const float RightBoundOfChild = ArrangementData.SlotOffset.X + ArrangementData.SlotSize.X;
 
+			PRAGMA_DISABLE_DEPRECATION_WARNINGS
 			// Rule: If required due to a wrapping width under specified threshold, start a new line and allocate all of it to this child.
-			if (Slot.SlotFillLineWhenSizeLessThan.IsSet() && WrapBox.PreferredSize.Get() < Slot.SlotFillLineWhenSizeLessThan.GetValue())
+			if (Slot.SlotFillLineWhenSizeLessThan.IsSet() && WrapBox.PreferredSize.Get() < Slot.SlotFillLineWhenSizeLessThan.GetValue()
+				|| (!Slot.SlotFillLineWhenSizeLessThan.IsSet() && Slot.SlotFillLineWhenWidthLessThan.IsSet() && WrapBox.PreferredSize.Get() < Slot.SlotFillLineWhenWidthLessThan.GetValue()))
 			{
+			PRAGMA_ENABLE_DEPRECATION_WARNINGS
 				// Begin a new line if the current one isn't empty, because we demand a whole line to ourselves.
 				if (!IsFirstChildInCurrentLine())
 				{
diff --git a/Engine/Source/Runtime/Slate/Public/Widgets/Layout/SWrapBox.h b/Engine/Source/Runtime/Slate/Public/Widgets/Layout/SWrapBox.h
--- a/Engine/Source/Runtime/Slate/Public/Widgets/Layout/SWrapBox.h
+++ b/Engine/Source/Runtime/Slate/Public/Widgets/Layout/SWrapBox.h
@@ -58,13 +58,23 @@ public:
 		FSlot()
 			: TSlotBase<FSlot>()
 			, TSupportsContentAlignmentMixin<FSlot>(HAlign_Fill, VAlign_Fill)
+			, SlotFillLineWhenWidthLessThan()
 			, SlotFillLineWhenSizeLessThan()
 			, bSlotFillEmptySpace(false)
 		{
 		}
 
-		/** If the total available space in the wrap panel drops below this threshold, this slot will attempt to fill an entire line. */
-		FSlot& FillLineWhenWidthLessThan(TOptional<float> InFillLineWhenSizeLessThan)
+		UE_DEPRECATED(4.26, "Deprecated, please use FillLineWhenSizeLessThan() instead")
+		FSlot& FillLineWhenWidthLessThan(TOptional<float> InFillLineWhenWidthLessThan)
+		{
+			PRAGMA_DISABLE_DEPRECATION_WARNINGS
+			SlotFillLineWhenWidthLessThan = InFillLineWhenWidthLessThan;
+			PRAGMA_ENABLE_DEPRECATION_WARNINGS
+			return *(static_cast<FSlot*>(this));
+		}
+
+		/** Dependently of the Orientation, if the total available horizontal or vertical space in the wrap panel drops below this threshold, this slot will attempt to fill an entire line. */
+		FSlot& FillLineWhenSizeLessThan(TOptional<float> InFillLineWhenSizeLessThan)
 		{
 			SlotFillLineWhenSizeLessThan = InFillLineWhenSizeLessThan;
 			return *(static_cast<FSlot*>(this));
@@ -77,6 +87,9 @@ public:
 			return *(static_cast<FSlot*>(this));
 		}
 
+		UE_DEPRECATED(4.26, "Deprecated, please use SlotFillLineWhenSizeLessThan instead")
+		TOptional<float> SlotFillLineWhenWidthLessThan;
+
 		TOptional<float> SlotFillLineWhenSizeLessThan;
 		bool bSlotFillEmptySpace;
 	};
diff --git a/Engine/Source/Runtime/UMG/Private/Components/WrapBoxSlot.cpp b/Engine/Source/Runtime/UMG/Private/Components/WrapBoxSlot.cpp
--- a/Engine/Source/Runtime/UMG/Private/Components/WrapBoxSlot.cpp
+++ b/Engine/Source/Runtime/UMG/Private/Components/WrapBoxSlot.cpp
@@ -31,7 +31,7 @@ void UWrapBoxSlot::BuildSlot(TSharedRef<SWrapBox> WrapBox)
 		.HAlign(HorizontalAlignment)
 		.VAlign(VerticalAlignment)
 		.FillEmptySpace(bFillEmptySpace)
-		.FillLineWhenWidthLessThan(FillSpanWhenLessThan == 0 ? TOptional<float>() : TOptional<float>(FillSpanWhenLessThan))
+		.FillLineWhenSizeLessThan(FillSpanWhenLessThan == 0 ? TOptional<float>() : TOptional<float>(FillSpanWhenLessThan))
 		[
 			Content == NULL ? SNullWidget::NullWidget : Content->TakeWidget()
 		];
@@ -60,7 +60,7 @@ void UWrapBoxSlot::SetFillSpanWhenLessThan(float InFillSpanWhenLessThan)
 	FillSpanWhenLessThan = InFillSpanWhenLessThan;
 	if ( Slot )
 	{
-		Slot->FillLineWhenWidthLessThan(InFillSpanWhenLessThan == 0 ? TOptional<float>() : TOptional<float>(InFillSpanWhenLessThan));
+		Slot->FillLineWhenSizeLessThan(InFillSpanWhenLessThan == 0 ? TOptional<float>() : TOptional<float>(InFillSpanWhenLessThan));
 	}
 }
 
