Fix CIS errors with a && expression
#RNX

diff --git a/Engine/Source/Runtime/Slate/Private/Widgets/Layout/SWrapBox.cpp b/Engine/Source/Runtime/Slate/Private/Widgets/Layout/SWrapBox.cpp
--- a/Engine/Source/Runtime/Slate/Private/Widgets/Layout/SWrapBox.cpp
+++ b/Engine/Source/Runtime/Slate/Private/Widgets/Layout/SWrapBox.cpp
@@ -171,7 +171,7 @@ void SWrapBox::FChildArranger::Arrange()
 
 			PRAGMA_DISABLE_DEPRECATION_WARNINGS
 			// Rule: If required due to a wrapping height under specified threshold, start a new line and allocate all of it to this child.
-			if (Slot.SlotFillLineWhenSizeLessThan.IsSet() && WrapBox.PreferredSize.Get() < Slot.SlotFillLineWhenSizeLessThan.GetValue()
+			if ((Slot.SlotFillLineWhenSizeLessThan.IsSet() && WrapBox.PreferredSize.Get() < Slot.SlotFillLineWhenSizeLessThan.GetValue())
 				|| (!Slot.SlotFillLineWhenSizeLessThan.IsSet() && Slot.SlotFillLineWhenWidthLessThan.IsSet() && WrapBox.PreferredSize.Get() < Slot.SlotFillLineWhenWidthLessThan.GetValue()))
 			{
 			PRAGMA_ENABLE_DEPRECATION_WARNINGS
@@ -206,7 +206,7 @@ void SWrapBox::FChildArranger::Arrange()
 
 			PRAGMA_DISABLE_DEPRECATION_WARNINGS
 			// Rule: If required due to a wrapping width under specified threshold, start a new line and allocate all of it to this child.
-			if (Slot.SlotFillLineWhenSizeLessThan.IsSet() && WrapBox.PreferredSize.Get() < Slot.SlotFillLineWhenSizeLessThan.GetValue()
+			if ((Slot.SlotFillLineWhenSizeLessThan.IsSet() && WrapBox.PreferredSize.Get() < Slot.SlotFillLineWhenSizeLessThan.GetValue())
 				|| (!Slot.SlotFillLineWhenSizeLessThan.IsSet() && Slot.SlotFillLineWhenWidthLessThan.IsSet() && WrapBox.PreferredSize.Get() < Slot.SlotFillLineWhenWidthLessThan.GetValue()))
 			{
 			PRAGMA_ENABLE_DEPRECATION_WARNINGS
