Added an option to TextBlock/RichTextBlock to control the lexical transformation of the display text

This can be used to make the display text appear in all lower or upper case, and is preferred when using rich-text as the transformation can affect the format markup and break the parser

diff --git a/Engine/Source/Runtime/Slate/Private/Framework/Text/TextLayout.cpp b/Engine/Source/Runtime/Slate/Private/Framework/Text/TextLayout.cpp
--- a/Engine/Source/Runtime/Slate/Private/Framework/Text/TextLayout.cpp
+++ b/Engine/Source/Runtime/Slate/Private/Framework/Text/TextLayout.cpp
@@ -7,6 +7,7 @@
 #include "Framework/Text/TextHitPoint.h"
 #include "Framework/Text/ILayoutBlock.h"
 #include "Internationalization/BreakIterator.h"
+#include "Internationalization/TextTransformer.h"
 #include "Framework/Text/ShapedTextCache.h"
 
 DECLARE_CYCLE_STAT(TEXT("Text Layout"), STAT_SlateTextLayout, STATGROUP_Slate);
@@ -986,6 +987,29 @@ void FTextLayout::ClearView()
 	LineViewsToJustify.Empty();
 }
 
+void FTextLayout::TransformLineText(FLineModel& LineModel) const
+{
+	auto ApplyTransformedText = [&LineModel](FString&& InNewText)
+	{
+		if (ensure(LineModel.Text->Len() == InNewText.Len()))
+		{
+			*LineModel.Text = MoveTemp(InNewText);
+		}
+	};
+
+	switch (TransformPolicy)
+	{
+	case ETextTransformPolicy::ToLower:
+		ApplyTransformedText(FTextTransformer::ToLower(*LineModel.Text));
+		break;
+	case ETextTransformPolicy::ToUpper:
+		ApplyTransformedText(FTextTransformer::ToUpper(*LineModel.Text));
+		break;
+	default:
+		break;
+	}
+}
+
 void FTextLayout::CalculateTextDirection()
 {
 	for (FLineModel& LineModel : LineModels)
@@ -1125,6 +1149,7 @@ FTextLayout::FTextLayout()
 	, Scale( 1.0f )
 	, WrappingWidth( 0 )
 	, WrappingPolicy( ETextWrappingPolicy::DefaultWrapping )
+	, TransformPolicy( ETextTransformPolicy::None )
 	, Margin()
 	, Justification( ETextJustify::Left )
 	, LineHeightPercentage( 1.0f )
@@ -2074,6 +2099,7 @@ void FTextLayout::AddLine( const FNewLineData& NewLine )
 {
 	{
 		FLineModel LineModel( NewLine.Text );
+		TransformLineText(LineModel);
 
 		for (const auto& Run : NewLine.Runs)
 		{
@@ -2132,6 +2158,7 @@ void FTextLayout::AddLines( const TArray<FNewLineData>& NewLines )
 	for (const auto& NewLine : NewLines)
 	{
 		FLineModel LineModel( NewLine.Text );
+		TransformLineText(LineModel);
 
 		for (const auto& Run : NewLine.Runs)
 		{
@@ -2580,6 +2607,15 @@ void FTextLayout::SetWrappingPolicy(ETextWrappingPolicy Value)
 	}
 }
 
+void FTextLayout::SetTransformPolicy(ETextTransformPolicy Value)
+{
+	if (TransformPolicy != Value)
+	{
+		TransformPolicy = Value;
+		DirtyFlags |= ETextLayoutDirtyState::Layout;
+	}
+}
+
 void FTextLayout::SetDebugSourceInfo(const TAttribute<FString>& InDebugSourceInfo)
 {
 	DebugSourceInfo = InDebugSourceInfo;
diff --git a/Engine/Source/Runtime/Slate/Private/Widgets/Text/SRichTextBlock.cpp b/Engine/Source/Runtime/Slate/Private/Widgets/Text/SRichTextBlock.cpp
--- a/Engine/Source/Runtime/Slate/Private/Widgets/Text/SRichTextBlock.cpp
+++ b/Engine/Source/Runtime/Slate/Private/Widgets/Text/SRichTextBlock.cpp
@@ -27,6 +27,7 @@ void SRichTextBlock::Construct( const FArguments& InArgs )
 	WrapTextAt = InArgs._WrapTextAt;
 	AutoWrapText = InArgs._AutoWrapText;
 	WrappingPolicy = InArgs._WrappingPolicy;
+	TransformPolicy = InArgs._TransformPolicy;
 	Margin = InArgs._Margin;
 	LineHeightPercentage = InArgs._LineHeightPercentage;
 	Justification = InArgs._Justification;
@@ -85,7 +86,7 @@ FVector2D SRichTextBlock::ComputeDesiredSize(float LayoutScaleMultiplier) const
 {
 	// ComputeDesiredSize will also update the text layout cache if required
 	const FVector2D TextSize = TextLayoutCache->ComputeDesiredSize(
-		FSlateTextBlockLayout::FWidgetArgs(BoundText, HighlightText, WrapTextAt, AutoWrapText, WrappingPolicy, Margin, LineHeightPercentage, Justification),
+		FSlateTextBlockLayout::FWidgetArgs(BoundText, HighlightText, WrapTextAt, AutoWrapText, WrappingPolicy, TransformPolicy, Margin, LineHeightPercentage, Justification),
 		LayoutScaleMultiplier * TextBlockScale, TextStyle) * TextBlockScale;
 
 	return FVector2D(FMath::Max(TextSize.X, MinDesiredWidth.Get()), TextSize.Y);
@@ -144,6 +145,11 @@ void SRichTextBlock::SetWrappingPolicy(const TAttribute<ETextWrappingPolicy>& In
 	SetAttribute(WrappingPolicy, InWrappingPolicy, EInvalidateWidgetReason::Layout);
 }
 
+void SRichTextBlock::SetTransformPolicy(const TAttribute<ETextTransformPolicy>& InTransformPolicy)
+{
+	SetAttribute(TransformPolicy, InTransformPolicy, EInvalidateWidgetReason::Layout);
+}
+
 void SRichTextBlock::SetLineHeightPercentage(const TAttribute<float>& InLineHeightPercentage)
 {
 	SetAttribute(LineHeightPercentage, InLineHeightPercentage, EInvalidateWidgetReason::Layout);
@@ -200,6 +206,7 @@ bool SRichTextBlock::ComputeVolatility() const
 		|| WrapTextAt.IsBound()
 		|| AutoWrapText.IsBound()
 		|| WrappingPolicy.IsBound()
+		|| TransformPolicy.IsBound()
 		|| Margin.IsBound()
 		|| Justification.IsBound()
 		|| LineHeightPercentage.IsBound()
diff --git a/Engine/Source/Runtime/Slate/Private/Widgets/Text/STextBlock.cpp b/Engine/Source/Runtime/Slate/Private/Widgets/Text/STextBlock.cpp
--- a/Engine/Source/Runtime/Slate/Private/Widgets/Text/STextBlock.cpp
+++ b/Engine/Source/Runtime/Slate/Private/Widgets/Text/STextBlock.cpp
@@ -42,6 +42,7 @@ void STextBlock::Construct( const FArguments& InArgs )
 	WrapTextAt = InArgs._WrapTextAt;
 	AutoWrapText = InArgs._AutoWrapText;
 	WrappingPolicy = InArgs._WrappingPolicy;
+	TransformPolicy = InArgs._TransformPolicy;
 	Margin = InArgs._Margin;
 	LineHeightPercentage = InArgs._LineHeightPercentage;
 	Justification = InArgs._Justification;
@@ -275,7 +276,7 @@ FVector2D STextBlock::ComputeDesiredSize(float LayoutScaleMultiplier) const
 	{
 		// ComputeDesiredSize will also update the text layout cache if required
 		const FVector2D TextSize = TextLayoutCache->ComputeDesiredSize(
-			FSlateTextBlockLayout::FWidgetArgs(BoundText, HighlightText, WrapTextAt, AutoWrapText, WrappingPolicy, Margin, LineHeightPercentage, Justification),
+			FSlateTextBlockLayout::FWidgetArgs(BoundText, HighlightText, WrapTextAt, AutoWrapText, WrappingPolicy, TransformPolicy, Margin, LineHeightPercentage, Justification),
 			LayoutScaleMultiplier, GetComputedTextStyle()
 		);
 
@@ -298,6 +299,7 @@ bool STextBlock::ComputeVolatility() const
 		|| WrapTextAt.IsBound()
 		|| AutoWrapText.IsBound()
 		|| WrappingPolicy.IsBound()
+		|| TransformPolicy.IsBound()
 		|| Margin.IsBound()
 		|| Justification.IsBound()
 		|| LineHeightPercentage.IsBound()
@@ -392,6 +394,15 @@ void STextBlock::SetWrappingPolicy(const TAttribute<ETextWrappingPolicy>& InWrap
 	}
 }
 
+void STextBlock::SetTransformPolicy(const TAttribute<ETextTransformPolicy>& InTransformPolicy)
+{
+	if(!TransformPolicy.IdenticalTo(InTransformPolicy))
+	{
+		TransformPolicy = InTransformPolicy;
+		InvalidateText(EInvalidateWidget::LayoutAndVolatility);
+	}
+}
+
 void STextBlock::SetShadowOffset(const TAttribute<FVector2D>& InShadowOffset)
 {
 	if(!ShadowOffset.IdenticalTo(InShadowOffset))
diff --git a/Engine/Source/Runtime/Slate/Private/Widgets/Text/SlateEditableTextLayout.cpp b/Engine/Source/Runtime/Slate/Private/Widgets/Text/SlateEditableTextLayout.cpp
--- a/Engine/Source/Runtime/Slate/Private/Widgets/Text/SlateEditableTextLayout.cpp
+++ b/Engine/Source/Runtime/Slate/Private/Widgets/Text/SlateEditableTextLayout.cpp
@@ -3426,7 +3426,7 @@ FVector2D FSlateEditableTextLayout::ComputeDesiredSize(float LayoutScaleMultipli
 		MarginValue.Right += CaretWidth;
 
 		const FVector2D HintTextSize = HintTextLayout->ComputeDesiredSize(
-			FSlateTextBlockLayout::FWidgetArgs(HintText, FText::GetEmpty(), WrapTextAt, AutoWrapText, WrappingPolicy, MarginValue, LineHeightPercentage, Justification),
+			FSlateTextBlockLayout::FWidgetArgs(HintText, FText::GetEmpty(), WrapTextAt, AutoWrapText, WrappingPolicy, ETextTransformPolicy::None, MarginValue, LineHeightPercentage, Justification),
 			LayoutScaleMultiplier, HintTextStyle
 			);
 
diff --git a/Engine/Source/Runtime/Slate/Private/Widgets/Text/SlateTextBlockLayout.cpp b/Engine/Source/Runtime/Slate/Private/Widgets/Text/SlateTextBlockLayout.cpp
--- a/Engine/Source/Runtime/Slate/Private/Widgets/Text/SlateTextBlockLayout.cpp
+++ b/Engine/Source/Runtime/Slate/Private/Widgets/Text/SlateTextBlockLayout.cpp
@@ -31,14 +31,23 @@ FVector2D FSlateTextBlockLayout::ComputeDesiredSize(const FWidgetArgs& InWidgetA
 	CachedWrapTextAt = InWidgetArgs.WrapTextAt.Get(0.0f);
 	bCachedAutoWrapText = InWidgetArgs.AutoWrapText.Get(false);
 
+	const ETextTransformPolicy PreviousTransformPolicy = TextLayout->GetTransformPolicy();
+
 	// Set the text layout information
 	TextLayout->SetScale(InScale);
 	TextLayout->SetWrappingWidth(CalculateWrappingWidth());
 	TextLayout->SetWrappingPolicy(InWidgetArgs.WrappingPolicy.Get());
+	TextLayout->SetTransformPolicy(InWidgetArgs.TransformPolicy.Get());
 	TextLayout->SetMargin(InWidgetArgs.Margin.Get());
 	TextLayout->SetJustification(InWidgetArgs.Justification.Get());
 	TextLayout->SetLineHeightPercentage(InWidgetArgs.LineHeightPercentage.Get());
 
+	// Has the transform policy changed? If so we need a full refresh as that is destructive to the model text
+	if (PreviousTransformPolicy != TextLayout->GetTransformPolicy())
+	{
+		Marshaller->MakeDirty();
+	}
+
 	// Has the style used for this text block changed?
 	if(!IsStyleUpToDate(InTextStyle))
 	{
diff --git a/Engine/Source/Runtime/Slate/Public/Framework/Text/TextLayout.h b/Engine/Source/Runtime/Slate/Public/Framework/Text/TextLayout.h
--- a/Engine/Source/Runtime/Slate/Public/Framework/Text/TextLayout.h
+++ b/Engine/Source/Runtime/Slate/Public/Framework/Text/TextLayout.h
@@ -46,6 +46,22 @@ namespace ETextJustify
 	};
 }
 
+/**
+ * The different methods that can be used if a word is too long to be broken by the default line-break iterator.
+ */
+UENUM(BlueprintType)
+enum class ETextTransformPolicy : uint8
+{
+	/** No transform, just use the given text as-is */
+	None = 0,
+
+	/** Convert the text to lowercase for display */
+	ToLower,
+
+	/** Convert the text to uppercase for display */
+	ToUpper,
+};
+
 /** 
  * The different methods that can be used if a word is too long to be broken by the default line-break iterator.
  */
@@ -395,6 +411,10 @@ public:
 	/** Get the visual justification for this document (based on the visual justification used by the first line of text) */
 	ETextJustify::Type GetVisualJustification() const;
 
+	/** @note This option is destructive to the model text, so changing it requires refreshing the text layout from its marshaller */
+	FORCEINLINE ETextTransformPolicy GetTransformPolicy() const { return TransformPolicy; }
+	void SetTransformPolicy(ETextTransformPolicy Value);
+
 	FORCEINLINE float GetScale() const { return Scale; }
 	void SetScale( float Value );
 
@@ -585,6 +605,12 @@ protected:
 	*/
 	void ClearView();
 
+	/**
+	 * Transform the given line model text based on the active transform policy.
+	 * @note This is destructive to the model text!
+	 */
+	void TransformLineText(FLineModel& LineModel) const;
+
 	/**
 	* Notifies all Runs that we are beginning to generate a new layout.
 	*/
@@ -699,6 +725,9 @@ protected:
 	/** The wrapping policy used by this text layout. */
 	ETextWrappingPolicy WrappingPolicy;
 
+	/** The transform policy used by this text layout. */
+	ETextTransformPolicy TransformPolicy;
+
 	/** The size of the margins to put about the text. This is an unscaled value. */
 	FMargin Margin;
 
diff --git a/Engine/Source/Runtime/Slate/Public/Widgets/Text/SRichTextBlock.h b/Engine/Source/Runtime/Slate/Public/Widgets/Text/SRichTextBlock.h
--- a/Engine/Source/Runtime/Slate/Public/Widgets/Text/SRichTextBlock.h
+++ b/Engine/Source/Runtime/Slate/Public/Widgets/Text/SRichTextBlock.h
@@ -40,6 +40,7 @@ public:
 		, _WrapTextAt( 0.0f )
 		, _AutoWrapText(false)
 		, _WrappingPolicy(ETextWrappingPolicy::DefaultWrapping)
+		, _TransformPolicy(ETextTransformPolicy::None)
 		, _Marshaller()
 		, _DecoratorStyleSet( &FCoreStyle::Get() )
 		, _TextStyle( &FCoreStyle::Get().GetWidgetStyle<FTextBlockStyle>( "NormalText" ) )
@@ -71,6 +72,9 @@ public:
 		/** The wrapping policy to use */
 		SLATE_ATTRIBUTE( ETextWrappingPolicy, WrappingPolicy )
 
+		/** The transform policy to use */
+		SLATE_ATTRIBUTE( ETextTransformPolicy, TransformPolicy )
+
 		/** The marshaller used to get/set the raw text to/from the text layout. */
 		SLATE_ARGUMENT(TSharedPtr<class FRichTextLayoutMarshaller>, Marshaller)
 
@@ -193,6 +197,9 @@ public:
 	/** Set WrappingPolicy attribute */
 	void SetWrappingPolicy(const TAttribute<ETextWrappingPolicy>& InWrappingPolicy);
 
+	/** Set TransformPolicy attribute */
+	void SetTransformPolicy(const TAttribute<ETextTransformPolicy>& InTransformPolicy);
+
 	/** See LineHeightPercentage attribute */
 	void SetLineHeightPercentage(const TAttribute<float>& InLineHeightPercentage);
 
@@ -248,6 +255,9 @@ private:
 	/** The wrapping policy we're using */
 	TAttribute<ETextWrappingPolicy> WrappingPolicy;
 
+	/** The transform policy we're using */
+	TAttribute<ETextTransformPolicy> TransformPolicy;
+
 	/** The amount of blank space left around the edges of text area. */
 	TAttribute< FMargin > Margin;
 
diff --git a/Engine/Source/Runtime/Slate/Public/Widgets/Text/STextBlock.h b/Engine/Source/Runtime/Slate/Public/Widgets/Text/STextBlock.h
--- a/Engine/Source/Runtime/Slate/Public/Widgets/Text/STextBlock.h
+++ b/Engine/Source/Runtime/Slate/Public/Widgets/Text/STextBlock.h
@@ -55,6 +55,7 @@ public:
 		, _WrapTextAt(0.0f)
 		, _AutoWrapText(false)
 		, _WrappingPolicy(ETextWrappingPolicy::DefaultWrapping)
+		, _TransformPolicy(ETextTransformPolicy::None)
 		, _Margin()
 		, _LineHeightPercentage(1.0f)
 		, _Justification(ETextJustify::Left)
@@ -108,6 +109,9 @@ public:
 		/** The wrapping policy to use */
 		SLATE_ATTRIBUTE( ETextWrappingPolicy, WrappingPolicy )
 
+		/** The transform policy to use */
+		SLATE_ATTRIBUTE( ETextTransformPolicy, TransformPolicy )
+
 		/** The amount of blank space left around the edges of text area. */
 		SLATE_ATTRIBUTE( FMargin, Margin )
 
@@ -216,6 +220,9 @@ public:
 	/** Set WrappingPolicy attribute */
 	void SetWrappingPolicy(const TAttribute<ETextWrappingPolicy>& InWrappingPolicy);
 
+	/** Set TransformPolicy attribute */
+	void SetTransformPolicy(const TAttribute<ETextTransformPolicy>& InTransformPolicy);
+
 	/** See ShadowOffset attribute */
 	void SetShadowOffset(const TAttribute<FVector2D>& InShadowOffset);
 
@@ -317,6 +324,9 @@ private:
 	/** The wrapping policy we're using */
 	TAttribute<ETextWrappingPolicy> WrappingPolicy;
 
+	/** The transform policy we're using */
+	TAttribute<ETextTransformPolicy> TransformPolicy;
+
 	/** The amount of blank space left around the edges of text area. */
 	TAttribute< FMargin > Margin;
 
diff --git a/Engine/Source/Runtime/Slate/Public/Widgets/Text/SlateTextBlockLayout.h b/Engine/Source/Runtime/Slate/Public/Widgets/Text/SlateTextBlockLayout.h
--- a/Engine/Source/Runtime/Slate/Public/Widgets/Text/SlateTextBlockLayout.h
+++ b/Engine/Source/Runtime/Slate/Public/Widgets/Text/SlateTextBlockLayout.h
@@ -28,6 +28,7 @@ public:
 			const TAttribute<float>& InWrapTextAt, 
 			const TAttribute<bool>& InAutoWrapText, 
 			const TAttribute<ETextWrappingPolicy>& InWrappingPolicy, 
+			const TAttribute<ETextTransformPolicy>& InTransformPolicy,
 			const TAttribute<FMargin>& InMargin, 
 			const TAttribute<float>& InLineHeightPercentage, 
 			const TAttribute<ETextJustify::Type>& InJustification
@@ -37,6 +38,7 @@ public:
 			, WrapTextAt(InWrapTextAt)
 			, AutoWrapText(InAutoWrapText)
 			, WrappingPolicy(InWrappingPolicy)
+			, TransformPolicy(InTransformPolicy)
 			, Margin(InMargin)
 			, LineHeightPercentage(InLineHeightPercentage)
 			, Justification(InJustification)
@@ -48,6 +50,7 @@ public:
 		const TAttribute<float>& WrapTextAt;
 		const TAttribute<bool>& AutoWrapText;
 		const TAttribute<ETextWrappingPolicy> WrappingPolicy;
+		const TAttribute<ETextTransformPolicy> TransformPolicy;
 		const TAttribute<FMargin>& Margin;
 		const TAttribute<float>& LineHeightPercentage;
 		const TAttribute<ETextJustify::Type>& Justification;
diff --git a/Engine/Source/Runtime/UMG/Private/Components/RichTextBlock.cpp b/Engine/Source/Runtime/UMG/Private/Components/RichTextBlock.cpp
--- a/Engine/Source/Runtime/UMG/Private/Components/RichTextBlock.cpp
+++ b/Engine/Source/Runtime/UMG/Private/Components/RichTextBlock.cpp
@@ -47,6 +47,7 @@ URichTextBlock::URichTextBlock(const FObjectInitializer& ObjectInitializer)
 	: Super(ObjectInitializer)
 {
 	Visibility = ESlateVisibility::SelfHitTestInvisible;
+	TextTransformPolicy = ETextTransformPolicy::None;
 }
 
 void URichTextBlock::ReleaseSlateResources(bool bReleaseChildren)
@@ -79,6 +80,7 @@ void URichTextBlock::SynchronizeProperties()
 	Super::SynchronizeProperties();
 
 	MyRichTextBlock->SetText(Text);
+	MyRichTextBlock->SetTransformPolicy(TextTransformPolicy);
 	MyRichTextBlock->SetMinDesiredWidth(MinDesiredWidth);
 
 	Super::SynchronizeTextLayoutProperties( *MyRichTextBlock );
@@ -335,6 +337,15 @@ void URichTextBlock::SetAutoWrapText(bool InAutoTextWrap)
 	}
 }
 
+void URichTextBlock::SetTextTransformPolicy(ETextTransformPolicy InTransformPolicy)
+{
+	TextTransformPolicy = InTransformPolicy;
+	if (MyRichTextBlock.IsValid())
+	{
+		MyRichTextBlock->SetTransformPolicy(TextTransformPolicy);
+	}
+}
+
 void URichTextBlock::ApplyUpdatedDefaultTextStyle()
 {
 	if (MyRichTextBlock.IsValid())
diff --git a/Engine/Source/Runtime/UMG/Private/Components/TextBlock.cpp b/Engine/Source/Runtime/UMG/Private/Components/TextBlock.cpp
--- a/Engine/Source/Runtime/UMG/Private/Components/TextBlock.cpp
+++ b/Engine/Source/Runtime/UMG/Private/Components/TextBlock.cpp
@@ -24,6 +24,7 @@ UTextBlock::UTextBlock(const FObjectInitializer& ObjectInitializer)
 	ColorAndOpacity = FLinearColor::White;
 	ShadowColorAndOpacity = FLinearColor::Transparent;
 	bAutoWrapText_DEPRECATED = false;
+	TextTransformPolicy = ETextTransformPolicy::None;
 
 	if (!IsRunningDedicatedServer())
 	{
@@ -135,6 +136,15 @@ void UTextBlock::SetAutoWrapText(bool InAutoWrapText)
 	}
 }
 
+void UTextBlock::SetTextTransformPolicy(ETextTransformPolicy InTransformPolicy)
+{
+	TextTransformPolicy = InTransformPolicy;
+	if(MyTextBlock.IsValid())
+	{
+		MyTextBlock->SetTransformPolicy(TextTransformPolicy);
+	}
+}
+
 UMaterialInstanceDynamic* UTextBlock::GetDynamicFontMaterial()
 {
 	if (Font.FontMaterial)
@@ -277,6 +287,7 @@ void UTextBlock::SynchronizeProperties()
 		MyTextBlock->SetShadowOffset( ShadowOffset );
 		MyTextBlock->SetShadowColorAndOpacity( ShadowColorAndOpacityBinding );
 		MyTextBlock->SetMinDesiredWidth( MinDesiredWidth );
+		MyTextBlock->SetTransformPolicy( TextTransformPolicy );
 		Super::SynchronizeTextLayoutProperties( *MyTextBlock );
 	}
 }
diff --git a/Engine/Source/Runtime/UMG/Public/Components/RichTextBlock.h b/Engine/Source/Runtime/UMG/Public/Components/RichTextBlock.h
--- a/Engine/Source/Runtime/UMG/Public/Components/RichTextBlock.h
+++ b/Engine/Source/Runtime/UMG/Public/Components/RichTextBlock.h
@@ -88,6 +88,13 @@ public:
 	UFUNCTION(BlueprintCallable, Category = "Appearance")
 	void SetAutoWrapText(bool InAutoTextWrap);
 
+	/**
+	 * Set the text transformation policy for this text block.
+	 * @param InTransformPolicy the new text transformation policy.
+	 */
+	UFUNCTION(BlueprintCallable, Category = "Appearance")
+	void SetTextTransformPolicy(ETextTransformPolicy InTransformPolicy);
+
 	/** 
 	 * Wholesale override of the currently established default text style
 	 * @param InDefaultTextStyle The new text style to apply to all default (i.e. undecorated) text in the block
@@ -183,6 +190,10 @@ protected:
 	UPROPERTY(EditAnywhere, BlueprintReadOnly, Category=Appearance)
 	float MinDesiredWidth;
 
+	/** The text transformation policy to apply to this text block */
+	UPROPERTY(EditAnywhere, BlueprintReadOnly, Category=Appearance, meta=(DisplayName="Transform Policy"))
+	ETextTransformPolicy TextTransformPolicy;
+
 	FTextBlockStyle DefaultTextStyle;
 
 	UPROPERTY(Transient)
diff --git a/Engine/Source/Runtime/UMG/Public/Components/TextBlock.h b/Engine/Source/Runtime/UMG/Public/Components/TextBlock.h
--- a/Engine/Source/Runtime/UMG/Public/Components/TextBlock.h
+++ b/Engine/Source/Runtime/UMG/Public/Components/TextBlock.h
@@ -89,6 +89,14 @@ public:
 	UFUNCTION(BlueprintCallable, Category = "Appearance")
 	void SetAutoWrapText(bool InAutoTextWrap);
 
+	/**
+	 * Set the text transformation policy for this text block.
+	 *
+	 * @param InTransformPolicy the new text transformation policy.
+	 */
+	UFUNCTION(BlueprintCallable, Category = "Appearance")
+	void SetTextTransformPolicy(ETextTransformPolicy InTransformPolicy);
+
 	/**  */
 	UFUNCTION(BlueprintCallable, Category="Appearance")
 	UMaterialInstanceDynamic* GetDynamicFontMaterial();
@@ -146,6 +154,10 @@ public:
 	UPROPERTY(meta = (DeprecatedProperty, DeprecationMessage = "bAutoWrapText is deprecated. Please use AutoWrapText instead."))
 	bool bAutoWrapText_DEPRECATED;
 
+	/** The text transformation policy to apply to this text block */
+	UPROPERTY(EditAnywhere, BlueprintReadOnly, Category=Appearance, meta=(DisplayName="Transform Policy"))
+	ETextTransformPolicy TextTransformPolicy;
+
 	/** 
 	 * Gets the widget text
 	 * @return The widget text
