Adding Radial Mode to DynamicEntryBox. Adding a SRadialBox container to support radial arrangement of items.

diff --git a/Engine/Source/Runtime/Slate/Private/Widgets/Layout/SRadialBox.cpp b/Engine/Source/Runtime/Slate/Private/Widgets/Layout/SRadialBox.cpp
new file mode 100644
--- /dev/null
+++ b/Engine/Source/Runtime/Slate/Private/Widgets/Layout/SRadialBox.cpp
@@ -0,0 +1,168 @@
+// Copyright Epic Games, Inc. All Rights Reserved.
+
+#include "Widgets/Layout/SRadialBox.h"
+#include "Layout/LayoutUtils.h"
+
+SRadialBox::SRadialBox()
+: Slots(this)
+{
+}
+
+SRadialBox::FSlot& SRadialBox::Slot()
+{
+	return *( new SRadialBox::FSlot() );
+}
+
+SRadialBox::FSlot& SRadialBox::AddSlot()
+{
+	SRadialBox::FSlot* NewSlot = new SRadialBox::FSlot();
+	Slots.Add(NewSlot);
+	return *NewSlot;
+}
+
+int32 SRadialBox::RemoveSlot( const TSharedRef<SWidget>& SlotWidget )
+{
+	for (int32 SlotIdx = 0; SlotIdx < Slots.Num(); ++SlotIdx)
+	{
+		if ( SlotWidget == Slots[SlotIdx].GetWidget() )
+		{
+			Slots.RemoveAt(SlotIdx);
+			return SlotIdx;
+		}
+	}
+
+	return -1;
+}
+
+void SRadialBox::Construct( const FArguments& InArgs )
+{
+	PreferredWidth = InArgs._PreferredWidth;
+	bUseAllottedWidth = InArgs._UseAllottedWidth;
+	StartingAngle = InArgs._StartingAngle;
+	bDistributeItemsEvenly = InArgs._bDistributeItemsEvenly;
+	AngleBetweenItems = InArgs._AngleBetweenItems;
+
+	for ( int32 ChildIndex=0; ChildIndex < InArgs.Slots.Num(); ++ChildIndex )
+	{
+		Slots.Add( InArgs.Slots[ChildIndex] );
+	}
+}
+
+void SRadialBox::Tick( const FGeometry& AllottedGeometry, const double InCurrentTime, const float InDeltaTime )
+{
+	if (bUseAllottedWidth)
+	{
+		PreferredWidth = AllottedGeometry.GetLocalSize().X;
+	}
+}
+
+/*
+ * Simple class for handling the circular arrangement of elements
+ */
+class SRadialBox::FChildArranger
+{
+public:
+	struct FArrangementData
+	{
+		FVector2D SlotOffset;
+		FVector2D SlotSize;
+	};
+
+	typedef TFunctionRef<void(const FSlot& Slot, const FArrangementData& ArrangementData)> FOnSlotArranged;
+
+	static void Arrange(const SRadialBox& RadialBox, const FOnSlotArranged& OnSlotArranged);
+
+private:
+	FChildArranger(const SRadialBox& RadialBox, const FOnSlotArranged& OnSlotArranged);
+	void Arrange();
+
+	const SRadialBox& RadialBox;
+	const FOnSlotArranged& OnSlotArranged;
+	TMap<int32, FArrangementData> OngoingArrangementDataMap;
+};
+
+
+SRadialBox::FChildArranger::FChildArranger(const SRadialBox& InRadialBox, const FOnSlotArranged& InOnSlotArranged)
+	: RadialBox(InRadialBox)
+	, OnSlotArranged(InOnSlotArranged)
+{
+	OngoingArrangementDataMap.Reserve(RadialBox.Slots.Num());
+}
+
+void SRadialBox::FChildArranger::Arrange()
+{
+	const int32 NumItems = RadialBox.Slots.Num();
+	const float Radius = RadialBox.PreferredWidth.Get() / 2.f;
+	const float DegreeIncrements = RadialBox.bDistributeItemsEvenly? 360.f / NumItems : RadialBox.AngleBetweenItems;
+	float DegreeOffset = RadialBox.StartingAngle * (-1.f) ;
+
+	//Offset to create the elements based on the middle of the widget as starting point
+	const float MiddlePointOffset = RadialBox.PreferredWidth.Get() / 2.f;
+
+	for (int32 ChildIndex = 0; ChildIndex < NumItems; ++ChildIndex)
+	{
+		const FSlot& Slot = RadialBox.Slots[ChildIndex];
+		const TSharedRef<SWidget>& Widget = Slot.GetWidget();
+
+		// Skip collapsed widgets.
+		if (Widget->GetVisibility() == EVisibility::Collapsed)
+		{
+			continue;
+		}
+
+		FArrangementData& ArrangementData = OngoingArrangementDataMap.Add(ChildIndex, FArrangementData());
+
+		const FVector2D DesiredSizeOfSlot = Widget->GetDesiredSize();
+
+		ArrangementData.SlotOffset.X = (Radius - DesiredSizeOfSlot.X / 2.f) * FMath::Cos(FMath::DegreesToRadians(DegreeOffset)) + MiddlePointOffset - DesiredSizeOfSlot.X / 2.f;
+		ArrangementData.SlotOffset.Y = (Radius - DesiredSizeOfSlot.Y / 2.f )* FMath::Sin(FMath::DegreesToRadians(DegreeOffset)) + MiddlePointOffset - DesiredSizeOfSlot.Y / 2.f;
+		ArrangementData.SlotSize.X = DesiredSizeOfSlot.X;
+		ArrangementData.SlotSize.Y = DesiredSizeOfSlot.Y;
+
+		DegreeOffset -= DegreeIncrements;
+
+		OnSlotArranged(Slot, ArrangementData);
+	}
+}
+
+void SRadialBox::FChildArranger::Arrange(const SRadialBox& RadialBox, const FOnSlotArranged& OnSlotArranged)
+{
+	FChildArranger(RadialBox, OnSlotArranged).Arrange();
+}
+
+void SRadialBox::OnArrangeChildren(const FGeometry& AllottedGeometry, FArrangedChildren& ArrangedChildren) const
+{
+	FChildArranger::Arrange(*this, [&](const FSlot& Slot, const FChildArranger::FArrangementData& ArrangementData)
+	{
+		ArrangedChildren.AddWidget(AllottedGeometry.MakeChild(Slot.GetWidget(), ArrangementData.SlotOffset, ArrangementData.SlotSize));
+	});
+}
+
+void SRadialBox::ClearChildren()
+{
+	Slots.Empty();
+}
+
+FVector2D SRadialBox::ComputeDesiredSize( float ) const
+{
+	FVector2D MyDesiredSize = FVector2D::ZeroVector;
+
+	FChildArranger::Arrange(*this, [&](const FSlot& Slot, const FChildArranger::FArrangementData& ArrangementData)
+	{
+		// Increase desired size to the maximum X and Y positions of any child widget.
+		MyDesiredSize.X = FMath::Max(MyDesiredSize.X, ArrangementData.SlotOffset.X + ArrangementData.SlotSize.X);
+		MyDesiredSize.Y = FMath::Max(MyDesiredSize.Y, ArrangementData.SlotOffset.Y + ArrangementData.SlotSize.Y);
+	});
+
+	return MyDesiredSize;
+}
+
+FChildren* SRadialBox::GetChildren()
+{
+	return &Slots;	
+}
+
+void SRadialBox::SetUseAllottedWidth(bool bInUseAllottedWidth)
+{
+	bUseAllottedWidth = bInUseAllottedWidth;
+}
diff --git a/Engine/Source/Runtime/Slate/Public/Widgets/Layout/SRadialBox.h b/Engine/Source/Runtime/Slate/Public/Widgets/Layout/SRadialBox.h
new file mode 100644
--- /dev/null
+++ b/Engine/Source/Runtime/Slate/Public/Widgets/Layout/SRadialBox.h
@@ -0,0 +1,119 @@
+// Copyright Epic Games, Inc. All Rights Reserved.
+
+#pragma once
+
+#include "CoreMinimal.h"
+#include "Misc/Attribute.h"
+#include "Layout/Visibility.h"
+#include "Widgets/DeclarativeSyntaxSupport.h"
+#include "SlotBase.h"
+#include "Widgets/SWidget.h"
+#include "Layout/Children.h"
+#include "Widgets/SPanel.h"
+
+class FArrangedChildren;
+
+/**
+ * Arranges widgets in a circular fashion
+ */
+class SLATE_API SRadialBox : public SPanel
+{
+
+public:
+
+	/** Basic Slot without padding or alignment */
+	class FSlot : public TSlotBase<FSlot>
+	{
+	public:
+		FSlot()
+			: TSlotBase<FSlot>()
+		{
+		}
+	};
+
+
+	SLATE_BEGIN_ARGS(SRadialBox)
+		: _PreferredWidth( 100.f )
+		, _UseAllottedWidth( false )
+		, _StartingAngle(0.f)
+		, _bDistributeItemsEvenly(true)
+		, _AngleBetweenItems(0.f)
+	{
+			_Visibility = EVisibility::SelfHitTestInvisible;
+		}
+
+		/** The slot supported by this panel */
+		SLATE_SUPPORTS_SLOT( FSlot )
+
+		/** The preferred width, if not set will fill the space */
+		SLATE_ATTRIBUTE( float, PreferredWidth )
+
+		/** if true, the PreferredWidth will always match the room available to the SRadialBox  */
+		SLATE_ARGUMENT( bool, UseAllottedWidth )
+
+		/** Offset of the first element in the circle in degrees */
+		SLATE_ARGUMENT( float, StartingAngle)
+
+		/** Ignore AngleBetweenItems and distribute items evenly inside the whole circle */
+		SLATE_ARGUMENT(bool, bDistributeItemsEvenly)
+
+		/** How many degrees apart should the elements be? */
+		SLATE_ARGUMENT(float, AngleBetweenItems)
+
+	SLATE_END_ARGS()
+
+	SRadialBox();
+
+	static FSlot& Slot();
+
+	FSlot& AddSlot();
+
+	/** Removes a slot from this radial box which contains the specified SWidget
+	 *
+	 * @param SlotWidget The widget to match when searching through the slots
+	 * @returns The index in the children array where the slot was removed and -1 if no slot was found matching the widget
+	 */
+	int32 RemoveSlot( const TSharedRef<SWidget>& SlotWidget );
+
+	void Construct( const FArguments& InArgs );
+
+	virtual void Tick( const FGeometry& AllottedGeometry, const double InCurrentTime, const float InDeltaTime ) override;
+
+	virtual void OnArrangeChildren( const FGeometry& AllottedGeometry, FArrangedChildren& ArrangedChildren ) const override;
+
+	void ClearChildren();
+
+	virtual FVector2D ComputeDesiredSize(float) const override;
+
+	virtual FChildren* GetChildren() override;
+
+	void SetStartingAngle(float InStartingAngle) { StartingAngle = InStartingAngle; }
+	void SetAngleBetweenItems(float InAngleBetweenItems) { AngleBetweenItems = InAngleBetweenItems; }
+	void SetDistributeItemsEvenly(bool bInDistributeItemsEvenly) { bDistributeItemsEvenly = bInDistributeItemsEvenly; }
+
+	void SetUseAllottedWidth(bool bInUseAllottedWidth);
+
+private:
+
+	/** How wide this panel should appear to be. */
+	TAttribute<float> PreferredWidth;
+
+	/** The slots that contain this panel's children. */
+	TPanelChildren<FSlot> Slots;
+
+	/** If true the box will have a preferred width equal to its alloted width  */
+	bool bUseAllottedWidth;
+
+	/** Offset of the first element in the circle in degrees */
+	float StartingAngle;
+
+	/** Ignore AngleBetweenItems and distribute items evenly inside the whole circle */
+	bool bDistributeItemsEvenly;
+
+	/** How many degrees apart should the elements be? */
+	float AngleBetweenItems;
+
+	class FChildArranger;
+	friend class SRadialBox::FChildArranger;
+};
+
diff --git a/Engine/Source/Runtime/UMG/Private/Components/DynamicEntryBoxBase.cpp b/Engine/Source/Runtime/UMG/Private/Components/DynamicEntryBoxBase.cpp
--- a/Engine/Source/Runtime/UMG/Private/Components/DynamicEntryBoxBase.cpp
+++ b/Engine/Source/Runtime/UMG/Private/Components/DynamicEntryBoxBase.cpp
@@ -7,6 +7,7 @@
 #include "Widgets/Layout/SWrapBox.h"
 #include "Widgets/SBoxPanel.h"
 #include "Editor/WidgetCompilerLog.h"
+#include "Widgets/Layout/SRadialBox.h"
 
 #define LOCTEXT_NAMESPACE "UMG"
 
@@ -42,6 +43,9 @@ void UDynamicEntryBoxBase::ResetInternal(bool bDeleteWidgets)
 		case EDynamicBoxType::VerticalWrap:
 			StaticCastSharedPtr<SWrapBox>(MyPanelWidget)->ClearChildren();
 			break;
+		case EDynamicBoxType::Radial:
+			StaticCastSharedPtr<SRadialBox>(MyPanelWidget)->ClearChildren();
+			break;
 		case EDynamicBoxType::Overlay:
 			StaticCastSharedPtr<SOverlay>(MyPanelWidget)->ClearChildren();
 			break;
@@ -78,6 +82,9 @@ void UDynamicEntryBoxBase::RemoveEntryInternal(UUserWidget* EntryWidget)
 				case EDynamicBoxType::VerticalWrap:
 					StaticCastSharedPtr<SWrapBox>(MyPanelWidget)->RemoveSlot(CachedEntryWidget.ToSharedRef());
 					break;
+				case EDynamicBoxType::Radial:
+					StaticCastSharedPtr<SRadialBox>(MyPanelWidget)->RemoveSlot(CachedEntryWidget.ToSharedRef());
+					break;
 				case EDynamicBoxType::Overlay:
 					StaticCastSharedPtr<SOverlay>(MyPanelWidget)->RemoveSlot(CachedEntryWidget.ToSharedRef());
 					break;
@@ -158,7 +165,7 @@ void UDynamicEntryBoxBase::SetEntrySpacing(const FVector2D& InEntrySpacing)
 				OverlaySlot.SlotPadding = Padding;
 			}
 		}
-		else
+		else if (EntryBoxType == EDynamicBoxType::Horizontal || EntryBoxType == EDynamicBoxType::Vertical)
 		{
 			// Vertical & Horizontal have to manually update the padding on each slot
 			const bool bIsHBox = EntryBoxType == EDynamicBoxType::Horizontal;
@@ -178,6 +185,18 @@ void UDynamicEntryBoxBase::SetEntrySpacing(const FVector2D& InEntrySpacing)
 	}
 }
 
+void UDynamicEntryBoxBase::SetRadialSettings(const FRadialBoxSettings& InSettings)
+{
+	RadialBoxSettings = InSettings;
+
+	if (MyPanelWidget.IsValid() && EntryBoxType == EDynamicBoxType::Radial)
+	{
+		StaticCastSharedPtr<SRadialBox>(MyPanelWidget)->SetStartingAngle(InSettings.StartingAngle);
+		StaticCastSharedPtr<SRadialBox>(MyPanelWidget)->SetDistributeItemsEvenly(InSettings.bDistributeItemsEvenly);
+		StaticCastSharedPtr<SRadialBox>(MyPanelWidget)->SetAngleBetweenItems(InSettings.AngleBetweenItems);
+	}
+}
+
 #if WITH_EDITOR
 
 const FText UDynamicEntryBoxBase::GetPaletteCategory()
@@ -209,6 +228,13 @@ TSharedRef<SWidget> UDynamicEntryBoxBase::RebuildWidget()
 			.Orientation(EOrientation::Orient_Vertical)
 			.InnerSlotPadding(EntrySpacing);
 		break;
+	case EDynamicBoxType::Radial:
+		EntryBoxWidget = SAssignNew(MyPanelWidget, SRadialBox)
+			.UseAllottedWidth(true)
+			.StartingAngle(RadialBoxSettings.StartingAngle)
+			.bDistributeItemsEvenly(RadialBoxSettings.bDistributeItemsEvenly)
+			.AngleBetweenItems(RadialBoxSettings.AngleBetweenItems);
+		break;
 	case EDynamicBoxType::Overlay:
 		EntryBoxWidget = SAssignNew(MyPanelWidget, SOverlay)
 			.Clipping(EWidgetClipping::ClipToBounds);
@@ -249,6 +275,7 @@ void UDynamicEntryBoxBase::SynchronizeProperties()
 	if (IsDesignTime())
 	{
 		SetEntrySpacing(EntrySpacing);
+		SetRadialSettings(RadialBoxSettings);
 	}
 #endif
 }
@@ -318,6 +345,10 @@ void UDynamicEntryBoxBase::AddEntryChild(UUserWidget& ChildWidget)
 			.HAlign(EntryHorizontalAlignment)
 			.VAlign(EntryVerticalAlignment);
 	}
+	else if (EntryBoxType == EDynamicBoxType::Radial)
+	{
+		NewSlot = &StaticCastSharedPtr<SRadialBox>(MyPanelWidget)->AddSlot();
+	}
 	else if (EntryBoxType == EDynamicBoxType::Overlay)
 	{
 		const int32 ChildIdx = MyPanelWidget->GetChildren()->Num();
diff --git a/Engine/Source/Runtime/UMG/Public/Components/DynamicEntryBoxBase.h b/Engine/Source/Runtime/UMG/Public/Components/DynamicEntryBoxBase.h
--- a/Engine/Source/Runtime/UMG/Public/Components/DynamicEntryBoxBase.h
+++ b/Engine/Source/Runtime/UMG/Public/Components/DynamicEntryBoxBase.h
@@ -4,7 +4,7 @@
 
 #include "Components/Widget.h"
 #include "Blueprint/UserWidgetPool.h"
-
+#include "RadialBoxSettings.h"
 #include "DynamicEntryBoxBase.generated.h"
 
 class UUserWidget;
@@ -16,6 +16,7 @@ enum class EDynamicBoxType : uint8
 	Vertical,
 	Wrap,
 	VerticalWrap,
+	Radial,
 	Overlay
 };
 
@@ -60,6 +61,9 @@ public:
 
 	UFUNCTION(BlueprintCallable, Category = DynamicEntryBox)
 	void SetEntrySpacing(const FVector2D& InEntrySpacing);
+
+	UFUNCTION(BlueprintCallable, Category = DynamicEntryBox)
+	void SetRadialSettings(const FRadialBoxSettings& InSettings);
 	
 protected:
 	virtual TSharedRef<SWidget> RebuildWidget() override;
@@ -128,6 +132,10 @@ protected:
 	UPROPERTY(EditAnywhere, BlueprintReadOnly, Category = EntryLayout)
 	int32 MaxElementSize = 0;
 
+	/** Settings only relevant to RadialBox */
+	UPROPERTY(EditAnywhere, BlueprintReadOnly, Category = EntryLayout)
+	FRadialBoxSettings RadialBoxSettings;
+
 	// Can be a horizontal, vertical, wrap box, or overlay
 	TSharedPtr<SPanel> MyPanelWidget;
 
diff --git a/Engine/Source/Runtime/UMG/Public/Components/RadialBoxSettings.h b/Engine/Source/Runtime/UMG/Public/Components/RadialBoxSettings.h
new file mode 100644
--- /dev/null
+++ b/Engine/Source/Runtime/UMG/Public/Components/RadialBoxSettings.h
@@ -0,0 +1,32 @@
+// Copyright Epic Games, Inc. All Rights Reserved.
+
+#pragma once
+
+#include "CoreMinimal.h"
+#include "UObject/ObjectMacros.h"
+#include "RadialBoxSettings.generated.h"
+
+USTRUCT(BlueprintType)
+struct FRadialBoxSettings
+{
+	GENERATED_BODY()
+
+	/* Distribute Items evenly in the whole circle. Checking this option ignores AngleBetweenItems */
+	UPROPERTY(EditAnywhere, Category = "Items")
+	bool bDistributeItemsEvenly;
+
+	/* Amount of Euler degrees that separate each item */
+	UPROPERTY(EditAnywhere, Category = "Items")
+	float AngleBetweenItems;
+
+	/* At what angle will we place the first element of the wheel? */
+	UPROPERTY(EditAnywhere, Category = "Items")
+	float StartingAngle;
+
+	FRadialBoxSettings()
+		: bDistributeItemsEvenly(true)
+		, AngleBetweenItems(0.f)
+		, StartingAngle(0.f)
+	{
+	}
+};
\ No newline at end of file
