Slate - Renaming TranslateMouseCoordinateFor3DChild to TranslateMouseCoordinateForCustomHitTestChild, it's a lot clearer what the function is actually for, and doesn't give the false impression that it's for something it's not, making it a little more broader a naming convention for future usage.

diff --git a/Engine/Source/Runtime/Slate/Private/Widgets/SViewport.cpp b/Engine/Source/Runtime/Slate/Private/Widgets/SViewport.cpp
--- a/Engine/Source/Runtime/Slate/Private/Widgets/SViewport.cpp
+++ b/Engine/Source/Runtime/Slate/Private/Widgets/SViewport.cpp
@@ -385,17 +385,17 @@ void SViewport::OnFinishedPointerInput()
 void SViewport::OnArrangeChildren( const FGeometry& AllottedGeometry, FArrangedChildren& ArrangedChildren ) const
 {
 	SCompoundWidget::OnArrangeChildren(AllottedGeometry, ArrangedChildren);
-	if( ArrangedChildren.Allows3DWidgets() && CustomHitTestPath.IsValid() )
+	if (ArrangedChildren.Allows3DWidgets() && CustomHitTestPath.IsValid())
 	{
-		CustomHitTestPath->ArrangeChildren( ArrangedChildren );
+		CustomHitTestPath->ArrangeCustomHitTestChildren(ArrangedChildren);
 	}
 }
 
-TSharedPtr<FVirtualPointerPosition> SViewport::TranslateMouseCoordinateFor3DChild(const TSharedRef<SWidget>& ChildWidget, const FGeometry& MyGeometry, const FVector2D& ScreenSpaceMouseCoordinate, const FVector2D& LastScreenSpaceMouseCoordinate) const
+TSharedPtr<FVirtualPointerPosition> SViewport::TranslateMouseCoordinateForCustomHitTestChild(const TSharedRef<SWidget>& ChildWidget, const FGeometry& MyGeometry, const FVector2D& ScreenSpaceMouseCoordinate, const FVector2D& LastScreenSpaceMouseCoordinate) const
 {
 	if( CustomHitTestPath.IsValid() )
 	{
-		return CustomHitTestPath->TranslateMouseCoordinateFor3DChild( ChildWidget, MyGeometry, ScreenSpaceMouseCoordinate, LastScreenSpaceMouseCoordinate );
+		return CustomHitTestPath->TranslateMouseCoordinateForCustomHitTestChild( ChildWidget, MyGeometry, ScreenSpaceMouseCoordinate, LastScreenSpaceMouseCoordinate );
 	}
 
 	return nullptr;
diff --git a/Engine/Source/Runtime/Slate/Public/Widgets/SViewport.h b/Engine/Source/Runtime/Slate/Public/Widgets/SViewport.h
--- a/Engine/Source/Runtime/Slate/Public/Widgets/SViewport.h
+++ b/Engine/Source/Runtime/Slate/Public/Widgets/SViewport.h
@@ -242,7 +242,7 @@ public:
 	virtual FPopupMethodReply OnQueryPopupMethod() const override;
 	virtual void OnFinishedPointerInput() override;
 	virtual void OnArrangeChildren( const FGeometry& AllottedGeometry, FArrangedChildren& ArrangedChildren ) const override;
-	virtual TSharedPtr<struct FVirtualPointerPosition> TranslateMouseCoordinateFor3DChild(const TSharedRef<SWidget>& ChildWidget, const FGeometry& MyGeometry, const FVector2D& ScreenSpaceMouseCoordinate, const FVector2D& LastScreenSpaceMouseCoordinate) const override;
+	virtual TSharedPtr<struct FVirtualPointerPosition> TranslateMouseCoordinateForCustomHitTestChild(const TSharedRef<SWidget>& ChildWidget, const FGeometry& MyGeometry, const FVector2D& ScreenSpaceMouseCoordinate, const FVector2D& LastScreenSpaceMouseCoordinate) const override;
 	virtual FNavigationReply OnNavigation( const FGeometry& MyGeometry, const FNavigationEvent& InNavigationEvent ) override;
 
 private:
diff --git a/Engine/Source/Runtime/SlateCore/Private/Layout/WidgetPath.cpp b/Engine/Source/Runtime/SlateCore/Private/Layout/WidgetPath.cpp
--- a/Engine/Source/Runtime/SlateCore/Private/Layout/WidgetPath.cpp
+++ b/Engine/Source/Runtime/SlateCore/Private/Layout/WidgetPath.cpp
@@ -317,7 +317,7 @@ FWeakWidgetPath::EPathResolutionResult::Result FWeakWidgetPath::ToWidgetPath( FW
 					{
 						if (PointerEvent && !VirtualPointerPos.IsValid())
 						{
-							VirtualPointerPos = CurWidget->TranslateMouseCoordinateFor3DChild(ChildWidgetPtr.ToSharedRef(), ParentGeometry, PointerEvent->GetScreenSpacePosition(), PointerEvent->GetLastScreenSpacePosition());
+							VirtualPointerPos = CurWidget->TranslateMouseCoordinateForCustomHitTestChild(ChildWidgetPtr.ToSharedRef(), ParentGeometry, PointerEvent->GetScreenSpacePosition(), PointerEvent->GetLastScreenSpacePosition());
 						}
 
 						bFoundChild = true;
@@ -381,7 +381,7 @@ FWeakWidgetPath::EPathResolutionResult::Result FWeakWidgetPath::ToWidgetPath( FW
 						{
 							if (PointerEvent && !VirtualPointerPos.IsValid())
 							{
-								VirtualPointerPos = CurWidget->TranslateMouseCoordinateFor3DChild(ArrangedWidget.Widget, ParentGeometry, PointerEvent->GetScreenSpacePosition(), PointerEvent->GetLastScreenSpacePosition());
+								VirtualPointerPos = CurWidget->TranslateMouseCoordinateForCustomHitTestChild(ArrangedWidget.Widget, ParentGeometry, PointerEvent->GetScreenSpacePosition(), PointerEvent->GetLastScreenSpacePosition());
 							}
 
 							bFoundChild = true;
diff --git a/Engine/Source/Runtime/SlateCore/Private/Widgets/SWidget.cpp b/Engine/Source/Runtime/SlateCore/Private/Widgets/SWidget.cpp
--- a/Engine/Source/Runtime/SlateCore/Private/Widgets/SWidget.cpp
+++ b/Engine/Source/Runtime/SlateCore/Private/Widgets/SWidget.cpp
@@ -553,7 +553,7 @@ FPopupMethodReply SWidget::OnQueryPopupMethod() const
 	return FPopupMethodReply::Unhandled();
 }
 
-TSharedPtr<struct FVirtualPointerPosition> SWidget::TranslateMouseCoordinateFor3DChild(const TSharedRef<SWidget>& ChildWidget, const FGeometry& MyGeometry, const FVector2D& ScreenSpaceMouseCoordinate, const FVector2D& LastScreenSpaceMouseCoordinate) const
+TSharedPtr<struct FVirtualPointerPosition> SWidget::TranslateMouseCoordinateForCustomHitTestChild(const TSharedRef<SWidget>& ChildWidget, const FGeometry& MyGeometry, const FVector2D& ScreenSpaceMouseCoordinate, const FVector2D& LastScreenSpaceMouseCoordinate) const
 {
 	return nullptr;
 }
diff --git a/Engine/Source/Runtime/SlateCore/Public/Input/HittestGrid.h b/Engine/Source/Runtime/SlateCore/Public/Input/HittestGrid.h
--- a/Engine/Source/Runtime/SlateCore/Public/Input/HittestGrid.h
+++ b/Engine/Source/Runtime/SlateCore/Public/Input/HittestGrid.h
@@ -18,9 +18,9 @@ public:
 
 	virtual TArray<FWidgetAndPointer> GetBubblePathAndVirtualCursors( const FGeometry& InGeometry, FVector2D DesktopSpaceCoordinate, bool bIgnoreEnabledStatus ) const = 0;
 
-	virtual void ArrangeChildren( FArrangedChildren& ArrangedChildren ) const = 0;
+	virtual void ArrangeCustomHitTestChildren( FArrangedChildren& ArrangedChildren ) const = 0;
 
-	virtual TSharedPtr<struct FVirtualPointerPosition> TranslateMouseCoordinateFor3DChild( const TSharedRef<SWidget>& ChildWidget, const FGeometry& ViewportGeometry, const FVector2D& ScreenSpaceMouseCoordinate, const FVector2D& LastScreenSpaceMouseCoordinate ) const = 0;
+	virtual TSharedPtr<struct FVirtualPointerPosition> TranslateMouseCoordinateForCustomHitTestChild( const TSharedRef<SWidget>& ChildWidget, const FGeometry& ViewportGeometry, const FVector2D& ScreenSpaceMouseCoordinate, const FVector2D& LastScreenSpaceMouseCoordinate ) const = 0;
 };
 
 class SLATECORE_API FHittestGrid
diff --git a/Engine/Source/Runtime/SlateCore/Public/Widgets/SWidget.h b/Engine/Source/Runtime/SlateCore/Public/Widgets/SWidget.h
--- a/Engine/Source/Runtime/SlateCore/Public/Widgets/SWidget.h
+++ b/Engine/Source/Runtime/SlateCore/Public/Widgets/SWidget.h
@@ -604,7 +604,13 @@ public:
 	 */
 	virtual FPopupMethodReply OnQueryPopupMethod() const;
 
-	virtual TSharedPtr<FVirtualPointerPosition> TranslateMouseCoordinateFor3DChild(const TSharedRef<SWidget>& ChildWidget, const FGeometry& MyGeometry, const FVector2D& ScreenSpaceMouseCoordinate, const FVector2D& LastScreenSpaceMouseCoordinate) const;
+	UE_DEPRECATED(4.27, "Renaming to TranslateMouseCoordinateForCustomHitTestChild")
+	TSharedPtr<FVirtualPointerPosition> TranslateMouseCoordinateFor3DChild(const TSharedRef<SWidget>& ChildWidget, const FGeometry& MyGeometry, const FVector2D& ScreenSpaceMouseCoordinate, const FVector2D& LastScreenSpaceMouseCoordinate) const
+	{
+		return TranslateMouseCoordinateForCustomHitTestChild(ChildWidget, MyGeometry, ScreenSpaceMouseCoordinate, LastScreenSpaceMouseCoordinate);
+	}
+
+	virtual TSharedPtr<FVirtualPointerPosition> TranslateMouseCoordinateForCustomHitTestChild(const TSharedRef<SWidget>& ChildWidget, const FGeometry& MyGeometry, const FVector2D& ScreenSpaceMouseCoordinate, const FVector2D& LastScreenSpaceMouseCoordinate) const;
 
 	/**
 	 * All the pointer (mouse, touch, stylus, etc.) events from this frame have been routed.
diff --git a/Engine/Source/Runtime/UMG/Private/Components/WidgetComponent.cpp b/Engine/Source/Runtime/UMG/Private/Components/WidgetComponent.cpp
--- a/Engine/Source/Runtime/UMG/Private/Components/WidgetComponent.cpp
+++ b/Engine/Source/Runtime/UMG/Private/Components/WidgetComponent.cpp
@@ -162,7 +162,7 @@ public:
 		return TArray<FWidgetAndPointer>();
 	}
 
-	virtual void ArrangeChildren( FArrangedChildren& ArrangedChildren ) const override
+	virtual void ArrangeCustomHitTestChildren( FArrangedChildren& ArrangedChildren ) const override
 	{
 		for( TWeakObjectPtr<UWidgetComponent> Component : RegisteredComponents )
 		{
@@ -177,7 +177,7 @@ public:
 		}
 	}
 
-	virtual TSharedPtr<struct FVirtualPointerPosition> TranslateMouseCoordinateFor3DChild( const TSharedRef<SWidget>& ChildWidget, const FGeometry& ViewportGeometry, const FVector2D& ScreenSpaceMouseCoordinate, const FVector2D& LastScreenSpaceMouseCoordinate ) const override
+	virtual TSharedPtr<struct FVirtualPointerPosition> TranslateMouseCoordinateForCustomHitTestChild( const TSharedRef<SWidget>& ChildWidget, const FGeometry& ViewportGeometry, const FVector2D& ScreenSpaceMouseCoordinate, const FVector2D& LastScreenSpaceMouseCoordinate ) const override
 	{
 		if ( World.IsValid() && ensure(World->IsGameWorld()) )
 		{
