Fix merge issue with MultiBoxBuilder preventing extenders from being disabled
#rnx

diff --git a/Engine/Source/Runtime/Slate/Private/Framework/MultiBox/MultiBoxBuilder.cpp b/Engine/Source/Runtime/Slate/Private/Framework/MultiBox/MultiBoxBuilder.cpp
--- a/Engine/Source/Runtime/Slate/Private/Framework/MultiBox/MultiBoxBuilder.cpp
+++ b/Engine/Source/Runtime/Slate/Private/Framework/MultiBox/MultiBoxBuilder.cpp
@@ -316,13 +316,16 @@ void FMenuBuilder::AddSearchWidget()
 
 void FMenuBuilder::ApplyHook(FName InExtensionHook, EExtensionHook::Position HookPosition)
 {
-	// this is a virtual function to get a properly typed "this" pointer
-	auto& Extender = ExtenderStack.Top();
-	if (InExtensionHook != NAME_None && Extender.IsValid())
+	if (ExtendersEnabled())
 	{
-		if (!MultiBox->IsInEditMode())
+		// this is a virtual function to get a properly typed "this" pointer
+		auto& Extender = ExtenderStack.Top();
+		if (InExtensionHook != NAME_None && Extender.IsValid())
 		{
-			Extender->Apply(InExtensionHook, HookPosition, *this);
+			if (!MultiBox->IsInEditMode())
+			{
+				Extender->Apply(InExtensionHook, HookPosition, *this);
+			}
 		}
 	}
 }
@@ -365,11 +368,14 @@ void FMenuBarBuilder::AddPullDownMenu(const FText& InMenuLabel, const FText& InT
 
 void FMenuBarBuilder::ApplyHook(FName InExtensionHook, EExtensionHook::Position HookPosition)
 {
-	// this is a virtual function to get a properly typed "this" pointer
-	auto& Extender = ExtenderStack.Top();
-	if (InExtensionHook != NAME_None && Extender.IsValid())
+	if (ExtendersEnabled())
 	{
-		Extender->Apply(InExtensionHook, HookPosition, *this);
+		// this is a virtual function to get a properly typed "this" pointer
+		auto& Extender = ExtenderStack.Top();
+		if (InExtensionHook != NAME_None && Extender.IsValid())
+		{
+			Extender->Apply(InExtensionHook, HookPosition, *this);
+		}
 	}
 }
 
@@ -542,11 +548,14 @@ void FToolBarBuilder::EndSection()
 
 void FToolBarBuilder::ApplyHook(FName InExtensionHook, EExtensionHook::Position HookPosition)
 {
-	// this is a virtual function to get a properly typed "this" pointer
-	auto& Extender = ExtenderStack.Top();
-	if (InExtensionHook != NAME_None && Extender.IsValid())
+	if (ExtendersEnabled())
 	{
-		Extender->Apply(InExtensionHook, HookPosition, *this);
+		// this is a virtual function to get a properly typed "this" pointer
+		auto& Extender = ExtenderStack.Top();
+		if (InExtensionHook != NAME_None && Extender.IsValid())
+		{
+			Extender->Apply(InExtensionHook, HookPosition, *this);
+		}
 	}
 }
 
