Fix foreground color of submenus

diff --git a/Engine/Source/Runtime/Slate/Private/Framework/MultiBox/SMenuEntryBlock.cpp b/Engine/Source/Runtime/Slate/Private/Framework/MultiBox/SMenuEntryBlock.cpp
--- a/Engine/Source/Runtime/Slate/Private/Framework/MultiBox/SMenuEntryBlock.cpp
+++ b/Engine/Source/Runtime/Slate/Private/Framework/MultiBox/SMenuEntryBlock.cpp
@@ -188,6 +188,8 @@ void SMenuEntryBlock::Construct( const FArguments& InArgs )
 	/*TimeToSubMenuOpen = 0.0f;*/
 	//SubMenuRequestState = Idle;
 
+	MenuBarButtonStyle = nullptr;
+
 	// No images by default
 	CheckedImage = nullptr;
 	UncheckedImage = nullptr;
@@ -206,13 +208,8 @@ TSharedRef< SWidget> SMenuEntryBlock::BuildMenuBarWidget( const FMenuEntryBuildP
 	const ISlateStyle* const StyleSet = InBuildParams.StyleSet;
 	const FName& StyleName = InBuildParams.StyleName;
 
-	/* Style for menu bar button with sub menu opened */
-	MenuBarButtonBorderSubmenuOpen = StyleSet->GetBrush(StyleName, ".Button.SubMenuOpen");
-	/* Style for menu bar button with no sub menu opened */
-	MenuBarButtonBorderSubmenuClosed = FCoreStyle::Get().GetBrush("NoBorder");
-
 	TSharedPtr< SMenuAnchor > NewMenuAnchor;
-
+	MenuBarButtonStyle = &StyleSet->GetWidgetStyle<FButtonStyle>(ISlateStyle::Join(StyleName, ".Button"));
 
 	// Create a menu bar button within a pop-up anchor
 	TSharedRef<SWidget> Widget =
@@ -223,20 +220,19 @@ TSharedRef< SWidget> SMenuEntryBlock::BuildMenuBarWidget( const FMenuEntryBuildP
 		[
 			SNew(SBorder)
 			.BorderImage(this, &SMenuEntryBlock::GetMenuBarButtonBorder)
-			.Padding(0)
 			[
 				// Create a button
 				SNew(SButton)
 				// Use the menu bar item style for this button
-				.ButtonStyle(StyleSet, ISlateStyle::Join(StyleName, ".Button"))
+				.ButtonStyle(StyleSet, "NoBorder")
 				// Pull-down menu bar items always activate on mouse-down, not mouse-up
 				.ClickMethod(EButtonClickMethod::MouseDown)
 				// Pass along the block's tool-tip string
 				.ToolTipText(this, &SMenuEntryBlock::GetFilteredToolTipText, EntryToolTip)
 				// Add horizontal padding between the edge of the button and the content.  Also add a bit of vertical
 				// padding to push the text down from the top of the menu bar a bit.
-				.ContentPadding(FMargin(10.0f, 2.0f))
-				.ForegroundColor(FSlateColor::UseForeground())
+				.ContentPadding(StyleSet->GetMargin(StyleName, ".MenuBar.Padding"))
+				.ForegroundColor(this, &SMenuEntryBlock::GetMenuBarForegroundColor)
 				.VAlign(VAlign_Center)
 				[
 					SNew(STextBlock)
@@ -322,6 +318,8 @@ public:
 	SLATE_STYLE_ARGUMENT(FButtonStyle, ButtonStyle)
 	/** Sets the rules to use for determining whether the button was clicked.  This is an advanced setting and generally should be left as the default. */
 	SLATE_ARGUMENT(EButtonClickMethod::Type, ClickMethod)
+	/** Spacing between button's border and the content. */
+	SLATE_ARGUMENT(FMargin, ContentPadding)
 	/** Called when the button is clicked */
 	SLATE_EVENT(FOnClicked, OnClicked)
 
@@ -340,8 +338,8 @@ public:
 		ButtonArgs.ButtonStyle(InArgs._ButtonStyle);
 		ButtonArgs.ClickMethod(InArgs._ClickMethod);
 		ButtonArgs.ToolTip(InArgs._ToolTip);
-		ButtonArgs.ContentPadding(FMargin(0, 2));
-		ButtonArgs.ForegroundColor(FSlateColor::UseForeground());
+		ButtonArgs.ContentPadding(InArgs._ContentPadding);
+		ButtonArgs.ForegroundColor(FSlateColor::UseStyle());
 		ButtonArgs.OnClicked(InArgs._OnClicked)
 		[
 			InArgs._Content.Widget
@@ -424,7 +422,7 @@ TSharedRef< SWidget > SMenuEntryBlock::BuildMenuEntryWidget( const FMenuEntryBui
 	const FName& StyleName = InBuildParams.StyleName;
 
 	// Grab the friendly text name for this action's input binding
-	FText InputBindingText = UICommand.IsValid() ? UICommand->GetInputText() : FText::GetEmpty();
+	FText InputBindingText = UICommand.IsValid() ? UICommand->GetInputText().ToUpper() : FText::GetEmpty();
 
 	// Allow menu item buttons to be triggered on mouse-up events if the menu is configured to be
 	// dismissed automatically after clicking.  This preserves the behavior people expect for context
@@ -446,6 +444,8 @@ TSharedRef< SWidget > SMenuEntryBlock::BuildMenuEntryWidget( const FMenuEntryBui
 		}
 	}
 
+	const float MenuIconSize = StyleSet->GetFloat(StyleName, ".MenuIconSize", 16.f);
+
 	if (bSectionContainsIcons && (IconWidget == SNullWidget::NullWidget))
 	{
 		// Section should have icons but this entry does not, which is inconsistent with our menu policy (either all or none of menu items in a section should have an icon)
@@ -453,7 +453,7 @@ TSharedRef< SWidget > SMenuEntryBlock::BuildMenuEntryWidget( const FMenuEntryBui
 		{
 			IconWidget = SNew(SColorBlock)
 				.Color(FColorList::Magenta)
-				.Size(FVector2D(MultiBoxConstants::MenuIconSize, MultiBoxConstants::MenuIconSize))
+				.Size(FVector2D(MenuIconSize, MenuIconSize))
 				.ToolTipText(NSLOCTEXT("SMenuEntryBlock", "MissingIconInMenu", "This menu entry is missing an icon and should be fixed (consistency within each section is required, either every entry in the section has an icon or no entries have an icon)"));
 		}
 		else
@@ -491,7 +491,7 @@ TSharedRef< SWidget > SMenuEntryBlock::BuildMenuEntryWidget( const FMenuEntryBui
 		CheckBoxVisibility = EVisibility::Hidden;
 	}
 
-	TAttribute<FSlateColor> CheckBoxForegroundColor = FSlateColor::UseForeground();
+	TAttribute<FSlateColor> CheckBoxForegroundColor = FSlateColor::UseStyle();
 	FName CheckBoxStyle = ISlateStyle::Join( StyleName, ".CheckBox" );
 	if (UserInterfaceType == EUserInterfaceActionType::Check)
 	{
@@ -503,6 +503,7 @@ TSharedRef< SWidget > SMenuEntryBlock::BuildMenuEntryWidget( const FMenuEntryBui
 		CheckBoxForegroundColor = TAttribute<FSlateColor>::Create( TAttribute<FSlateColor>::FGetter::CreateRaw( this, &SMenuEntryBlock::TintOnHover ) );
 	}
 
+
 	// If there is custom menu widget, set it
 	// If there isn't, create it
 	TSharedPtr< SWidget > ButtonContent = MenuEntryBlock->EntryWidget;
@@ -517,14 +518,14 @@ TSharedRef< SWidget > SMenuEntryBlock::BuildMenuEntryWidget( const FMenuEntryBui
 		[
 			SNew( SBox )
 			.Visibility(IconWidget != SNullWidget::NullWidget ? EVisibility::Visible : EVisibility::Collapsed)
-			.WidthOverride( MultiBoxConstants::MenuIconSize + 2 )
-			.HeightOverride( MultiBoxConstants::MenuIconSize )
+			.WidthOverride(MenuIconSize + 2 )
+			.HeightOverride(MenuIconSize)
 			.HAlign(HAlign_Center)
 			.VAlign(VAlign_Center)
 			[
 				SNew( SBox )
-				.WidthOverride( MultiBoxConstants::MenuIconSize )
-				.HeightOverride( MultiBoxConstants::MenuIconSize )
+				.WidthOverride(MenuIconSize)
+				.HeightOverride(MenuIconSize)
 				[
 					IconWidget
 				]
@@ -542,6 +543,7 @@ TSharedRef< SWidget > SMenuEntryBlock::BuildMenuEntryWidget( const FMenuEntryBui
 		]
 		+ SHorizontalBox::Slot()
 		.AutoWidth()
+		.Padding(FMargin(25,0,0,0))
 		.VAlign( VAlign_Center )
 		.HAlign( HAlign_Right )
 		[
@@ -558,15 +560,16 @@ TSharedRef< SWidget > SMenuEntryBlock::BuildMenuEntryWidget( const FMenuEntryBui
 	}
 
 	// Create a wrapper containing the checkbox, and the generated button content
-	TSharedPtr< SWidget > CheckBoxAndButtonContent = 
+	TSharedRef< SWidget > CheckBoxAndButtonContent = 
 		SNew( SHorizontalBox )
 		+SHorizontalBox::Slot()
+		.Padding(FMargin(2, 0, 8, 0))
 		.AutoWidth()
 		[
 			SNew(SBox)
 			.Visibility(CheckBoxVisibility)
-			.WidthOverride( MultiBoxConstants::MenuCheckBoxSize )
-			.HeightOverride( MultiBoxConstants::MenuCheckBoxSize )
+			.WidthOverride(MenuIconSize)
+			.HeightOverride(MenuIconSize)
 			.HAlign(HAlign_Center)
 			.VAlign(VAlign_Center)
 			[
@@ -583,6 +586,7 @@ TSharedRef< SWidget > SMenuEntryBlock::BuildMenuEntryWidget( const FMenuEntryBui
 			]
 		]
 		+SHorizontalBox::Slot()
+		.Padding(FMargin(2, 0, 0, 0))
 		[
 			ButtonContent.ToSharedRef()
 		];
@@ -593,12 +597,13 @@ TSharedRef< SWidget > SMenuEntryBlock::BuildMenuEntryWidget( const FMenuEntryBui
 		.ButtonStyle( StyleSet, ISlateStyle::Join( StyleName, ".Button" ) )
 		// Set our click method for this menu item.  It will be different for pull-down/context menus.
 		.ClickMethod( ButtonClickMethod )
+		.ContentPadding(StyleSet->GetMargin(StyleName,".Block.Padding"))
 		// Pass along the block's tool-tip string
 		.ToolTip( FMultiBoxSettings::ToolTipConstructor.Execute(EntryToolTip, nullptr, UICommand ) )
 		// Bind the button's "on clicked" event to our object's method for this
 		.OnClicked(this, &SMenuEntryBlock::OnMenuItemButtonClicked)
 		[
-			CheckBoxAndButtonContent.ToSharedRef()
+			CheckBoxAndButtonContent
 		];
 
 	return MenuEntryWidget.ToSharedRef();
@@ -627,6 +632,8 @@ public:
 		SLATE_ATTRIBUTE( bool, ShouldAppearHovered )
 		/** The style to use */
 		SLATE_STYLE_ARGUMENT( FButtonStyle, ButtonStyle )
+		/** Spacing between button's border and the content. */
+		SLATE_ARGUMENT(FMargin, ContentPadding)
 	SLATE_END_ARGS()
 
 	void Construct( const FArguments& InArgs )
@@ -638,13 +645,13 @@ public:
 		ButtonArgs.ForegroundColor( this, &SSubMenuButton::InvertOnHover );
 		ButtonArgs.HAlign(HAlign_Fill);
 		ButtonArgs.VAlign(VAlign_Fill);
-		ButtonArgs.ContentPadding(FMargin(0,2));
+		ButtonArgs.ContentPadding(InArgs._ContentPadding);
 
 		if ( InArgs._ButtonStyle )
 		{
 			ButtonArgs.ButtonStyle( InArgs._ButtonStyle );
 		}
-
+	
 		ButtonArgs.OnClicked(InArgs._OnClicked);
 		ButtonArgs.ClickMethod(EButtonClickMethod::MouseDown)
 		[
@@ -669,7 +676,7 @@ private:
 		}
 		else
 		{
-			return FSlateColor::UseForeground();
+			return FSlateColor::UseStyle();
 		}
 	}
 
@@ -746,6 +753,8 @@ TSharedRef< SWidget> SMenuEntryBlock::BuildSubMenuWidget( const FMenuEntryBuildP
 		CheckBoxForegroundColor = TAttribute<FSlateColor>::Create( TAttribute<FSlateColor>::FGetter::CreateRaw( this, &SMenuEntryBlock::TintOnHover ) );
 	}
 
+	const float MenuIconSize = StyleSet->GetFloat(StyleName, ".MenuIconSize", 16.f);
+
 	TSharedPtr< SWidget > ButtonContent = MenuEntryBlock->EntryWidget;
 	if ( !ButtonContent.IsValid() )
 	{
@@ -754,39 +763,18 @@ TSharedRef< SWidget> SMenuEntryBlock::BuildSubMenuWidget( const FMenuEntryBuildP
 		// Whatever we have in the icon area goes first
 		+ SHorizontalBox::Slot()
 		.AutoWidth()
-		[
-			SNew(SBox)
-			.Visibility(CheckBoxVisibility)
-			.WidthOverride( MultiBoxConstants::MenuCheckBoxSize )
-			.HeightOverride( MultiBoxConstants::MenuCheckBoxSize )
-			.HAlign(HAlign_Center)
-			.VAlign(VAlign_Center)
-			[
-				// For check style menus, use an image instead of a CheckBox because it can't really be checked.
-				UserInterfaceType == EUserInterfaceActionType::Check
-					? StaticCastSharedRef<SWidget>(SNew( SImage )
-						.Image(this, &SMenuEntryBlock::GetCheckBoxImageBrushFromStyle, &StyleSet->GetWidgetStyle<FCheckBoxStyle>(CheckBoxStyle)))
-					: StaticCastSharedRef<SWidget>(SNew( SCheckBox )
-						.ForegroundColor( CheckBoxForegroundColor )
-						.IsChecked( this, &SMenuEntryBlock::IsChecked )
-						.Style( StyleSet, CheckBoxStyle )
-						.OnCheckStateChanged( this, &SMenuEntryBlock::OnCheckStateChanged ))
-			]
-		]
-		+ SHorizontalBox::Slot()
-		.AutoWidth()
 		.Padding(FMargin(2, 0, 2, 0))
 		[
 			SNew( SBox )
 			.Visibility(IconWidget != SNullWidget::NullWidget ? EVisibility::Visible : EVisibility::Collapsed)
-			.WidthOverride( MultiBoxConstants::MenuIconSize + 2 )
-			.HeightOverride( MultiBoxConstants::MenuIconSize )
+			.WidthOverride(MenuIconSize + 2)
+			.HeightOverride(MenuIconSize)
 			.HAlign(HAlign_Center)
 			.VAlign(VAlign_Center)
 			[
 				SNew( SBox )
-				.WidthOverride( MultiBoxConstants::MenuIconSize )
-				.HeightOverride( MultiBoxConstants::MenuIconSize )
+				.WidthOverride(MenuIconSize)
+				.HeightOverride(MenuIconSize)
 				[
 					IconWidget
 				]
@@ -804,6 +792,38 @@ TSharedRef< SWidget> SMenuEntryBlock::BuildSubMenuWidget( const FMenuEntryBuildP
 		];
 	}
 
+	// Create a wrapper containing the checkbox, and the generated button content
+	TSharedRef< SWidget > CheckBoxAndButtonContent =
+		SNew(SHorizontalBox)
+		+ SHorizontalBox::Slot()
+		.Padding(FMargin(2, 0, 8, 0))
+		.AutoWidth()
+		[
+			SNew(SBox)
+			.Visibility(CheckBoxVisibility)
+			.WidthOverride(MenuIconSize)
+			.HeightOverride(MenuIconSize)
+			.HAlign(HAlign_Center)
+			.VAlign(VAlign_Center)
+			[
+				// For check style menus, use an image instead of a CheckBox because it can't really be checked.
+				UserInterfaceType == EUserInterfaceActionType::Check
+				? StaticCastSharedRef<SWidget>(SNew(SImage)
+					.ColorAndOpacity(FSlateColor::UseForeground())
+					.Image(this, &SMenuEntryBlock::GetCheckBoxImageBrushFromStyle, &StyleSet->GetWidgetStyle<FCheckBoxStyle>(CheckBoxStyle)))
+				: StaticCastSharedRef<SWidget>(SNew(SCheckBox)
+					.ForegroundColor(CheckBoxForegroundColor)
+					.IsChecked(this, &SMenuEntryBlock::IsChecked)
+					.Style(StyleSet, CheckBoxStyle)
+					.OnCheckStateChanged(this, &SMenuEntryBlock::OnCheckStateChanged))
+			]
+		]
+		+ SHorizontalBox::Slot()
+		.Padding(FMargin(2, 0, 0, 0))
+		[
+			ButtonContent.ToSharedRef()
+		];
+
 	TSharedPtr< SMenuAnchor > NewMenuAnchorPtr;
 	TSharedRef< SWidget > Widget = 
 	SAssignNew( NewMenuAnchorPtr, SMenuAnchor )
@@ -817,6 +837,7 @@ TSharedRef< SWidget> SMenuEntryBlock::BuildSubMenuWidget( const FMenuEntryBuildP
 			.ToolTipText( EntryToolTip )
 			// Style to use
 			.ButtonStyle( &StyleSet->GetWidgetStyle<FButtonStyle>( ISlateStyle::Join( StyleName, ".Button" ) ) )
+			.ContentPadding(StyleSet->GetMargin(StyleName, ".Block.Padding"))
 			// Allow the button to change its state depending on the state of the submenu
 			.ShouldAppearHovered( this, &SMenuEntryBlock::ShouldSubMenuAppearHovered )
 			[
@@ -824,7 +845,7 @@ TSharedRef< SWidget> SMenuEntryBlock::BuildSubMenuWidget( const FMenuEntryBuildP
 				+ SHorizontalBox::Slot()
 				.FillWidth( 1.0f )
 				[
-					ButtonContent.ToSharedRef()
+					CheckBoxAndButtonContent
 				]
 				+ SHorizontalBox::Slot()
 				.AutoWidth()
@@ -944,12 +965,15 @@ void SMenuEntryBlock::BuildMultiBlockWidget(const ISlateStyle* StyleSet, const F
 			ChildSlot.AttachWidget( BuildMenuEntryWidget( BuildParams ) );
 		}
 	}
-	else if( ensure( MultiBox->GetType() == EMultiBoxType::MenuBar || MultiBox->GetType() == EMultiBoxType::ToolMenuBar ) )
+	else if( ensure( MultiBox->GetType() == EMultiBoxType::MenuBar) )
 	{
 		// Menu bar items cannot be submenus
 		check( !MenuEntryBlock->bIsSubMenu );
 		
-		ChildSlot[ BuildMenuBarWidget( BuildParams ) ];
+		ChildSlot
+		[ 
+			BuildMenuBarWidget( BuildParams ) 
+		];
 	}
 
 	if (!bIsEditing)
@@ -1475,16 +1499,41 @@ TSharedRef< SWidget > SMenuEntryBlock::MakeNewMenuWidget() const
  *
  * @return	The appropriate border to use
  */
+
 const FSlateBrush* SMenuEntryBlock::GetMenuBarButtonBorder( ) const
 {
 	TSharedPtr<SMenuAnchor> MenuAnchorSharedPtr = MenuAnchor.Pin();
 
 	if (MenuAnchorSharedPtr.IsValid() && MenuAnchorSharedPtr->IsOpen())
 	{
-		return MenuBarButtonBorderSubmenuOpen;
+		return &MenuBarButtonStyle->Pressed;
+	}
+	else if (IsHovered())
+	{
+		return &MenuBarButtonStyle->Hovered;
+	}
+	else
+	{
+		return &MenuBarButtonStyle->Normal;
 	}
+}
+
+FSlateColor SMenuEntryBlock::GetMenuBarForegroundColor() const
+{
+	TSharedPtr<SMenuAnchor> MenuAnchorSharedPtr = MenuAnchor.Pin();
 
-	return MenuBarButtonBorderSubmenuClosed;
+	if (MenuAnchorSharedPtr.IsValid() && MenuAnchorSharedPtr->IsOpen())
+	{
+		return MenuBarButtonStyle->PressedForeground;
+	}
+	else if (IsHovered())
+	{
+		return MenuBarButtonStyle->HoveredForeground;
+	}
+	else
+	{
+		return MenuBarButtonStyle->NormalForeground;
+	}
 }
 
 FSlateColor SMenuEntryBlock::TintOnHover() const
@@ -1507,13 +1556,17 @@ FSlateColor SMenuEntryBlock::TintOnHover() const
 
 FSlateColor SMenuEntryBlock::InvertOnHover() const
 {
-	if ( this->IsHovered() )
+	if (ShouldSubMenuAppearHovered())
 	{
 		return FLinearColor::Black;
 	}
+	else if (this->IsHovered())
+	{
+		return FSlateColor::UseStyle();
+	}
 	else
 	{
-		return FSlateColor::UseForeground();
+		return FSlateColor::UseStyle();
 	}
 }
 
