Move gradient widgets out of AppFramework and into Slate. They are neeed for slate core styling

diff --git a/Engine/Source/Runtime/AppFramework/Private/Widgets/Colors/SComplexGradient.cpp b/Engine/Source/Runtime/AppFramework/Private/Widgets/Colors/SComplexGradient.cpp
deleted file mode 100644
diff --git a/Engine/Source/Runtime/AppFramework/Private/Widgets/Colors/SSimpleGradient.cpp b/Engine/Source/Runtime/AppFramework/Private/Widgets/Colors/SSimpleGradient.cpp
deleted file mode 100644
diff --git a/Engine/Source/Runtime/AppFramework/Public/Widgets/Colors/SComplexGradient.h b/Engine/Source/Runtime/AppFramework/Public/Widgets/Colors/SComplexGradient.h
deleted file mode 100644
diff --git a/Engine/Source/Runtime/AppFramework/Public/Widgets/Colors/SSimpleGradient.h b/Engine/Source/Runtime/AppFramework/Public/Widgets/Colors/SSimpleGradient.h
deleted file mode 100644
diff --git a/Engine/Source/Runtime/Slate/Private/Widgets/Colors/SComplexGradient.cpp b/Engine/Source/Runtime/Slate/Private/Widgets/Colors/SComplexGradient.cpp
new file mode 100644
--- /dev/null
+++ b/Engine/Source/Runtime/Slate/Private/Widgets/Colors/SComplexGradient.cpp
@@ -0,0 +1,63 @@
+// Copyright Epic Games, Inc. All Rights Reserved.
+
+#include "Widgets/Colors/SComplexGradient.h"
+#include "Rendering/DrawElements.h"
+#include "Styling/CoreStyle.h"
+#include "Application/SlateWindowHelper.h"
+
+
+/* SComplexGradient interface
+ *****************************************************************************/
+
+void SComplexGradient::Construct( const FArguments& InArgs )
+{
+	GradientColors = InArgs._GradientColors;
+	bHasAlphaBackground = InArgs._HasAlphaBackground.Get();
+	Orientation = InArgs._Orientation.Get();
+}
+
+
+/* SCompoundWidget overrides
+ *****************************************************************************/
+
+int32 SComplexGradient::OnPaint( const FPaintArgs& Args, const FGeometry& AllottedGeometry, const FSlateRect& MyCullingRect, FSlateWindowElementList& OutDrawElements, int32 LayerId, const FWidgetStyle& InWidgetStyle, bool bParentEnabled ) const
+{
+	ESlateDrawEffect DrawEffects = (bParentEnabled && IsEnabled()) ? ESlateDrawEffect::None : ESlateDrawEffect::DisabledEffect;
+
+	if (bHasAlphaBackground)
+	{
+		const FSlateBrush* StyleInfo = FCoreStyle::Get().GetBrush("ColorPicker.AlphaBackground");
+
+		FSlateDrawElement::MakeBox(
+			OutDrawElements,
+			LayerId,
+			AllottedGeometry.ToPaintGeometry(),
+			StyleInfo,
+			DrawEffects
+		);
+	}
+
+	const TArray<FLinearColor>& Colors = GradientColors.Get();
+	int32 NumColors = Colors.Num();
+
+	if (NumColors > 0)
+	{
+		TArray<FSlateGradientStop> GradientStops;
+
+		for (int32 ColorIndex = 0; ColorIndex < NumColors; ++ColorIndex)
+		{
+			GradientStops.Add(FSlateGradientStop(AllottedGeometry.GetLocalSize() * (float(ColorIndex) / (NumColors - 1)), Colors[ColorIndex]));
+		}
+
+		FSlateDrawElement::MakeGradient(
+			OutDrawElements,
+			LayerId + 1,
+			AllottedGeometry.ToPaintGeometry(),
+			GradientStops,
+			Orientation,
+			DrawEffects
+		);
+	}
+
+	return LayerId + 1;
+}
diff --git a/Engine/Source/Runtime/Slate/Private/Widgets/Colors/SSimpleGradient.cpp b/Engine/Source/Runtime/Slate/Private/Widgets/Colors/SSimpleGradient.cpp
new file mode 100644
--- /dev/null
+++ b/Engine/Source/Runtime/Slate/Private/Widgets/Colors/SSimpleGradient.cpp
@@ -0,0 +1,56 @@
+// Copyright Epic Games, Inc. All Rights Reserved.
+
+#include "Widgets/Colors/SSimpleGradient.h"
+#include "Rendering/DrawElements.h"
+#include "Styling/CoreStyle.h"
+#include "Application/SlateWindowHelper.h"
+
+
+/* SSimpleGradient interface
+ *****************************************************************************/
+
+void SSimpleGradient::Construct( const FArguments& InArgs )
+{
+	StartColor = InArgs._StartColor;
+	EndColor = InArgs._EndColor;
+	bHasAlphaBackground = InArgs._HasAlphaBackground.Get();
+	Orientation = InArgs._Orientation.Get();
+}
+
+
+/* SCompoundWidget overrides
+ *****************************************************************************/
+
+int32 SSimpleGradient::OnPaint( const FPaintArgs& Args, const FGeometry& AllottedGeometry, const FSlateRect& MyCullingRect, FSlateWindowElementList& OutDrawElements, int32 LayerId, const FWidgetStyle& InWidgetStyle, bool bParentEnabled ) const
+{
+	const ESlateDrawEffect DrawEffects = this->ShouldBeEnabled(bParentEnabled) ? ESlateDrawEffect::None : ESlateDrawEffect::DisabledEffect;
+
+	if (bHasAlphaBackground)
+	{
+		const FSlateBrush* StyleInfo = FCoreStyle::Get().GetBrush("ColorPicker.AlphaBackground");
+
+		FSlateDrawElement::MakeBox(
+			OutDrawElements,
+			LayerId,
+			AllottedGeometry.ToPaintGeometry(),
+			StyleInfo,
+			DrawEffects
+		);
+	}
+
+	TArray<FSlateGradientStop> GradientStops;
+
+	GradientStops.Add(FSlateGradientStop(FVector2D::ZeroVector, StartColor.Get()));
+	GradientStops.Add(FSlateGradientStop(AllottedGeometry.GetLocalSize(), EndColor.Get()));
+
+	FSlateDrawElement::MakeGradient(
+		OutDrawElements,
+		LayerId + 1,
+		AllottedGeometry.ToPaintGeometry(),
+		GradientStops,
+		Orientation,
+		DrawEffects | ESlateDrawEffect::NoGamma
+	);
+
+	return LayerId + 1;
+}
diff --git a/Engine/Source/Runtime/Slate/Public/Widgets/Colors/SComplexGradient.h b/Engine/Source/Runtime/Slate/Public/Widgets/Colors/SComplexGradient.h
new file mode 100644
--- /dev/null
+++ b/Engine/Source/Runtime/Slate/Public/Widgets/Colors/SComplexGradient.h
@@ -0,0 +1,61 @@
+// Copyright Epic Games, Inc. All Rights Reserved.
+
+#pragma once
+
+#include "CoreMinimal.h"
+#include "Misc/Attribute.h"
+#include "Widgets/DeclarativeSyntaxSupport.h"
+#include "Widgets/SCompoundWidget.h"
+
+class FPaintArgs;
+class FSlateWindowElementList;
+
+/**
+ * Implements a Slate widget that renders a color gradient consisting of multiple stops.
+ */
+class SLATE_API SComplexGradient
+	: public SCompoundWidget
+{
+public:
+
+	SLATE_BEGIN_ARGS(SComplexGradient)
+		: _GradientColors()
+		, _HasAlphaBackground(false)
+		, _Orientation(Orient_Vertical)
+	{ }
+
+		/** The colors used in the gradient */
+		SLATE_ATTRIBUTE(TArray<FLinearColor>, GradientColors)
+
+		/** Whether a checker background is displayed for alpha viewing */
+		SLATE_ATTRIBUTE(bool, HasAlphaBackground)
+		
+		/** Horizontal or vertical gradient */
+		SLATE_ATTRIBUTE(EOrientation, Orientation)
+
+	SLATE_END_ARGS()
+
+	/**
+	 * Construct this widget
+	 *
+	 * @param InArgs The declaration data for this widget.
+	 */
+	void Construct( const FArguments& InArgs );
+
+protected:
+
+	// SCompoundWidget overrides
+
+	virtual int32 OnPaint( const FPaintArgs& Args, const FGeometry& AllottedGeometry, const FSlateRect& MyCullingRect, FSlateWindowElementList& OutDrawElements, int32 LayerId, const FWidgetStyle& InWidgetStyle, bool bParentEnabled ) const override;
+	
+private:
+
+	/** The colors used in the gradient. */
+	TAttribute<TArray<FLinearColor>> GradientColors;
+
+	/** Whether a checker background is displayed for alpha viewing. */
+	bool bHasAlphaBackground;
+	
+	/** Horizontal or vertical gradient. */
+	EOrientation Orientation;
+};
diff --git a/Engine/Source/Runtime/Slate/Public/Widgets/Colors/SSimpleGradient.h b/Engine/Source/Runtime/Slate/Public/Widgets/Colors/SSimpleGradient.h
new file mode 100644
--- /dev/null
+++ b/Engine/Source/Runtime/Slate/Public/Widgets/Colors/SSimpleGradient.h
@@ -0,0 +1,75 @@
+// Copyright Epic Games, Inc. All Rights Reserved.
+
+#pragma once
+
+#include "CoreMinimal.h"
+#include "Misc/Attribute.h"
+#include "Widgets/DeclarativeSyntaxSupport.h"
+#include "Widgets/SCompoundWidget.h"
+
+class FPaintArgs;
+class FSlateWindowElementList;
+
+/**
+ * Implements a Slate widget that renders a simple color gradient.
+ */
+class SLATE_API SSimpleGradient
+	: public SCompoundWidget
+{
+public:
+
+		SLATE_BEGIN_ARGS(SSimpleGradient)
+			: _StartColor(FLinearColor(0.0f, 0.0f, 0.0f))
+			, _EndColor(FLinearColor(1.0f, 1.0f, 1.0f))
+			, _HasAlphaBackground(false)
+			, _Orientation(Orient_Vertical)
+			, _UseSRGB(true)
+		{ }
+
+		/** The leftmost gradient color */
+		SLATE_ATTRIBUTE(FLinearColor, StartColor)
+		
+		/** The rightmost gradient color */
+		SLATE_ATTRIBUTE(FLinearColor, EndColor)
+
+		/** Whether a checker background is displayed for alpha viewing */
+		SLATE_ATTRIBUTE(bool, HasAlphaBackground)
+
+		/** Horizontal or vertical gradient */
+		SLATE_ATTRIBUTE(EOrientation, Orientation)
+
+		/** Whether to display sRGB color */
+		SLATE_ATTRIBUTE(bool, UseSRGB)
+
+	SLATE_END_ARGS()
+
+	/**
+	 * Constructs the widget
+	 *
+	 * @param InArgs The declaration data for this widget
+	 */
+	void Construct( const FArguments& InArgs );
+
+protected:
+
+	// SCompoundWidget overrides
+
+	virtual int32 OnPaint( const FPaintArgs& Args, const FGeometry& AllottedGeometry, const FSlateRect& MyCullingRect, FSlateWindowElementList& OutDrawElements, int32 LayerId, const FWidgetStyle& InWidgetStyle, bool bParentEnabled ) const override;
+
+private:
+
+	/** The leftmost gradient color */
+	TAttribute<FLinearColor> StartColor;
+
+	/** The rightmost gradient color */
+	TAttribute<FLinearColor> EndColor;
+
+	/** Whether a checker background is displayed for alpha viewing */
+	bool bHasAlphaBackground;
+	
+	/** Horizontal or vertical gradient */
+	EOrientation Orientation;
+
+	/** Whether to display sRGB color */
+	TAttribute<bool> UseSRGB;
+};
