Complex gradients now take in a desired size override so every one of them does not have to be wrapped in a box to specify size

diff --git a/Engine/Source/Runtime/Slate/Private/Widgets/Colors/SComplexGradient.cpp b/Engine/Source/Runtime/Slate/Private/Widgets/Colors/SComplexGradient.cpp
--- a/Engine/Source/Runtime/Slate/Private/Widgets/Colors/SComplexGradient.cpp
+++ b/Engine/Source/Runtime/Slate/Private/Widgets/Colors/SComplexGradient.cpp
@@ -6,11 +6,17 @@
 #include "Application/SlateWindowHelper.h"
 
 
+SComplexGradient::SComplexGradient()
+{
+	SetCanTick(false);
+}
+
 /* SComplexGradient interface
  *****************************************************************************/
 
 void SComplexGradient::Construct( const FArguments& InArgs )
 {
+	DesiredSizeOverride = InArgs._DesiredSizeOverride;
 	GradientColors = InArgs._GradientColors;
 	bHasAlphaBackground = InArgs._HasAlphaBackground.Get();
 	Orientation = InArgs._Orientation.Get();
@@ -61,3 +67,14 @@ int32 SComplexGradient::OnPaint( const FPaintArgs& Args, const FGeometry& Allott
 
 	return LayerId + 1;
 }
+
+FVector2D SComplexGradient::ComputeDesiredSize(float LayoutScaleMultiplier) const
+{
+	const TOptional<FVector2D>& CurrentSizeOverride = DesiredSizeOverride.Get();
+	if (CurrentSizeOverride.IsSet())
+	{
+		return CurrentSizeOverride.GetValue();
+	}
+	
+	return SCompoundWidget::ComputeDesiredSize(LayoutScaleMultiplier);
+}
diff --git a/Engine/Source/Runtime/Slate/Public/Widgets/Colors/SComplexGradient.h b/Engine/Source/Runtime/Slate/Public/Widgets/Colors/SComplexGradient.h
--- a/Engine/Source/Runtime/Slate/Public/Widgets/Colors/SComplexGradient.h
+++ b/Engine/Source/Runtime/Slate/Public/Widgets/Colors/SComplexGradient.h
@@ -22,7 +22,9 @@ public:
 		: _GradientColors()
 		, _HasAlphaBackground(false)
 		, _Orientation(Orient_Vertical)
-	{ }
+	{
+		_Visibility = EVisibility::SelfHitTestInvisible;
+	}
 
 		/** The colors used in the gradient */
 		SLATE_ATTRIBUTE(TArray<FLinearColor>, GradientColors)
@@ -33,7 +35,12 @@ public:
 		/** Horizontal or vertical gradient */
 		SLATE_ATTRIBUTE(EOrientation, Orientation)
 
+		/** When specified use this as the gradients desired size */
+		SLATE_ATTRIBUTE(TOptional<FVector2D>, DesiredSizeOverride)
+
 	SLATE_END_ARGS()
+	
+	SComplexGradient();
 
 	/**
 	 * Construct this widget
@@ -47,12 +54,15 @@ protected:
 	// SCompoundWidget overrides
 
 	virtual int32 OnPaint( const FPaintArgs& Args, const FGeometry& AllottedGeometry, const FSlateRect& MyCullingRect, FSlateWindowElementList& OutDrawElements, int32 LayerId, const FWidgetStyle& InWidgetStyle, bool bParentEnabled ) const override;
-	
+	virtual FVector2D ComputeDesiredSize(float LayoutScaleMultiplier) const override;
 private:
 
 	/** The colors used in the gradient. */
 	TAttribute<TArray<FLinearColor>> GradientColors;
 
+	/** Optiona override for desired size */
+	TAttribute<TOptional<FVector2D>> DesiredSizeOverride;
+
 	/** Whether a checker background is displayed for alpha viewing. */
 	bool bHasAlphaBackground;
 	
