Fixed nomad tabs having the style of minor tabs when active and docked with major tabs

diff --git a/Engine/Source/Runtime/Slate/Private/Widgets/Docking/SDockTab.cpp b/Engine/Source/Runtime/Slate/Private/Widgets/Docking/SDockTab.cpp
--- a/Engine/Source/Runtime/Slate/Private/Widgets/Docking/SDockTab.cpp
+++ b/Engine/Source/Runtime/Slate/Private/Widgets/Docking/SDockTab.cpp
@@ -731,7 +731,7 @@ const FSlateBrush* SDockTab::GetColorOverlayImageBrush() const
 
 EVisibility SDockTab::GetActiveTabIndicatorVisibility() const
 {
-	return IsActive() ? EVisibility::HitTestInvisible : EVisibility::Collapsed;
+	return IsActive() && GetVisualTabRole() != ETabRole::MajorTab ? EVisibility::HitTestInvisible : EVisibility::Collapsed;
 }
 
 FSlateColor SDockTab::GetTabColor() const
