Fix Mac Startup Crash

diff --git a/Engine/Source/Runtime/Slate/Public/Framework/Application/SWindowTitleBar.h b/Engine/Source/Runtime/Slate/Public/Framework/Application/SWindowTitleBar.h
--- a/Engine/Source/Runtime/Slate/Public/Framework/Application/SWindowTitleBar.h
+++ b/Engine/Source/Runtime/Slate/Public/Framework/Application/SWindowTitleBar.h
@@ -56,19 +56,23 @@ class SAppIconWidget
 	{
 		SetCanTick(false);
 
+		IconImage = FSlateApplicationBase::Get().MakeImage(
+			FSlateApplicationBase::Get().GetAppIcon(),
+			Args._IconColorAndOpacity,
+			EVisibility::HitTestInvisible
+		);
+
 		this->ChildSlot
-			[
-				SNew(SVerticalBox)
-				+ SVerticalBox::Slot()
-				.Padding( FMargin( 0.0f, 0.0f, 2.0f, 0.0f ) )
-				[
-					FSlateApplicationBase::Get().MakeImage(
-						FSlateApplicationBase::Get().GetAppIcon(),
-						Args._IconColorAndOpacity,
-						EVisibility::HitTestInvisible
-					)
-				]
-			];
+		.Padding(FAppStyle::Get().GetMargin("AppIconPadding"))
+		[
+			IconImage.ToSharedRef()
+		];
+	}
+
+	void SetIcon(const FSlateBrush* InIcon, const FMargin& InIconPadding)
+	{
+		ChildSlot.Padding(InIconPadding);
+		IconImage->SetImage(InIcon);
 	}
 
 	virtual EWindowZone::Type GetWindowZoneOverride() const override
@@ -76,6 +80,9 @@ class SAppIconWidget
 		// Pretend we are a REAL system menu so the user can click to open a menu, or double-click to close the app on Windows
 		return EWindowZone::SysMenu;
 	}
+
+private:
+	TSharedPtr<SImage> IconImage;
 };
 
 
@@ -116,6 +123,8 @@ public:
 		ShowAppIcon = InArgs._ShowAppIcon;
 		Title = InArgs._Title;
 
+		WindowMenuSlot = nullptr;
+
 		if (!Title.IsSet() && !Title.IsBound())
 		{
 			Title = TAttribute<FText>::Create(TAttribute<FText>::FGetter::CreateSP(this, &SWindowTitleBar::HandleWindowTitleText));
@@ -133,9 +142,9 @@ public:
 				+ SOverlay::Slot()
 				[
 					SNew(SImage)
-						.Visibility(this, &SWindowTitleBar::GetWindowFlashVisibility)
-						.Image(&Style->FlashTitleBrush )
-						.ColorAndOpacity(this, &SWindowTitleBar::GetWindowTitleAreaColor)
+					.Visibility(this, &SWindowTitleBar::GetWindowFlashVisibility)
+					.Image(&Style->FlashTitleBrush )
+					.ColorAndOpacity(this, &SWindowTitleBar::GetWindowTitleAreaColor)
 				]
 				+ SOverlay::Slot()
 				[
@@ -158,6 +167,63 @@ public:
 		TitleFlashSequence.Play(this->AsShared());
 	}
 
+	virtual void UpdateWindowMenu(TSharedPtr<SWidget> MenuContent) override
+	{
+		if(MenuContent.IsValid())
+		{
+			(*WindowMenuSlot)
+			[
+				MenuContent.ToSharedRef()
+			];
+		}
+		else
+		{
+			(*WindowMenuSlot)
+			[
+				SNullWidget::NullWidget
+			];
+		}
+	}
+
+	virtual void UpdateBackgroundContent(TSharedPtr<SWidget> BackgroundContent) override
+	{
+		if (BackgroundContent.IsValid())
+		{
+            if (RightSideContentSlot != nullptr)
+            {
+                (*RightSideContentSlot)
+                [
+                    BackgroundContent.ToSharedRef()
+                ];
+            }
+		}
+		else
+		{
+            if (RightSideContentSlot != nullptr)
+            {
+                (*RightSideContentSlot)
+                [
+                    SNullWidget::NullWidget
+                ];
+            }
+		}
+	}
+
+	virtual void SetUseLargeIcon(bool bUseLargeIcon) override
+	{
+		if(AppIconWidget)
+		{
+			if (bUseLargeIcon)
+			{
+				AppIconWidget->SetIcon(FSlateApplicationBase::Get().GetAppIcon(), FAppStyle::Get().GetMargin("AppIconPadding", nullptr, FMargin(0)));
+			}
+			else
+			{
+				AppIconWidget->SetIcon(FSlateApplicationBase::Get().GetAppIconSmall(), FAppStyle::Get().GetMargin("AppIconPadding.Small", nullptr, FMargin(0)));
+
+			}
+		}
+	}
 	//~ Begin IWindowTitleBar Interface
 	
 protected:
@@ -257,11 +323,28 @@ protected:
 		// Windows UI layout
 		if (ShowAppIcon && bHasWindowButtons)
 		{
-			OutLeftContent = SNew(SAppIconWidget)
-				.IconColorAndOpacity(this, &SWindowTitleBar::GetWindowTitleContentColor);
+			OutLeftContent = 
+				SNew(SHorizontalBox)
+				+ SHorizontalBox::Slot()
+				.AutoWidth()
+				.HAlign(HAlign_Left)
+				.VAlign(VAlign_Top)
+				[
+					SAssignNew(AppIconWidget, SAppIconWidget)
+					.IconColorAndOpacity(this, &SWindowTitleBar::GetWindowTitleContentColor)
+				]
+				+ SHorizontalBox::Slot()
+				.VAlign(VAlign_Top)
+				.FillWidth(1)
+				.Expose(WindowMenuSlot);
+
+			// Default everything to use thge small icon unless specifically set to use the large icon
+			SetUseLargeIcon(false);
 		}
 		else
 		{
+			WindowMenuSlot = nullptr;
+
 			OutLeftContent = SNew(SSpacer);
 		}
 
@@ -273,31 +356,37 @@ protected:
 				[
 					// Minimize
 					SNew(SHorizontalBox)
-						.Visibility(EVisibility::SelfHitTestInvisible)
+					+ SHorizontalBox::Slot()
+					.Expose(RightSideContentSlot)
 
 					+ SHorizontalBox::Slot()
-						.AutoWidth()
-						[
-							MinimizeButton.ToSharedRef()
-						]
+					.VAlign(VAlign_Top)
+					.AutoWidth()
+					[
+						MinimizeButton.ToSharedRef()
+					]
 
 					// Maximize/Restore
 					+ SHorizontalBox::Slot()
-						.AutoWidth()
-						[
-							MaximizeRestoreButton.ToSharedRef()
-						]
+					.VAlign(VAlign_Top)
+					.AutoWidth()
+					[
+						MaximizeRestoreButton.ToSharedRef()
+					]
 
 					// Close button
 					+ SHorizontalBox::Slot()
-						.AutoWidth()
-						[
-							CloseButton.ToSharedRef()
-						]
+					.VAlign(VAlign_Top)
+					.AutoWidth()
+					[
+						CloseButton.ToSharedRef()
+					]
 				];
 		}
 		else
 		{
+			RightSideContentSlot = nullptr;
+
 			OutRightContent = SNew(SSpacer);
 		}
 
@@ -358,68 +447,65 @@ protected:
 		float SpacerHeight = FMath::Max(LeftSize.Y, RightSize.Y);
 
 		// create title bar
-		return SAssignNew(TitleArea, SBox)
+		return 
+			SAssignNew(TitleArea, SBox)
 			.Visibility(EVisibility::SelfHitTestInvisible)
 			[
 				SNew(SOverlay)
+				.Visibility(EVisibility::SelfHitTestInvisible)
+				+ SOverlay::Slot()
+				[
+					SNew(SHorizontalBox)
 					.Visibility(EVisibility::SelfHitTestInvisible)
+						
+					+SHorizontalBox::Slot()
+					.AutoWidth()
+					.HAlign(HAlign_Left)
+					.VAlign(VAlign_Top)
+					[
+						SNew(SSpacer)
+						.Size(FVector2D(LeftSize.X, SpacerHeight))
+					]
 
+					+ SHorizontalBox::Slot()
+					.HAlign(CenterContentAlignment)
+					.VAlign(VAlign_Center)
+					.FillWidth(1.0f)
+					[
+						CenterContent.ToSharedRef()
+					]
+					+ SHorizontalBox::Slot()
+					.AutoWidth()
+					.HAlign(HAlign_Right)
+					.VAlign(VAlign_Top)
+					[
+						SNew(SSpacer)
+						.Size(FVector2D(RightSize.X, SpacerHeight))
+					]
+				]
 				+ SOverlay::Slot()
+				[
+					SNew(SHorizontalBox)
+					.Visibility(EVisibility::SelfHitTestInvisible)
+					+ SHorizontalBox::Slot()
+					.AutoWidth()
+					.HAlign(HAlign_Left)
+					.VAlign(VAlign_Top)
 					[
-						SNew(SHorizontalBox)
-							.Visibility(EVisibility::SelfHitTestInvisible)
-						
-						+ SHorizontalBox::Slot()
-							.AutoWidth()
-							.HAlign(HAlign_Left)
-							.VAlign(VAlign_Top)
-							[
-								SNew(SSpacer)
-									.Size(FVector2D(LeftSize.X, SpacerHeight))
-							]
-
-						+ SHorizontalBox::Slot()
-							.HAlign(CenterContentAlignment)
-							.VAlign(VAlign_Center)
-							.FillWidth(1.0f)
-							[
-								CenterContent.ToSharedRef()
-							]
-
-						+ SHorizontalBox::Slot()
-							.AutoWidth()
-							.HAlign(HAlign_Right)
-							.VAlign(VAlign_Top)
-							[
-								SNew(SSpacer)
-									.Size(FVector2D(RightSize.X, SpacerHeight))
-							]
+						LeftContent.ToSharedRef()
 					]
 
-				+ SOverlay::Slot()
+					+ SHorizontalBox::Slot()
+					.FillWidth(1.0f)
+
+					+ SHorizontalBox::Slot()
+					.AutoWidth()
+					.HAlign(HAlign_Right)
+					.VAlign(VAlign_Top)
 					[
-						SNew(SHorizontalBox)
-							.Visibility(EVisibility::SelfHitTestInvisible)
-
-						+ SHorizontalBox::Slot()
-							.AutoWidth()
-							.HAlign(HAlign_Left)
-							.VAlign(VAlign_Top)
-							[
-								LeftContent.ToSharedRef()
-							]
-
-						+ SHorizontalBox::Slot()
-							.FillWidth(1.0f)
-
-						+ SHorizontalBox::Slot()
-							.AutoWidth()
-							.HAlign(HAlign_Right)
-							.VAlign(VAlign_Top)
-							[
-								RightContent.ToSharedRef()
-							]
+						RightContent.ToSharedRef()
 					]
+				]
 			];
 	}
 
@@ -667,7 +753,12 @@ private:
 	// Holds the close button.
 	TSharedPtr<SButton> CloseButton;
 
+	TAttribute<FText> Title;
+
+	SHorizontalBox::FSlot* WindowMenuSlot;
+	SHorizontalBox::FSlot* RightSideContentSlot;
+
+	TSharedPtr<SAppIconWidget> AppIconWidget;
 	bool ShowAppIcon;
 
-	TAttribute<FText> Title;
-};
+}; 
