Missed a file

diff --git a/Engine/Source/Runtime/Slate/Private/Framework/Docking/SDockingTabStack.cpp b/Engine/Source/Runtime/Slate/Private/Framework/Docking/SDockingTabStack.cpp
--- a/Engine/Source/Runtime/Slate/Private/Framework/Docking/SDockingTabStack.cpp
+++ b/Engine/Source/Runtime/Slate/Private/Framework/Docking/SDockingTabStack.cpp
@@ -15,6 +15,7 @@
 #include "Framework/Docking/SDockingCross.h"
 #include "Framework/Docking/FDockingDragOperation.h"
 #include "Framework/Docking/TabCommands.h"
+#include "Brushes/SlateColorBrush.h"
 
 #define LOCTEXT_NAMESPACE "DockTabStack"
 
@@ -51,12 +52,14 @@ void SDockingTabStack::Construct( const FArguments& InArgs, const TSharedRef<FTa
 	Tabs = PersistentNode->Tabs;
 	this->SetSizeCoefficient(PersistentNode->GetSizeCoefficient());
 
+	// the value of this is determined every time a tab is added
+	bShowingTitleBarArea = false;
+
 	bIsDocumentArea = InArgs._IsDocumentArea;
 
-	InlineContentAreaLeft = NULL;
-	InlineContentAreaRight = NULL;
-	BackgroundContentArea = NULL;
-	TitleBarSlot = NULL;
+	InlineContentAreaLeft = nullptr;
+	InlineContentAreaRight = nullptr;
+	TitleBarSlot = nullptr;
 
 	this->TabStackGeometry = FGeometry();
 
@@ -74,7 +77,7 @@ void SDockingTabStack::Construct( const FArguments& InArgs, const TSharedRef<FTa
 		
 	}
 
-	// In TabStack mode we glue together a TabWell, two InlineContent areas and a ContentOverlay
+	// In TabStack mode we glue together a TabWell, InlineContent areas and a ContentOverlay
 	// that shows the content of the currently selected Tab.
 	//                                         ________ TabWell
 	//                                        |
@@ -93,11 +96,8 @@ void SDockingTabStack::Construct( const FArguments& InArgs, const TSharedRef<FTa
 	// create inline title bar content
 	TitleBarContent = 
 	SNew(SOverlay)
-	+ SOverlay::Slot().Expose(BackgroundContentArea)
 	+ SOverlay::Slot()
 	[
-	
-		
 		SNew(SHorizontalBox)
 		.Visibility(EVisibility::SelfHitTestInvisible)
 
@@ -112,7 +112,6 @@ void SDockingTabStack::Construct( const FArguments& InArgs, const TSharedRef<FTa
 		[
 			SNew(SVerticalBox)
 			.Visibility(EVisibility::SelfHitTestInvisible)
-
 			+ SVerticalBox::Slot()
 			.AutoHeight()
 			[
@@ -120,7 +119,6 @@ void SDockingTabStack::Construct( const FArguments& InArgs, const TSharedRef<FTa
 				.Visibility(this, &SDockingTabStack::GetMaximizeSpacerVisibility)
 				.Size(FVector2D(0.0f, 10.0f))
 			]
-
 			+ SVerticalBox::Slot()
 			.AutoHeight()
 			[
@@ -134,7 +132,7 @@ void SDockingTabStack::Construct( const FArguments& InArgs, const TSharedRef<FTa
 		.AutoWidth()
 		.Expose(InlineContentAreaRight)
 		.Padding(5.0f, 0.0f, 0.0f, 0.0f)
-		.VAlign((VAlign_Center))
+		.VAlign(VAlign_Center)
 	];
 
 	ChildSlot
@@ -149,7 +147,7 @@ void SDockingTabStack::Construct( const FArguments& InArgs, const TSharedRef<FTa
 			SNew(SBorder)
 			.Visibility(this, &SDockingTabStack::GetTabWellVisibility)
 			.DesiredSizeScale(this, &SDockingTabStack::GetTabWellScale)
-			.BorderImage(FCoreStyle::Get().GetBrush("NoBorder"))
+			.BorderImage(this, &SDockingTabStack::GetTabStackBorderImage)
 			.VAlign(VAlign_Bottom)
 			.OnMouseButtonDown(this, &SDockingTabStack::TabWellRightClicked)
 			.Padding(0.0f)
@@ -230,7 +228,7 @@ void SDockingTabStack::Construct( const FArguments& InArgs, const TSharedRef<FTa
 
 	if (bIsDocumentArea)
 	{
-		this->SetNodeContent(SDocumentAreaWidget::MakeDocumentAreaWidget(), SNullWidget::NullWidget, SNullWidget::NullWidget, SNullWidget::NullWidget);
+		this->SetNodeContent(SDocumentAreaWidget::MakeDocumentAreaWidget(), FDockingStackOptionalContent());
 	}
 }
 
@@ -241,11 +239,11 @@ void SDockingTabStack::OnLastTabRemoved()
 	{
 		// Stop holding onto any meaningful window content.
 		// The user should not see any content in this DockNode.
-		this->SetNodeContent(SNullWidget::NullWidget, SNullWidget::NullWidget, SNullWidget::NullWidget, SNullWidget::NullWidget);
+		this->SetNodeContent(SNullWidget::NullWidget, FDockingStackOptionalContent());
 	}
 	else
 	{
-		this->SetNodeContent(SDocumentAreaWidget::MakeDocumentAreaWidget(), SNullWidget::NullWidget, SNullWidget::NullWidget, SNullWidget::NullWidget);
+		this->SetNodeContent(SDocumentAreaWidget::MakeDocumentAreaWidget(), FDockingStackOptionalContent());
 	}
 }
 
@@ -350,12 +348,20 @@ void SDockingTabStack::BringToFront( const TSharedRef<SDockTab>& TabToBringToFro
 	TabWell->BringTabToFront(TabToBringToFront);
 }
 
-void SDockingTabStack::SetNodeContent(const TSharedRef<SWidget>& InContent, const TSharedRef<SWidget>& ContentLeft, const TSharedRef<SWidget>& ContentRight, const TSharedRef<SWidget>& InContentBackground)
+void SDockingTabStack::SetNodeContent(const TSharedRef<SWidget>& InContent, const FDockingStackOptionalContent& OptionalContent)
 {
 	ContentSlot->SetContent(InContent);
-	(*InlineContentAreaLeft)[ContentLeft];
-	(*InlineContentAreaRight)[ContentRight];
-	(*BackgroundContentArea)[InContentBackground];
+	(*InlineContentAreaLeft)[OptionalContent.ContentLeft];
+	(*InlineContentAreaRight)[OptionalContent.ContentRight];
+
+	if(TabWell->GetForegroundTab())
+	{
+		if (TSharedPtr<SWindow> ParentWindow = TabWell->GetForegroundTab()->GetParentWindow())
+		{
+			ParentWindow->GetTitleBar()->UpdateBackgroundContent(OptionalContent.TitleBarContentRight);
+		}
+	}
+
 }
 
 FReply SDockingTabStack::OnDragOver( const FGeometry& MyGeometry, const FDragDropEvent& DragDropEvent )
@@ -761,22 +767,24 @@ TSharedPtr<FTabManager::FLayoutNode> SDockingTabStack::GatherPersistentLayout()
 
 void SDockingTabStack::ClearReservedSpace()
 {
+	bShowingTitleBarArea = false;
 	(*TitleBarSlot).Padding(0);
 }
 
-
-void SDockingTabStack::ReserveSpaceForWindowChrome(EChromeElement Element)
+void SDockingTabStack::ReserveSpaceForWindowChrome(EChromeElement Element, bool bIncludePaddingForMenuBar, bool bOnlyMinorTabs)
 {
-	// @todo: It would be nice to read these from the desired size of the title bar
-	//        instead of hard-coding, but at least the long-standing TTP is fixed!
 	#if PLATFORM_MAC
 		static const FMargin ControlsPadding = FMargin(64, 0, 0, 0);
 		static const FMargin IconPadding = FMargin(0);
 	#else
-		static const FMargin ControlsPadding = FMargin(0, 0, 96, 0);
-		static const FMargin IconPadding = FMargin(32, 0, 0, 0);
+		static const float TopPaddingForMenuBar = 27.0f;
+		static const float LeftPaddingForIcon = FSlateApplication::Get().GetAppIcon()->GetImageSize().X;
+		const FMargin ControlsPadding = FMargin(8, bIncludePaddingForMenuBar ? TopPaddingForMenuBar : 5, 0, 0);
+		// If we are including top padding for the menu bar we do not need to pad the left side since we will be below the left icon
+		const FMargin IconPadding = FMargin(bIncludePaddingForMenuBar ? LeftPaddingForIcon + 12 : 25, bOnlyMinorTabs ? 5 : 0, 0, 0);
 	#endif
 
+	bShowingTitleBarArea = true;
 	const FMargin CurrentPadding = TitleBarSlot->SlotPadding.Get();
 	switch (Element)
 	{
@@ -837,7 +845,7 @@ FMargin SDockingTabStack::GetContentPadding() const
 	TSharedPtr<SDockTab> ForegroundTab = TabWell->GetForegroundTab();
 	return (ForegroundTab.IsValid())
 		? ForegroundTab->GetContentPadding()
-		: FMargin(2);
+		: FMargin(0);
 }
 
 EVisibility SDockingTabStack::GetTabWellVisibility() const
@@ -993,6 +1001,12 @@ FSlateColor SDockingTabStack::GetUnhideTabWellButtonOpacity() const
 	return FLinearColor( 1,1,1, 1.0f - ShowHideTabWell.GetLerp() );
 }
 
+const FSlateBrush* SDockingTabStack::GetTabStackBorderImage() const
+{
+	static const FSlateColorBrush MajorTabBackgroundBrush(FAppStyle::Get().GetSlateColor("Colors.Title"));
+	return bShowingTitleBarArea ? &MajorTabBackgroundBrush : FStyleDefaults::GetNoBrush();
+}
+
 int32 SDockingTabStack::OpenPersistentTab( const FTabId& TabId, int32 OpenLocationAmongActiveTabs )
 {
 	const int32 ExistingClosedTabIndex = Tabs.IndexOfByPredicate(FTabMatcher(TabId, ETabState::ClosedTab));
@@ -1084,6 +1098,7 @@ void SDockingTabStack::RemovePersistentTab( const FTabId& TabId )
 
 EVisibility SDockingTabStack::GetMaximizeSpacerVisibility() const
 {
+/*
 	if(GetDockArea().IsValid() && GetDockArea()->GetParentWindow().IsValid())
 	{
 		if (GetDockArea()->GetParentWindow()->IsWindowMaximized())
@@ -1094,8 +1109,7 @@ EVisibility SDockingTabStack::GetMaximizeSpacerVisibility() const
 		{
 			return EVisibility::SelfHitTestInvisible;
 		}
-	}
-
+	}*/
 	return EVisibility::Collapsed;
 }
 
