Fix a few missing FStyleColor usages that should be using FSlateColor

diff --git a/Engine/Source/Editor/UnrealEd/Private/SEditorViewport.cpp b/Engine/Source/Editor/UnrealEd/Private/SEditorViewport.cpp
--- a/Engine/Source/Editor/UnrealEd/Private/SEditorViewport.cpp
+++ b/Engine/Source/Editor/UnrealEd/Private/SEditorViewport.cpp
@@ -107,7 +107,7 @@ void SEditorViewport::Construct( const FArguments& InArgs )
 	];
 
 	// This makes a gradient that displays whether or not a viewport is active
-	FLinearColor ActiveBorderColor = FAppStyle::Get().GetColor("EditorViewport.ActiveBorderColor");
+	FLinearColor ActiveBorderColor = FAppStyle::Get().GetSlateColor("EditorViewport.ActiveBorderColor").GetSpecifiedColor();
 	FLinearColor ActiveBorderColorTransparent = ActiveBorderColor;
 	ActiveBorderColorTransparent.A = 0.0f;
 
diff --git a/Engine/Source/Runtime/Slate/Private/Widgets/Docking/SDockTab.cpp b/Engine/Source/Runtime/Slate/Private/Widgets/Docking/SDockTab.cpp
--- a/Engine/Source/Runtime/Slate/Private/Widgets/Docking/SDockTab.cpp
+++ b/Engine/Source/Runtime/Slate/Private/Widgets/Docking/SDockTab.cpp
@@ -525,7 +525,7 @@ void SDockTab::Construct( const FArguments& InArgs )
 	const FTextBlockStyle* const TabTextStyle = &GetCurrentStyle().TabTextStyle;
 
 	// This makes a gradient that displays whether or not a viewport is active
-	static FLinearColor ActiveBorderColor = FAppStyle::Get().GetColor("Docking.Tab.ActiveTabIndicatorColor");
+	static FLinearColor ActiveBorderColor = FAppStyle::Get().GetSlateColor("Docking.Tab.ActiveTabIndicatorColor").GetSpecifiedColor();
 	static FLinearColor ActiveBorderColorTransparent = FLinearColor(ActiveBorderColor.R, ActiveBorderColor.G, ActiveBorderColor.B, 0.0f);
 	static TArray<FLinearColor> GradientStops{ ActiveBorderColorTransparent, ActiveBorderColor, ActiveBorderColorTransparent };
 
