Menu Icon spacing and shading

diff --git a/Engine/Source/Runtime/Slate/Private/Framework/MultiBox/SMenuEntryBlock.cpp b/Engine/Source/Runtime/Slate/Private/Framework/MultiBox/SMenuEntryBlock.cpp
--- a/Engine/Source/Runtime/Slate/Private/Framework/MultiBox/SMenuEntryBlock.cpp
+++ b/Engine/Source/Runtime/Slate/Private/Framework/MultiBox/SMenuEntryBlock.cpp
@@ -439,8 +439,8 @@ TSharedRef< SWidget > SMenuEntryBlock::BuildMenuEntryWidget( const FMenuEntryBui
 		const FSlateBrush* IconBrush = ActualIcon.GetIcon();
 		if( IconBrush->GetResourceName() != NAME_None )
 		{
-			IconWidget =
-				SNew( SImage )
+			IconWidget = SNew( SImage )
+				.ColorAndOpacity( FSlateColor::UseSubduedForeground() )
 				.Image( IconBrush );
 		}
 	}
@@ -515,7 +515,7 @@ TSharedRef< SWidget > SMenuEntryBlock::BuildMenuEntryWidget( const FMenuEntryBui
 		// Whatever we have in the icon area goes first
 		+ SHorizontalBox::Slot()
 		.AutoWidth()
-		.Padding(FMargin(2, 0, 2, 0))
+		.Padding(FMargin(2, 0, 6, 0))
 		[
 			SNew( SBox )
 			.Visibility(IconWidget != SNullWidget::NullWidget ? EVisibility::Visible : EVisibility::Collapsed)
@@ -723,6 +723,7 @@ TSharedRef< SWidget> SMenuEntryBlock::BuildSubMenuWidget( const FMenuEntryBuildP
 		{
 			IconWidget =
 				SNew( SImage )
+				.ColorAndOpacity( FSlateColor::UseSubduedForeground() )
 				.Image( IconBrush );
 		}
 	}
@@ -764,7 +765,7 @@ TSharedRef< SWidget> SMenuEntryBlock::BuildSubMenuWidget( const FMenuEntryBuildP
 		// Whatever we have in the icon area goes first
 		+ SHorizontalBox::Slot()
 		.AutoWidth()
-		.Padding(FMargin(2, 0, 2, 0))
+		.Padding(FMargin(2, 0, 6, 0))
 		[
 			SNew( SBox )
 			.Visibility(IconWidget != SNullWidget::NullWidget ? EVisibility::Visible : EVisibility::Collapsed)
