The search filter for MultiBoxes now resets its search text when it gains focus.

diff --git a/Engine/Source/Runtime/Slate/Private/Framework/MultiBox/MultiBox.cpp b/Engine/Source/Runtime/Slate/Private/Framework/MultiBox/MultiBox.cpp
--- a/Engine/Source/Runtime/Slate/Private/Framework/MultiBox/MultiBox.cpp
+++ b/Engine/Source/Runtime/Slate/Private/Framework/MultiBox/MultiBox.cpp
@@ -1248,6 +1248,8 @@ FReply SMultiBoxWidget::FocusNextWidget(EUINavigation NavigationType)
 
 FReply SMultiBoxWidget::OnFocusReceived( const FGeometry& MyGeometry, const FFocusEvent& InFocusEvent )
 {
+	ResetSearch();
+
 	if (InFocusEvent.GetCause() == EFocusCause::Navigation)
 	{
 		// forward focus to children
