SLinkedBox - provides a way for widgets unrelated in the widget hierarchy to share a size

diff --git a/Engine/Source/Runtime/Slate/Public/Widgets/Layout/SLinkedBox.cpp b/Engine/Source/Runtime/Slate/Public/Widgets/Layout/SLinkedBox.cpp
new file mode 100644
--- /dev/null
+++ b/Engine/Source/Runtime/Slate/Public/Widgets/Layout/SLinkedBox.cpp
@@ -0,0 +1,108 @@
+// Copyright Epic Games, Inc. All Rights Reserved.
+
+#include "Widgets/Layout/SLinkedBox.h"
+
+#include "Types/PaintArgs.h"
+#include "Layout/ArrangedChildren.h"
+#include "Layout/LayoutUtils.h"
+#include "Widgets/Layout/SBorder.h"
+
+FLinkedBoxManager::FLinkedBoxManager()
+: CachedUniformSize(FVector2D::ZeroVector)
+{
+}
+
+FLinkedBoxManager::~FLinkedBoxManager()
+{
+}
+
+void FLinkedBoxManager::RegisterLinkedBox(SLinkedBox* InBox)
+{
+	Siblings.Add(InBox);	
+}
+
+void FLinkedBoxManager::UnregisterLinkedBox(SLinkedBox* InBox)
+{
+	Siblings.Remove(InBox);	
+}
+
+FVector2D FLinkedBoxManager::GetUniformCellSize() const
+{
+	if (FrameCounterLastCached == GFrameCounter)
+	{
+		return CachedUniformSize;
+	}
+
+	FVector2D MaxDesiredSize = FVector2D::ZeroVector;
+	for (auto Sibling : Siblings)
+	{
+		if (Sibling->GetVisibility() != EVisibility::Collapsed)
+		{
+			Sibling->CustomChildPrepass();
+			FVector2D SiblingDesiredSize = Sibling->GetChildrensDesiredSize();
+			if (!SiblingDesiredSize.IsZero())
+			{
+				MaxDesiredSize.X = FMath::Max( MaxDesiredSize.X, SiblingDesiredSize.X );
+				MaxDesiredSize.Y = FMath::Max( MaxDesiredSize.Y, SiblingDesiredSize.Y );
+			}
+		}
+	}
+
+	FrameCounterLastCached = GFrameCounter;
+	CachedUniformSize = MaxDesiredSize;
+	return CachedUniformSize;
+}
+
+SLinkedBox::SLinkedBox()
+{
+	bHasCustomPrepass = true;
+}
+
+SLinkedBox::~SLinkedBox()
+{
+	Manager->UnregisterLinkedBox(this);
+}
+
+void SLinkedBox::Construct( const FArguments& InArgs, TSharedRef<FLinkedBoxManager> InManager )
+{
+	Manager = InManager;
+	Manager->RegisterLinkedBox(this);
+
+	SBox::Construct( SBox::FArguments()
+		.HAlign(InArgs._HAlign)
+		.VAlign(InArgs._VAlign)
+		.Padding(InArgs._Padding)
+		.WidthOverride(InArgs._WidthOverride)
+		.HeightOverride(InArgs._HeightOverride)
+		.MinDesiredWidth(InArgs._MinDesiredWidth)
+		.MinDesiredHeight(InArgs._MinDesiredHeight)
+		.MaxDesiredWidth(InArgs._MaxDesiredWidth)
+		.MaxDesiredHeight(InArgs._MaxDesiredHeight)
+		.MinAspectRatio(InArgs._MinAspectRatio)
+		.MaxAspectRatio(InArgs._MaxAspectRatio)
+		[
+			InArgs._Content.Widget
+		]
+	);
+}
+
+void SLinkedBox::CustomChildPrepass()
+{
+	ChildSlot.GetWidget()->SlatePrepass();	
+}
+
+bool SLinkedBox::CustomPrepass(float LayoutScaleMultiplier) 
+{ 
+	return false; 
+}
+
+FVector2D SLinkedBox::GetChildrensDesiredSize() const
+{
+	return ChildSlot.GetWidget()->GetDesiredSize() + ChildSlot.SlotPadding.Get().GetDesiredSize();
+}
+
+FVector2D SLinkedBox::ComputeDesiredSize(float) const
+{
+	return Manager->GetUniformCellSize();
+}
+
diff --git a/Engine/Source/Runtime/Slate/Public/Widgets/Layout/SLinkedBox.h b/Engine/Source/Runtime/Slate/Public/Widgets/Layout/SLinkedBox.h
new file mode 100644
--- /dev/null
+++ b/Engine/Source/Runtime/Slate/Public/Widgets/Layout/SLinkedBox.h
@@ -0,0 +1,143 @@
+// Copyright Epic Games, Inc. All Rights Reserved.
+
+#pragma once
+
+#include "CoreMinimal.h"
+#include "Misc/Attribute.h"
+#include "Layout/Visibility.h"
+#include "SlotBase.h"
+#include "Containers/Ticker.h"
+#include "Layout/Margin.h"
+#include "Layout/Children.h"
+#include "Widgets/DeclarativeSyntaxSupport.h"
+#include "Widgets/SPanel.h"
+#include "Types/SlateEnums.h"
+#include "Widgets/Layout/SBox.h"
+	
+
+class FArrangedChildren;
+
+class SLinkedBox;
+
+/** 
+ * Manages a group of SLinkedBoxes that report the same size based on the largest size of any LinkedBoxes in this managed group.
+ *  
+ *  FLinkedBoxManager needs to be created and passed to each LinkedBox during construction.
+ */
+class SLATE_API FLinkedBoxManager : public TSharedFromThis<FLinkedBoxManager> 
+{
+
+public: 
+
+	FLinkedBoxManager();
+	~FLinkedBoxManager();
+
+	/* Add an SLinkedBox */
+	void RegisterLinkedBox(SLinkedBox* InBox);
+
+	/* Remove an SLinkedBox */
+	void UnregisterLinkedBox(SLinkedBox* InBox);
+
+	/* Used by the individual SLinkedBoxes to acquire the computed uniform size */
+	FVector2D GetUniformCellSize() const;
+
+protected:
+	
+	mutable uint32 FrameCounterLastCached = 0;
+	mutable FVector2D CachedUniformSize;
+
+private:
+
+	TSet< SLinkedBox* > Siblings;
+
+};
+
+/** A panel that */
+class SLATE_API SLinkedBox: public SBox
+{
+
+public:
+
+	SLATE_BEGIN_ARGS(SLinkedBox)
+		: _HAlign(HAlign_Fill)
+		, _VAlign(VAlign_Fill)
+		, _Padding(0.0f)
+		, _Content()
+		, _WidthOverride(FOptionalSize())
+		, _HeightOverride(FOptionalSize())
+		, _MinDesiredWidth(FOptionalSize())
+		, _MinDesiredHeight(FOptionalSize())
+		, _MaxDesiredWidth(FOptionalSize())
+		, _MaxDesiredHeight(FOptionalSize())
+	{
+		_Visibility = EVisibility::SelfHitTestInvisible;
+	}
+
+		/** Horizontal alignment of content in the area allotted to the SBox by its parent */
+	SLATE_ARGUMENT(EHorizontalAlignment, HAlign)
+
+	/** Vertical alignment of content in the area allotted to the SBox by its parent */
+	SLATE_ARGUMENT(EVerticalAlignment, VAlign)
+
+	/** Padding between the SBox and the content that it presents. Padding affects desired size. */
+	SLATE_ATTRIBUTE(FMargin, Padding)
+
+	/** The widget content presented by the SBox */
+	SLATE_DEFAULT_SLOT(FArguments, Content)
+
+	/** When specified, ignore the content's desired size and report the WidthOverride as the Box's desired width. */
+	SLATE_ATTRIBUTE(FOptionalSize, WidthOverride)
+
+	/** When specified, ignore the content's desired size and report the HeightOverride as the Box's desired height. */
+	SLATE_ATTRIBUTE(FOptionalSize, HeightOverride)
+
+	/** When specified, will report the MinDesiredWidth if larger than the content's desired width. */
+	SLATE_ATTRIBUTE(FOptionalSize, MinDesiredWidth)
+
+	/** When specified, will report the MinDesiredHeight if larger than the content's desired height. */
+	SLATE_ATTRIBUTE(FOptionalSize, MinDesiredHeight)
+
+	/** When specified, will report the MaxDesiredWidth if smaller than the content's desired width. */
+	SLATE_ATTRIBUTE(FOptionalSize, MaxDesiredWidth)
+
+	/** When specified, will report the MaxDesiredHeight if smaller than the content's desired height. */
+	SLATE_ATTRIBUTE(FOptionalSize, MaxDesiredHeight)
+
+	SLATE_ATTRIBUTE(FOptionalSize, MinAspectRatio)
+
+	SLATE_ATTRIBUTE(FOptionalSize, MaxAspectRatio)
+
+	SLATE_END_ARGS()
+
+	FVector2D GetChildrensDesiredSize() const;
+
+	virtual FVector2D ComputeDesiredSize(float) const override;
+
+	/**
+	 * Most panels do not create widgets as part of their implementation, so
+	 * they do not need to implement a Construct()
+	 */
+	void Construct( const FArguments& InArgs, TSharedRef<FLinkedBoxManager> InManager ); 
+
+	SLinkedBox();
+	~SLinkedBox();
+	
+protected:
+
+	/**
+	 *  CustomPrepass - Returns false so instead of
+	 *  each SLinkedBox prepass being called in the usual depth first order, 
+	 *  the Manager can call a prepass on all of the sibling 
+	 *  LinkedBoxes at once.  
+	*/
+	virtual bool CustomPrepass(float LayoutScaleMultiplier) override;
+
+	void CustomChildPrepass();
+
+	friend class FLinkedBoxManager;
+
+private:
+
+	TSharedPtr<FLinkedBoxManager> Manager;
+
+};
