Menu widths adjust with the presence of checks

diff --git a/Engine/Source/Runtime/Slate/Private/Framework/MultiBox/MultiBox.cpp b/Engine/Source/Runtime/Slate/Private/Framework/MultiBox/MultiBox.cpp
--- a/Engine/Source/Runtime/Slate/Private/Framework/MultiBox/MultiBox.cpp
+++ b/Engine/Source/Runtime/Slate/Private/Framework/MultiBox/MultiBox.cpp
@@ -508,6 +508,7 @@ bool FMultiBox::IsInEditMode() const
 
 void SMultiBoxWidget::Construct( const FArguments& InArgs )
 {
+	LinkedBoxManager = MakeShared<FLinkedBoxManager>();
 	ContentScale = InArgs._ContentScale;
 }
 
diff --git a/Engine/Source/Runtime/Slate/Private/Framework/MultiBox/SMenuEntryBlock.cpp b/Engine/Source/Runtime/Slate/Private/Framework/MultiBox/SMenuEntryBlock.cpp
--- a/Engine/Source/Runtime/Slate/Private/Framework/MultiBox/SMenuEntryBlock.cpp
+++ b/Engine/Source/Runtime/Slate/Private/Framework/MultiBox/SMenuEntryBlock.cpp
@@ -8,6 +8,7 @@
 #include "Widgets/Layout/SBorder.h"
 #include "Widgets/Images/SImage.h"
 #include "Widgets/Layout/SBox.h"
+#include "Widgets/Layout/SLinkedBox.h"
 #include "Widgets/Input/SButton.h"
 #include "Widgets/SToolTip.h"
 #include "Widgets/Colors/SColorBlock.h"
@@ -478,7 +479,7 @@ TSharedRef< SWidget > SMenuEntryBlock::BuildMenuEntryWidget( const FMenuEntryBui
 			UserInterfaceType == EUserInterfaceActionType::RadioButton ||
 			UserInterfaceType == EUserInterfaceActionType::Check ) ?
 				EVisibility::Visible :
-				EVisibility::Hidden;
+				EVisibility::Collapsed;
 
 	// Collapse (rather than Hide) the checkbox when using a custom menu widget with a button action, otherwise we add additional padding around the user defined widget
 	if ((MenuEntryBlock->EntryWidget.IsValid() && UserInterfaceType == EUserInterfaceActionType::Button) ||
@@ -489,7 +490,7 @@ TSharedRef< SWidget > SMenuEntryBlock::BuildMenuEntryWidget( const FMenuEntryBui
 	else if (MultiBox->IsInEditMode())
 	{
 		// Hide but keep spacing to avoid confusing user during editing
-		CheckBoxVisibility = EVisibility::Hidden;
+		CheckBoxVisibility = EVisibility::Collapsed;
 	}
 
 	TAttribute<FSlateColor> CheckBoxForegroundColor = FSlateColor::UseStyle();
@@ -567,19 +568,16 @@ TSharedRef< SWidget > SMenuEntryBlock::BuildMenuEntryWidget( const FMenuEntryBui
 		.Padding(FMargin(2, 0, 8, 0))
 		.AutoWidth()
 		[
-			SNew(SBox)
-			.Visibility(CheckBoxVisibility)
-			.WidthOverride(MenuIconSize)
-			.HeightOverride(MenuIconSize)
-			.HAlign(HAlign_Center)
-			.VAlign(VAlign_Center)
+			SNew(SLinkedBox, OwnerMultiBoxWidget.Pin()->GetLinkedBoxManager())
 			[
 				// For check style menus, use an image instead of a CheckBox because it can't really be checked.
 				UserInterfaceType == EUserInterfaceActionType::Check
 					? StaticCastSharedRef<SWidget>(SNew(SImage)
+						.Visibility(CheckBoxVisibility)
 						.ColorAndOpacity(FSlateColor::UseForeground())
 						.Image(this, &SMenuEntryBlock::GetCheckBoxImageBrushFromStyle, &StyleSet->GetWidgetStyle<FCheckBoxStyle>(CheckBoxStyle)))
 					: StaticCastSharedRef<SWidget>(SNew(SCheckBox)
+						.Visibility(CheckBoxVisibility)
 						.ForegroundColor( CheckBoxForegroundColor )
 						.IsChecked( this, &SMenuEntryBlock::IsChecked )
 						.Style( StyleSet, CheckBoxStyle )
@@ -741,7 +739,7 @@ TSharedRef< SWidget> SMenuEntryBlock::BuildSubMenuWidget( const FMenuEntryBuildP
 			UserInterfaceType == EUserInterfaceActionType::RadioButton ||
 			UserInterfaceType == EUserInterfaceActionType::Check ) ?
 				EVisibility::Visible :
-				EVisibility::Hidden;
+				EVisibility::Collapsed;
 
 	TAttribute<FSlateColor> CheckBoxForegroundColor = FSlateColor::UseForeground();
 	FName CheckBoxStyle = ISlateStyle::Join( StyleName, ".CheckBox" );
@@ -801,20 +799,17 @@ TSharedRef< SWidget> SMenuEntryBlock::BuildSubMenuWidget( const FMenuEntryBuildP
 		.Padding(FMargin(2, 0, 8, 0))
 		.AutoWidth()
 		[
-			SNew(SBox)
-			.Visibility(CheckBoxVisibility)
-			.WidthOverride(MenuIconSize)
-			.HeightOverride(MenuIconSize)
-			.HAlign(HAlign_Center)
-			.VAlign(VAlign_Center)
+			SNew(SLinkedBox, OwnerMultiBoxWidget.Pin()->GetLinkedBoxManager())
 			[
 				// For check style menus, use an image instead of a CheckBox because it can't really be checked.
 				UserInterfaceType == EUserInterfaceActionType::Check
 				? StaticCastSharedRef<SWidget>(SNew(SImage)
+					.Visibility(CheckBoxVisibility)
 					.ColorAndOpacity(FSlateColor::UseForeground())
 					.Image(this, &SMenuEntryBlock::GetCheckBoxImageBrushFromStyle, &StyleSet->GetWidgetStyle<FCheckBoxStyle>(CheckBoxStyle)))
 				: StaticCastSharedRef<SWidget>(SNew(SCheckBox)
 					.ForegroundColor(CheckBoxForegroundColor)
+					.Visibility(CheckBoxVisibility)
 					.IsChecked(this, &SMenuEntryBlock::IsChecked)
 					.Style(StyleSet, CheckBoxStyle)
 					.OnCheckStateChanged(this, &SMenuEntryBlock::OnCheckStateChanged))
diff --git a/Engine/Source/Runtime/Slate/Public/Framework/MultiBox/MultiBox.h b/Engine/Source/Runtime/Slate/Public/Framework/MultiBox/MultiBox.h
--- a/Engine/Source/Runtime/Slate/Public/Framework/MultiBox/MultiBox.h
+++ b/Engine/Source/Runtime/Slate/Public/Framework/MultiBox/MultiBox.h
@@ -8,6 +8,7 @@
 #include "Widgets/DeclarativeSyntaxSupport.h"
 #include "Input/Reply.h"
 #include "Widgets/SWidget.h"
+#include "Widgets/Layout/SLinkedBox.h"
 #include "Widgets/SCompoundWidget.h"
 #include "Widgets/Layout/SMenuOwner.h"
 #include "Framework/Commands/UIAction.h"
@@ -27,10 +28,6 @@ class SUniformToolbarPanel;
 
 namespace MultiBoxConstants
 {	
-	// @todo Slate MultiBox: Hard coded dimensions
-	const float MenuIconSize = 16.0f;
-	const float MenuCheckBoxSize = 16.0f;
-
 	/** The time that a mouse should be hovered over a sub-menu before it automatically opens */
 	const float SubMenuOpenTime = 0.0f;
 
@@ -132,6 +129,18 @@ public:
 		return TutorialHighlightName;
 	}
 
+	/** Sets the style name which will be used for this block instead of the owning multibox's style */
+	void SetStyleNameOverride(FName InStyleNameOverride)
+	{
+		StyleNameOverride = InStyleNameOverride;
+	}
+
+	/** Gets the style name which will be used for this block instead of the owning multibox's style */
+	FName GetStyleNameOverride() const
+	{
+		return StyleNameOverride;
+	}
+
 	/**
 	 * Creates a MultiBlock widget for this MultiBlock
 	 *
@@ -219,6 +228,8 @@ private:
 	/** Name to identify a widget for tutorials */
 	FName TutorialHighlightName;
 
+	FName StyleNameOverride;
+
 	/** Whether this block can be searched */
 	bool bSearchable;
 
@@ -585,6 +596,11 @@ public:
 	*/
 	bool GetSearchable() const;
 
+	/**
+	 * @return the the LinkedBoxManager
+	 */
+	TSharedRef<FLinkedBoxManager> GetLinkedBoxManager() { return LinkedBoxManager.ToSharedRef(); }
+
 	/**
 	* Sets optional maximum height of widget
 	*/
@@ -782,4 +798,7 @@ private:
 
 	/** Optional maximum height of widget */
 	TAttribute<float> MaxHeight;
+
+	/** Allows Menu Elementes to size properly */
+	TSharedPtr<FLinkedBoxManager> LinkedBoxManager;
 };
