SLinkedBox - checks for collapsed state when computing childslot desired size

diff --git a/Engine/Source/Runtime/Slate/Public/Widgets/Layout/SLinkedBox.cpp b/Engine/Source/Runtime/Slate/Public/Widgets/Layout/SLinkedBox.cpp
--- a/Engine/Source/Runtime/Slate/Public/Widgets/Layout/SLinkedBox.cpp
+++ b/Engine/Source/Runtime/Slate/Public/Widgets/Layout/SLinkedBox.cpp
@@ -98,7 +98,12 @@ bool SLinkedBox::CustomPrepass(float LayoutScaleMultiplier)
 
 FVector2D SLinkedBox::GetChildrensDesiredSize() const
 {
-	return ChildSlot.GetWidget()->GetDesiredSize() + ChildSlot.SlotPadding.Get().GetDesiredSize();
+	EVisibility ChildVisibility = ChildSlot.GetWidget()->GetVisibility();
+	if ( ChildVisibility != EVisibility::Collapsed )
+	{
+		return ChildSlot.GetWidget()->GetDesiredSize() + ChildSlot.SlotPadding.Get().GetDesiredSize();
+	}
+	return FVector2D::ZeroVector;
 }
 
 FVector2D SLinkedBox::ComputeDesiredSize(float) const
