Reverting CL 13601508 because it showed some UI options that shouldn't be displayed (on the toolbar specifically).

diff --git a/Engine/Source/Runtime/Slate/Private/Framework/MultiBox/MultiBox.cpp b/Engine/Source/Runtime/Slate/Private/Framework/MultiBox/MultiBox.cpp
--- a/Engine/Source/Runtime/Slate/Private/Framework/MultiBox/MultiBox.cpp
+++ b/Engine/Source/Runtime/Slate/Private/Framework/MultiBox/MultiBox.cpp
@@ -1249,8 +1249,6 @@ FReply SMultiBoxWidget::FocusNextWidget(EUINavigation NavigationType)
 
 FReply SMultiBoxWidget::OnFocusReceived( const FGeometry& MyGeometry, const FFocusEvent& InFocusEvent )
 {
-	ResetSearch();
-
 	if (InFocusEvent.GetCause() == EFocusCause::Navigation)
 	{
 		// forward focus to children
