Fix table header rows not clipping text when they get too small

diff --git a/Engine/Source/Runtime/Slate/Private/Widgets/Views/SHeaderRow.cpp b/Engine/Source/Runtime/Slate/Private/Widgets/Views/SHeaderRow.cpp
--- a/Engine/Source/Runtime/Slate/Private/Widgets/Views/SHeaderRow.cpp
+++ b/Engine/Source/Runtime/Slate/Private/Widgets/Views/SHeaderRow.cpp
@@ -200,6 +200,7 @@ public:
 			.HAlign( Column.HeaderHAlignment )
 			.VAlign( Column.HeaderVAlignment )
 			.Padding( Column.HeaderContentPadding.Get( AdjustedDefaultHeaderContentPadding ) )
+			.Clipping(EWidgetClipping::ClipToBounds)
 			[
 				Overlay
 			]
