Added the concept of a tab drawers. Active tabs can be moved to a drawer which closes them until a button on the side of the window is clicked to summon them.  Once in a drawer they will close automatically when they lose focus.

diff --git a/Engine/Source/Editor/EditorStyle/Private/StarshipStyle.cpp b/Engine/Source/Editor/EditorStyle/Private/StarshipStyle.cpp
--- a/Engine/Source/Editor/EditorStyle/Private/StarshipStyle.cpp
+++ b/Engine/Source/Editor/EditorStyle/Private/StarshipStyle.cpp
@@ -7756,7 +7756,7 @@ void FStarshipEditorStyle::FStyle::SetupStatusBarStyle()
 
 	Set("StatusBar.ContentBrowserUp", new CORE_IMAGE_BRUSH_SVG("Starship/Common/chevron-up", Icon16x16));
 	Set("StatusBar.ContentBrowserDown", new CORE_IMAGE_BRUSH_SVG("Starship/Common/chevron-down", Icon16x16));
-	Set("StatusBar.ContentBrowserShadow", new BOX_BRUSH("Starship/StatusBar/drawer-shadow", FMargin(10/64.0f, 20/64.f, 10/64.f, 0), FLinearColor(0,0,0,1)));
+	Set("StatusBar.ContentBrowserShadow", new BOX_BRUSH("Starship/StatusBar/drawer-shadow-bottom", FMargin(10/64.0f, 20/64.f, 10/64.f, 0), FLinearColor(0,0,0,1)));
 
 	Set("StatusBar.Height", 32.0f);
 	Set("StatusBar.Background", new FSlateColorBrush(FStyleColors::Background));
diff --git a/Engine/Source/Editor/LevelEditor/Private/LevelViewportActions.cpp b/Engine/Source/Editor/LevelEditor/Private/LevelViewportActions.cpp
--- a/Engine/Source/Editor/LevelEditor/Private/LevelViewportActions.cpp
+++ b/Engine/Source/Editor/LevelEditor/Private/LevelViewportActions.cpp
@@ -32,6 +32,7 @@ void FLevelViewportCommands::RegisterCommands()
 	UI_COMMAND( ToggleMaximize, "Maximize Viewport", "Toggles the Maximize state of the current viewport", EUserInterfaceActionType::ToggleButton, FInputChord() );
 	UI_COMMAND( ToggleGameView, "Game View", "Toggles game view.  Game view shows the scene as it appears in game", EUserInterfaceActionType::ToggleButton, FInputChord( EKeys::G ) );
 	UI_COMMAND( ToggleImmersive, "Immersive Mode", "Switches this viewport between immersive mode and regular mode", EUserInterfaceActionType::ToggleButton, PLATFORM_MAC ? FInputChord( EModifierKey::Control, EKeys::F11 ) : FInputChord( EKeys::F11 ) );
+	UI_COMMAND(ToggleSidebarAllTabs, "Sidebar All Tabs", "Moves all tabs except the level editor to a sidebar or restores any previous state before if all tabs are already sidebared", EUserInterfaceActionType::ToggleButton, FInputChord(EKeys::F10));
 
 	UI_COMMAND( HighResScreenshot, "High Resolution Screenshot...", "Opens the control panel for high resolution screenshots", EUserInterfaceActionType::Button, FInputChord() );
 	
diff --git a/Engine/Source/Editor/LevelEditor/Private/SLevelEditor.cpp b/Engine/Source/Editor/LevelEditor/Private/SLevelEditor.cpp
--- a/Engine/Source/Editor/LevelEditor/Private/SLevelEditor.cpp
+++ b/Engine/Source/Editor/LevelEditor/Private/SLevelEditor.cpp
@@ -1061,11 +1061,12 @@ TSharedRef<SWidget> SLevelEditor::RestoreContentArea( const TSharedRef<SDockTab>
 			const FText ViewportTooltip = NSLOCTEXT("LevelEditorTabs", "LevelEditorViewportTooltip", "Open a Viewport tab. Use this to view and edit the current level.");
 			const FSlateIcon ViewportIcon(FEditorStyle::GetStyleSetName(), "LevelEditor.Tabs.Viewports");
 
-			LevelEditorTabManager->RegisterTabSpawner(LevelEditorTabIds::LevelEditorViewport, FOnSpawnTab::CreateSP(this, &SLevelEditor::SpawnLevelEditorTab, LevelEditorTabIds::LevelEditorViewport, FString()) )
+			LevelEditorTabManager->RegisterTabSpawner(LevelEditorTabIds::LevelEditorViewport, FOnSpawnTab::CreateSP(this, &SLevelEditor::SpawnLevelEditorTab, LevelEditorTabIds::LevelEditorViewport, FString()))
 				.SetDisplayName(NSLOCTEXT("LevelEditorTabs", "LevelEditorViewport", "Viewport 1"))
 				.SetTooltipText(ViewportTooltip)
-				.SetGroup( MenuStructure.GetLevelEditorViewportsCategory() )
-				.SetIcon(ViewportIcon);
+				.SetGroup(MenuStructure.GetLevelEditorViewportsCategory())
+				.SetIcon(ViewportIcon)
+				.SetCanSidebarTab(false);
 
 			LevelEditorTabManager->RegisterTabSpawner(LevelEditorTabIds::LevelEditorViewport_Clone1, FOnSpawnTab::CreateSP(this, &SLevelEditor::SpawnLevelEditorTab, LevelEditorTabIds::LevelEditorViewport_Clone1, FString()) )
 				.SetDisplayName(NSLOCTEXT("LevelEditorTabs", "LevelEditorViewport_Clone1", "Viewport 2"))
diff --git a/Engine/Source/Editor/LevelEditor/Private/SLevelViewport.cpp b/Engine/Source/Editor/LevelEditor/Private/SLevelViewport.cpp
--- a/Engine/Source/Editor/LevelEditor/Private/SLevelViewport.cpp
+++ b/Engine/Source/Editor/LevelEditor/Private/SLevelViewport.cpp
@@ -1312,6 +1312,10 @@ void SLevelViewport::BindOptionCommands( FUICommandList& OutCommandList )
 		FCanExecuteAction(),
 		FIsActionChecked::CreateSP( this, &SLevelViewport::IsImmersive ) );
 
+	OutCommandList.MapAction(
+		ViewportActions.ToggleSidebarAllTabs,
+		FExecuteAction::CreateSP(this, &SLevelViewport::OnToggleSidebarTabs)
+	);
 
 	OutCommandList.MapAction(
 		ViewportActions.ToggleCinematicPreview,
@@ -1863,6 +1867,11 @@ void SLevelViewport::OnToggleImmersive()
 	}
 }
 
+void SLevelViewport::OnToggleSidebarTabs()
+{
+	ParentLevelEditor.Pin()->GetTabManager()->ToggleSidebarOpenTabs();
+}
+
 bool SLevelViewport::IsImmersive() const
 {
 	if( ParentLayout.IsValid() && !ConfigKey.IsNone())
diff --git a/Engine/Source/Editor/LevelEditor/Public/LevelViewportActions.h b/Engine/Source/Editor/LevelEditor/Public/LevelViewportActions.h
--- a/Engine/Source/Editor/LevelEditor/Public/LevelViewportActions.h
+++ b/Engine/Source/Editor/LevelEditor/Public/LevelViewportActions.h
@@ -89,6 +89,9 @@ PRAGMA_ENABLE_DEPRECATION_WARNINGS
 
 	/** Toggles immersive mode in the viewport */
 	TSharedPtr< FUICommandInfo > ToggleImmersive;
+	 
+	/** Toggles moving all tabs except the viewport to a sidebar or removing them from a sidebar */
+	TSharedPtr< FUICommandInfo > ToggleSidebarAllTabs;
 
 	/** Toggles maximize mode in the viewport */
 	TSharedPtr< FUICommandInfo > ToggleMaximize;
diff --git a/Engine/Source/Editor/LevelEditor/Public/SLevelViewport.h b/Engine/Source/Editor/LevelEditor/Public/SLevelViewport.h
--- a/Engine/Source/Editor/LevelEditor/Public/SLevelViewport.h
+++ b/Engine/Source/Editor/LevelEditor/Public/SLevelViewport.h
@@ -444,6 +444,9 @@ private:
 	 */
 	void OnToggleImmersive();
 
+	/** Called when moving tabs in and out of a sidebar is activated by the user */
+	void OnToggleSidebarTabs();
+
 	/**
 	* Called to determine whether the maximize mode of current viewport can be toggled
 	*/
diff --git a/Engine/Source/Runtime/Slate/Private/Framework/Docking/SDockingArea.cpp b/Engine/Source/Runtime/Slate/Private/Framework/Docking/SDockingArea.cpp
--- a/Engine/Source/Runtime/Slate/Private/Framework/Docking/SDockingArea.cpp
+++ b/Engine/Source/Runtime/Slate/Private/Framework/Docking/SDockingArea.cpp
@@ -7,6 +7,7 @@
 #include "Framework/Docking/SDockingTarget.h"
 #include "Framework/Docking/FDockingDragOperation.h"
 #include "HAL/PlatformApplicationMisc.h"
+#include "STabSidebar.h"
 
 
 void SDockingArea::Construct( const FArguments& InArgs, const TSharedRef<FTabManager>& InTabManager, const TSharedRef<FTabManager::FArea>& PersistentNode )
@@ -29,8 +30,24 @@ void SDockingArea::Construct( const FArguments& InArgs, const TSharedRef<FTabMan
 		.Visibility( EVisibility::SelfHitTestInvisible )
 		+SOverlay::Slot()
 		[
-			SAssignNew(Splitter, SSplitter)
-			. Orientation( PersistentNode->GetOrientation() )
+			SNew(SHorizontalBox)
+			+ SHorizontalBox::Slot()
+			.AutoWidth()
+			[
+				SAssignNew(LeftSidebar, STabSidebar)
+				.Location(ESidebarLocation::Left)
+			]
+			+ SHorizontalBox::Slot()
+			[
+				SAssignNew(Splitter, SSplitter)
+				.Orientation(PersistentNode->GetOrientation())
+			]
+			+SHorizontalBox::Slot()
+			.AutoWidth()
+			[
+				SAssignNew(RightSidebar, STabSidebar)
+				.Location(ESidebarLocation::Right)
+			]
 		]
 		
 		+ SOverlay::Slot()
@@ -222,10 +239,6 @@ void SDockingArea::CleanUp( ELayoutModification RemovalMethod )
 				ParentWindow->HideWindow();
 				MyTabManager.Pin()->GetPrivateApi().OnDockAreaClosing( SharedThis(this) );
 			}
-			else
-			{
-				ensure( RemovalMethod == TabRemoval_None );
-			}
 		}
 	}
 	else
@@ -236,7 +249,7 @@ void SDockingArea::CleanUp( ELayoutModification RemovalMethod )
 	// In some cases a dock area will control the window,
 	// and we need to move some of the tabs out of the way
 	// to make room for window chrome.
-	MakeRoomForWindowChrome();
+	UpdateWindowChromeAndSidebar();
 }
 
 void SDockingArea::SetParentWindow( const TSharedRef<SWindow>& NewParentWindow )
@@ -246,7 +259,7 @@ void SDockingArea::SetParentWindow( const TSharedRef<SWindow>& NewParentWindow )
 		NewParentWindow->SetRequestDestroyWindowOverride( FRequestDestroyWindowOverride::CreateSP( this, &SDockingArea::OnOwningWindowBeingDestroyed ) );
 	}
 
-	// Even thought we don't manage the parent window's lifetime, we are still responsible for making its window chrome.
+	// Even though we don't manage the parent window's lifetime, we are still responsible for making its window chrome.
 	{
 		TSharedPtr<IWindowTitleBar> TitleBar;
 
@@ -255,8 +268,11 @@ void SDockingArea::SetParentWindow( const TSharedRef<SWindow>& NewParentWindow )
 		Args.CenterContentAlignment = HAlign_Fill;
 
 		TSharedRef<SWidget> TitleBarWidget = FSlateApplication::Get().MakeWindowTitleBar(Args, TitleBar);
+		(*WindowControlsArea)
+		[
+			TitleBarWidget
+		];
 
-		(*WindowControlsArea)[TitleBarWidget];
 		NewParentWindow->SetTitleBar(TitleBar);
 	}
 
@@ -324,6 +340,93 @@ TSharedRef<FTabManager> SDockingArea::GetTabManager() const
 	return MyTabManager.Pin().ToSharedRef();
 }
 
+ESidebarLocation SDockingArea::AddTabToSidebar(TSharedRef<SDockTab> TabToAdd)
+{
+	ESidebarLocation Location = ESidebarLocation::None;
+
+	if(ensure(bCanHaveSidebar))
+	{
+		// Determine which sidebar to add the tab to. Testing is done in absolute desktop space because we need the tab and the area to be in the same space
+		FSlateRect AreaRect = GetPaintSpaceGeometry().GetLayoutBoundingRect();
+		FSlateRect LeftRect = FSlateRect(AreaRect.Left, AreaRect.Top, AreaRect.Right / 2, AreaRect.Bottom);
+		FSlateRect RightRect = FSlateRect(AreaRect.Right / 2, AreaRect.Top, AreaRect.Right, AreaRect.Bottom);
+
+		FSlateRect TabRect = TabToAdd->GetPaintSpaceGeometry().GetLayoutBoundingRect();
+
+		bool bOverlapsLeft = false;
+		FSlateRect LeftOverlap = LeftRect.IntersectionWith(TabRect, bOverlapsLeft);
+
+		bool bOverlapsRight = false;
+		FSlateRect RightOverlap = RightRect.IntersectionWith(TabRect, bOverlapsRight);
+
+		if (bOverlapsLeft && bOverlapsRight)
+		{
+			if (LeftOverlap.GetArea() > RightOverlap.GetArea())
+			{
+				// Left side
+				Location = ESidebarLocation::Left;
+			}
+			else
+			{
+				// Right side
+				Location = ESidebarLocation::Right;
+			}
+		}
+		else if (bOverlapsLeft)
+		{
+			// left side
+			Location = ESidebarLocation::Left;
+		}
+		else
+		{
+			ensure(bOverlapsRight);
+			// right side
+			Location = ESidebarLocation::Right;
+		}
+
+		if (Location == ESidebarLocation::Left)
+		{
+			LeftSidebar->AddTab(TabToAdd);
+		}
+		else
+		{
+			RightSidebar->AddTab(TabToAdd);
+		}
+	}
+	
+
+	return Location;
+
+}
+
+bool SDockingArea::RestoreTabFromSidebar(TSharedRef<SDockTab> TabToRemove)
+{
+	return LeftSidebar->RestoreTab(TabToRemove) || RightSidebar->RestoreTab(TabToRemove);
+}
+
+bool SDockingArea::IsTabInSidebar(TSharedRef<SDockTab> Tab) const
+{
+	return LeftSidebar->ContainsTab(Tab) || RightSidebar->ContainsTab(Tab);
+}
+
+bool SDockingArea::TryOpenSidebarDrawer(TSharedRef<SDockTab> TabToOpen) const
+{
+	return LeftSidebar->TryOpenSidebarDrawer(TabToOpen) || RightSidebar->TryOpenSidebarDrawer(TabToOpen);
+}
+
+void SDockingArea::AddSidebarTabsFromRestoredLayout(const FSidebarTabLists& SidebarTabs)
+{
+	for (const TSharedRef<SDockTab>& Tab : SidebarTabs.LeftSidebarTabs)
+	{
+		LeftSidebar->AddTab(Tab);
+	}
+
+	for (const TSharedRef<SDockTab>& Tab : SidebarTabs.RightSidebarTabs)
+	{
+		RightSidebar->AddTab(Tab);
+	}
+}
+
 SDockingNode::ECleanupRetVal SDockingArea::CleanUpNodes()
 {
 	SDockingNode::ECleanupRetVal ReturnValue = SDockingSplitter::CleanUpNodes();
@@ -459,7 +562,7 @@ void SDockingArea::OnLiveTabAdded()
 	CleanUp(SDockingNode::TabRemoval_None);
 }
 
-void SDockingArea::MakeRoomForWindowChrome()
+void SDockingArea::UpdateWindowChromeAndSidebar()
 {
 	TArray< TSharedRef<SDockingNode> > AllNodes = this->GetChildNodesRecursively();
 	if (AllNodes.Num() > 0)
@@ -474,15 +577,67 @@ void SDockingArea::MakeRoomForWindowChrome()
 			}
 		}
 
-		if (this->ParentWindowPtr.IsValid())
+		bCanHaveSidebar = true;
+
+		if (TSharedPtr<SWindow> ParentWindow = GetParentWindow())
 		{
+			bool bAccountForMenuBarPadding = false;
+			bool bContainsMajorTabs = false;
+			bool bContainsNomadTabs = false;
+			if (MyTabManager.Pin()->AllowsWindowMenuBar())
+			{
+				TArray<TSharedRef<SDockTab>> AllTabs = GetAllChildTabs();
+				for (auto& Tab : AllTabs)
+				{
+					if (Tab->GetParentWindow() == ParentWindow && Tab->GetTabRole() == ETabRole::MajorTab)
+					{
+						bAccountForMenuBarPadding = true;
+						bContainsMajorTabs = true;
+						break;
+					}
+					else if (Tab->GetTabRole() == ETabRole::NomadTab)
+					{
+						bContainsNomadTabs = true;
+					}
+				}
+			}
+
+			bCanHaveSidebar = !bContainsMajorTabs;
+
+			if (!bCanHaveSidebar)
+			{
+				LeftSidebar->SetVisibility(EVisibility::Collapsed);
+				RightSidebar->SetVisibility(EVisibility::Collapsed);
+			}
+			else
+			{
+				if (bAccountForMenuBarPadding)
+				{
+					// Menu bar padding has already applied an offset so no need to do it again
+					LeftSidebar->SetOffset(0);
+					RightSidebar->SetOffset(0);
+				}
+				else
+				{
+					LeftSidebar->SetOffset(35.f);
+					RightSidebar->SetOffset(35.f);
+				}
+			}
+
+			const bool bNoMajorOrNomadTabs = !(bContainsMajorTabs || bContainsNomadTabs);
+
 			// Reserve some space for the minimize, restore, and close controls
 			TSharedRef<SDockingTabStack> WindowControlHousing = this->FindTabStackToHouseWindowControls();
-			WindowControlHousing->ReserveSpaceForWindowChrome(SDockingTabStack::EChromeElement::Controls);
+			WindowControlHousing->ReserveSpaceForWindowChrome(SDockingTabStack::EChromeElement::Controls, bAccountForMenuBarPadding, bNoMajorOrNomadTabs);
 
 			// Reserve some space for the app icons
 			TSharedRef<SDockingTabStack> IconHousing = this->FindTabStackToHouseWindowIcon();
-			IconHousing->ReserveSpaceForWindowChrome(SDockingTabStack::EChromeElement::Icon);
+			IconHousing->ReserveSpaceForWindowChrome(SDockingTabStack::EChromeElement::Icon, bAccountForMenuBarPadding, bNoMajorOrNomadTabs);
+
+			if (ParentWindow->GetTitleBar())
+			{
+				ParentWindow->GetTitleBar()->SetUseLargeIcon(bAccountForMenuBarPadding);
+			}
 		}
 	}
 
diff --git a/Engine/Source/Runtime/Slate/Private/Framework/Docking/SDockingArea.h b/Engine/Source/Runtime/Slate/Private/Framework/Docking/SDockingArea.h
--- a/Engine/Source/Runtime/Slate/Private/Framework/Docking/SDockingArea.h
+++ b/Engine/Source/Runtime/Slate/Private/Framework/Docking/SDockingArea.h
@@ -12,6 +12,15 @@
 #include "SDockingNode.h"
 #include "SDockingSplitter.h"
 
+class STabSidebar;
+
+/** List of tabs that should be in each sidebar */
+struct FSidebarTabLists
+{
+	TArray<TSharedRef<SDockTab>> LeftSidebarTabs;
+	TArray<TSharedRef<SDockTab>> RightSidebarTabs;
+};
+
 /**
  * Represents the root node in a hierarchy of DockNodes.
  */
@@ -83,6 +92,39 @@ public:
 
 	TSharedRef<FTabManager> GetTabManager() const;
 
+	/**
+	 * Adds a tab to a drawer in the sidebar. 
+	 *
+	 * @param TabToAdd	The tab to add to the sidebar
+	 * @return The location of the sidebar that the tab was added to
+	 */
+	ESidebarLocation AddTabToSidebar(TSharedRef<SDockTab> TabToAdd);
+
+	/**
+	 * Restores a tab from the sidebar to its parent tab stack and removes the tab from the sidebar.  
+	 *
+	 * @return true if the tab was found in this area and restore
+	 */
+	bool RestoreTabFromSidebar(TSharedRef<SDockTab> TabToRemove);
+
+	/**
+	 * @return true if the specified tab is in the sidebar
+	 */
+	bool IsTabInSidebar(TSharedRef<SDockTab> Tab) const;
+
+	/**
+	 * Attempts to open a sidebar drawer that may the tab to open
+	 *
+	 * @true if the drawer was opened, false if the tab is not in a drawer
+	 */
+	bool TryOpenSidebarDrawer(TSharedRef<SDockTab> TabToOpen) const;
+
+	/**
+	 * Adds all tabs back to a sidebar that were saved in a sidebar from a previous session
+	 */
+	void AddSidebarTabsFromRestoredLayout(const FSidebarTabLists& SidebarTabs);
+
+	bool CanHaveSidebar() const { return bCanHaveSidebar; }
 protected:
 	virtual SDockingNode::ECleanupRetVal CleanUpNodes() override;
 
@@ -104,11 +146,15 @@ private:
 	 * If this dock area controls a window, then we need
 	 * to reserve some room in the upper left and upper right tab wells
 	 * so that there is no overlap with the window chrome.
+	 *
+	 * We also update the sidebar to account for major tabs.  Docking areas for major tabs do not have a sidebar.
 	 */
-	void MakeRoomForWindowChrome();
-
+	void UpdateWindowChromeAndSidebar();
 private:	
-	
+	/** Left and right sidebar widgets */
+	TSharedPtr<STabSidebar> LeftSidebar;
+	TSharedPtr<STabSidebar> RightSidebar;
+
 	/** The window this dock area is embedded within.  If bIsManagingParentWindow is true, the dock area will also
 	    destroy the window when the last tab goes away. */
 	TWeakPtr<SWindow> ParentWindowPtr;
@@ -133,4 +179,7 @@ private:
 
 	/** True when the last tab has been pulled from this area, meaning that this DockArea will not be necessary once that tab finds a new home. */
 	bool bCleanUpUponTabRelocation;
+
+	/** True if this area can ever show sidebars (minor tab areas only) */
+	bool bCanHaveSidebar;
 };
diff --git a/Engine/Source/Runtime/Slate/Private/Framework/Docking/SDockingNode.h b/Engine/Source/Runtime/Slate/Private/Framework/Docking/SDockingNode.h
--- a/Engine/Source/Runtime/Slate/Private/Framework/Docking/SDockingNode.h
+++ b/Engine/Source/Runtime/Slate/Private/Framework/Docking/SDockingNode.h
@@ -75,8 +75,12 @@ public:
 	/** A tab can be removed from a stack when a uset drags it away or when the user closes it */
 	enum ELayoutModification
 	{
+		/** Tab was dragged out of the node */
 		TabRemoval_DraggedOut,
+		/** Tab is being closed */
 		TabRemoval_Closed,
+		/** Tab is being put in a sidebar */
+		TabRemoval_Sidebar,
 		TabRemoval_None
 	};
 	
@@ -88,6 +92,9 @@ public:
 	/** Recursively searches through all children looking for child tabs */
 	virtual TArray< TSharedRef<SDockTab> > GetAllChildTabs() const {return TArray< TSharedRef<SDockTab> >();}
 
+	/** Gets the number of tabs in all children recursively */
+	virtual int32 GetNumTabs() const { return 0; }
+
 	/**
 	 * Attempt to dock the TabWidget from the DragDropEvent next to this node.
 	 * 
diff --git a/Engine/Source/Runtime/Slate/Private/Framework/Docking/SDockingSplitter.cpp b/Engine/Source/Runtime/Slate/Private/Framework/Docking/SDockingSplitter.cpp
--- a/Engine/Source/Runtime/Slate/Private/Framework/Docking/SDockingSplitter.cpp
+++ b/Engine/Source/Runtime/Slate/Private/Framework/Docking/SDockingSplitter.cpp
@@ -272,9 +272,9 @@ TArray< TSharedRef<SDockingNode> > SDockingSplitter::GetChildNodesRecursively()
 	return ChildNodes;
 }
 
-TArray< TSharedRef<SDockTab> > SDockingSplitter::GetAllChildTabs() const
+TArray<TSharedRef<SDockTab>> SDockingSplitter::GetAllChildTabs() const
 {
-	TArray< TSharedRef<SDockTab> > ChildTabs;
+	TArray<TSharedRef<SDockTab>> ChildTabs;
 	for (int32 i = 0; i < Children.Num(); ++i)
 	{
 		const TSharedRef<SDockingNode>& Child = Children[i];
@@ -284,6 +284,18 @@ TArray< TSharedRef<SDockTab> > SDockingSplitter::GetAllChildTabs() const
 }
 
 
+int32 SDockingSplitter::GetNumTabs() const
+{
+	int32 NumTabs = 0;
+	for (int32 i = 0; i < Children.Num(); ++i)
+	{
+		const TSharedRef<SDockingNode>& Child = Children[i];
+		NumTabs += Child->GetNumTabs();
+	}
+
+	return NumTabs;
+}
+
 EOrientation SDockingSplitter::GetOrientation() const
 {
 	return Splitter->GetOrientation();
diff --git a/Engine/Source/Runtime/Slate/Private/Framework/Docking/SDockingSplitter.h b/Engine/Source/Runtime/Slate/Private/Framework/Docking/SDockingSplitter.h
--- a/Engine/Source/Runtime/Slate/Private/Framework/Docking/SDockingSplitter.h
+++ b/Engine/Source/Runtime/Slate/Private/Framework/Docking/SDockingSplitter.h
@@ -76,6 +76,9 @@ public:
 	/** Recursively searches through all children looking for child tabs */
 	virtual TArray< TSharedRef<SDockTab> > GetAllChildTabs() const override;
 
+	/** Gets the number of tabs in all children recursively */
+	virtual int32 GetNumTabs() const override;
+
 	/** Should this node auto-size or be a percentage of its parent size */
 	virtual SSplitter::ESizeRule GetSizeRule() const override;
 
diff --git a/Engine/Source/Runtime/Slate/Private/Framework/Docking/SDockingTabStack.cpp b/Engine/Source/Runtime/Slate/Private/Framework/Docking/SDockingTabStack.cpp
--- a/Engine/Source/Runtime/Slate/Private/Framework/Docking/SDockingTabStack.cpp
+++ b/Engine/Source/Runtime/Slate/Private/Framework/Docking/SDockingTabStack.cpp
@@ -61,8 +61,6 @@ void SDockingTabStack::Construct( const FArguments& InArgs, const TSharedRef<FTa
 	InlineContentAreaRight = nullptr;
 	TitleBarSlot = nullptr;
 
-	this->TabStackGeometry = FGeometry();
-
 	// Animation that toggles the tabs
 	{
 		ShowHideTabWell = FCurveSequence(0,0.15);
@@ -232,7 +230,6 @@ void SDockingTabStack::Construct( const FArguments& InArgs, const TSharedRef<FTa
 	}
 }
 
-
 void SDockingTabStack::OnLastTabRemoved()
 {
 	if (!bIsDocumentArea)
@@ -247,7 +244,7 @@ void SDockingTabStack::OnLastTabRemoved()
 	}
 }
 
-void SDockingTabStack::OnTabClosed( const TSharedRef<SDockTab>& ClosedTab )
+void SDockingTabStack::OnTabClosed(const TSharedRef<SDockTab>& ClosedTab, SDockingNode::ELayoutModification RemovalMethod)
 {
 	const FTabId& TabIdBeingClosed = ClosedTab->GetLayoutIdentifier();
 	
@@ -255,24 +252,24 @@ void SDockingTabStack::OnTabClosed( const TSharedRef<SDockTab>& ClosedTab )
 	const bool bIsTabPersistable = TabIdBeingClosed.IsTabPersistable();
 	if (bIsTabPersistable)
 	{
-		ClosePersistentTab(TabIdBeingClosed);
+		// Sidebar tabs should still exist in the stacks layout so we can restore it
+		if (RemovalMethod != SDockingNode::ELayoutModification::TabRemoval_Sidebar)
+		{
+			ClosePersistentTab(TabIdBeingClosed);
+		}
 	}
 	else
 	{
-		RemovePersistentTab( TabIdBeingClosed );
+		RemovePersistentTab(TabIdBeingClosed);
 	}	
 }
 
+
 void SDockingTabStack::OnTabRemoved( const FTabId& TabId )
 {
 	RemovePersistentTab( TabId );
 }
 
-void SDockingTabStack::Tick( const FGeometry& AllottedGeometry, const double InCurrentTime, const float InDeltaTime )
-{
-	TabStackGeometry = AllottedGeometry;
-}
-
 void SDockingTabStack::OpenTab( const TSharedRef<SDockTab>& InTab, int32 InsertLocationAmongActiveTabs )
 {
 	const int32 TabIndex = OpenPersistentTab( InTab->GetLayoutIdentifier(), InsertLocationAmongActiveTabs );
@@ -302,6 +299,31 @@ void SDockingTabStack::AddTabWidget( const TSharedRef<SDockTab>& InTab, int32 At
 
 }
 
+void SDockingTabStack::AddSidebarTab(const TSharedRef<SDockTab>& InTab)
+{
+	InTab->SetParent(TabWell);
+}
+
+float SDockingTabStack::GetTabSidebarSizeCoefficient(const TSharedRef<SDockTab>& InTab)
+{
+	FTabManager::FTab* Tab = Tabs.FindByPredicate(FTabMatcher(InTab->GetLayoutIdentifier()));
+	if (Tab)
+	{
+		return Tab->SidebarSizeCoefficient;
+	}
+
+	return 0;
+}
+
+void SDockingTabStack::SetTabSidebarSizeCoefficient(const TSharedRef<SDockTab>& InTab, float InSizeCoefficient)
+{
+	FTabManager::FTab* Tab = Tabs.FindByPredicate(FTabMatcher(InTab->GetLayoutIdentifier()));
+	if (Tab)
+	{
+		Tab->SidebarSizeCoefficient = InSizeCoefficient;
+	}
+}
+
 const TSlotlessChildren<SDockTab>& SDockingTabStack::GetTabs() const
 {
 	return TabWell->GetTabs();
@@ -319,7 +341,7 @@ bool SDockingTabStack::HasTab(const struct FTabMatcher& TabMatcher) const
 
 FGeometry SDockingTabStack::GetTabStackGeometry() const
 {
-	return TabStackGeometry;
+	return GetTickSpaceGeometry();
 }
 
 void SDockingTabStack::RemoveClosedTabsWithName( FName InName )
@@ -576,14 +598,28 @@ TSharedRef<SWidget> SDockingTabStack::MakeContextMenu()
 		MenuBuilder.BeginSection("DockingTabStackOptions", LOCTEXT("TabOptionsHeading", "Options") );
 		{
 			MenuBuilder.AddMenuEntry(
-				LOCTEXT("HideTabWell", "Hide Tab"),
-				LOCTEXT("HideTabWellTooltip", "Hide the tabs to save room."),
+				LOCTEXT("CollapseTabWell", "Collapse Tab Well"),
+				LOCTEXT("CollapseTabWellTooltip", "Collapses the tabs headers to save room."),
 				FSlateIcon(),
 				FUIAction(
 					FExecuteAction::CreateSP( this, &SDockingTabStack::ToggleTabWellVisibility ),
 					FCanExecuteAction::CreateSP( this, &SDockingTabStack::CanHideTabWell )
 				) 
 			);
+
+			if(IsTabAllowedInSidebar(TabWell->GetForegroundTab()))
+			{
+				MenuBuilder.AddMenuEntry(
+					LOCTEXT("MoveToSidebar", "Move To Sidebar"),
+					LOCTEXT("MoveToSidebarTooltip", "Moves this tab to a sidebar drawer on the side of the window closest to the tab.\nThe tab can be opened from the drawer and will automatically close again when clicking off it."),
+					FSlateIcon(),
+					FUIAction(
+						FExecuteAction::CreateSP(this, &SDockingTabStack::MoveForegroundTabToSidebar),
+						FCanExecuteAction::CreateSP(this, &SDockingTabStack::CanMoveForegroundTabToSidebar)
+					)
+				);
+			}
+
 		}
 		MenuBuilder.EndSection();
 
@@ -753,7 +789,7 @@ TSharedPtr<FTabManager::FLayoutNode> SDockingTabStack::GatherPersistentLayout()
 			const bool bIsTabPersistable = Tabs[TabIndex].TabId.IsTabPersistable();
 			if ( bIsTabPersistable )
 			{
-				PersistentStack->AddTab(Tabs[TabIndex].TabId, Tabs[TabIndex].TabState);
+				PersistentStack->AddTab(Tabs[TabIndex]);
 			}			
 		}
 		return PersistentStack;
@@ -887,6 +923,31 @@ void SDockingTabStack::ToggleTabWellVisibility()
 	ShowHideTabWell.Reverse();
 }
 
+void SDockingTabStack::MoveForegroundTabToSidebar()
+{
+	if (TSharedPtr<SDockTab> ForegroundTabPtr = TabWell->GetForegroundTab())
+	{
+		MoveTabToSidebar(ForegroundTabPtr.ToSharedRef());
+	}
+}
+
+void SDockingTabStack::MoveTabToSidebar(TSharedRef<SDockTab> Tab)
+{
+	const int32 TabIndex = Tabs.IndexOfByPredicate(FTabMatcher(Tab->GetLayoutIdentifier(), ETabState::OpenedTab));
+	if(TabIndex != INDEX_NONE)
+	{
+		ESidebarLocation SidebarLoc = GetDockArea()->AddTabToSidebar(Tab);
+
+		if (SidebarLoc != ESidebarLocation::None)
+		{
+			Tabs[TabIndex].TabState = ETabState::SidebarTab;
+			Tabs[TabIndex].SidebarLocation = SidebarLoc;
+
+			TabWell->RemoveAndDestroyTab(Tab, ELayoutModification::TabRemoval_Sidebar);
+		}
+	}
+}
+
 FReply SDockingTabStack::UnhideTabWell()
 {
 	SetTabWellHidden(false);
@@ -948,6 +1009,29 @@ bool SDockingTabStack::CanCloseAllButForegroundTab() const
 	return false;
 }
 
+bool SDockingTabStack::CanMoveForegroundTabToSidebar() const
+{
+	if(TSharedPtr<SDockTab> ForegroundTabPtr = TabWell->GetForegroundTab())
+	{
+		return CanMoveTabToSideBar(ForegroundTabPtr.ToSharedRef());
+	}
+
+	return false;
+}
+
+bool SDockingTabStack::CanMoveTabToSideBar(TSharedRef<SDockTab> Tab) const
+{
+	const FTabId TabIdBeingClosed = Tab->GetLayoutIdentifier();
+	// Only persistable non-major tabs can be put into a sidebar. There must also be more than one tab or else adding to a sidebar doesnt make si
+	return TabIdBeingClosed.IsTabPersistable() && Tab->GetVisualTabRole() != ETabRole::MajorTab && GetDockArea()->GetNumTabs() > 1;
+}
+
+bool SDockingTabStack::IsTabAllowedInSidebar(TSharedPtr<SDockTab> Tab) const
+{
+	// Major tabs are not allowed to be sidebared
+	return Tab.IsValid() && Tab->GetTabRole() != ETabRole::MajorTab && Tab->GetTabManager()->IsTabAllowedInSidebar(Tab->GetLayoutIdentifier());
+}
+
 SSplitter::ESizeRule SDockingTabStack::GetSizeRule() const
 {
 	int32 NumTabs = this->GetNumTabs();
@@ -1009,7 +1093,7 @@ const FSlateBrush* SDockingTabStack::GetTabStackBorderImage() const
 
 int32 SDockingTabStack::OpenPersistentTab( const FTabId& TabId, int32 OpenLocationAmongActiveTabs )
 {
-	const int32 ExistingClosedTabIndex = Tabs.IndexOfByPredicate(FTabMatcher(TabId, ETabState::ClosedTab));
+	const int32 ExistingClosedTabIndex = Tabs.IndexOfByPredicate(FTabMatcher(TabId, static_cast<ETabState::Type>(ETabState::ClosedTab|ETabState::SidebarTab)));
 
 	if (OpenLocationAmongActiveTabs == INDEX_NONE)
 	{						
@@ -1082,7 +1166,7 @@ int32 SDockingTabStack::OpenPersistentTab( const FTabId& TabId, int32 OpenLocati
 
 int32 SDockingTabStack::ClosePersistentTab( const FTabId& TabId )
 {
-	const int32 TabIndex = Tabs.IndexOfByPredicate(FTabMatcher(TabId, ETabState::OpenedTab));
+	const int32 TabIndex = Tabs.IndexOfByPredicate(FTabMatcher(TabId, static_cast<ETabState::Type>(ETabState::OpenedTab|ETabState::SidebarTab)));
 	if (TabIndex != INDEX_NONE)
 	{
 		Tabs[TabIndex].TabState = ETabState::ClosedTab;
diff --git a/Engine/Source/Runtime/Slate/Private/Framework/Docking/SDockingTabStack.h b/Engine/Source/Runtime/Slate/Private/Framework/Docking/SDockingTabStack.h
--- a/Engine/Source/Runtime/Slate/Private/Framework/Docking/SDockingTabStack.h
+++ b/Engine/Source/Runtime/Slate/Private/Framework/Docking/SDockingTabStack.h
@@ -11,12 +11,8 @@
 #include "Input/Reply.h"
 #include "Widgets/SWidget.h"
 #include "Animation/CurveSequence.h"
-#include "Widgets/SBoxPanel.h"
-#include "Widgets/SOverlay.h"
 #include "Framework/Docking/TabManager.h"
-#include "Widgets/Layout/SSplitter.h"
 #include "SDockingNode.h"
-#include "Widgets/Docking/SDockTab.h"
 
 #define DEBUG_TAB_MANAGEMENT 0
 
@@ -24,9 +20,29 @@ class FUICommandList;
 class FWeakWidgetPath;
 class FWidgetPath;
 class SDockingTabWell;
+class SBorder;
 
 template<typename ChildType> class TSlotlessChildren;
 
+struct FDockingStackOptionalContent
+{
+	FDockingStackOptionalContent()
+		: ContentLeft(SNullWidget::NullWidget)
+		, ContentRight(SNullWidget::NullWidget)
+		, TitleBarContentRight(SNullWidget::NullWidget)
+	{}
+
+	// Content that appears on the left side of tabs in this docking stack
+	TSharedRef<SWidget> ContentLeft;
+
+	// Content that appears on the right side of tabs in this docking stack
+	TSharedRef<SWidget> ContentRight;
+
+	// Content that appears on the right side of the title bar in the window this stack is in
+	TSharedRef<SWidget> TitleBarContentRight;
+};
+
+
 /**
  * A node in the Docking/Tabbing hierarchy.
  * A DockTabStack shows a row of tabs and the content of one selected tab.
@@ -47,7 +63,8 @@ public:
 
 	void OnLastTabRemoved();
 
-	void OnTabClosed( const TSharedRef<SDockTab>& ClosedTab );
+	void OnTabClosed(const TSharedRef<SDockTab>& ClosedTab, SDockingNode::ELayoutModification RemovalMethod);
+
 	void OnTabRemoved( const FTabId& TabId );
 
 	void Construct( const FArguments& InArgs, const TSharedRef<FTabManager::FStack>& PersistentNode );
@@ -58,11 +75,19 @@ public:
 
 	void AddTabWidget( const TSharedRef<SDockTab>& InTab, int32 AtLocation = INDEX_NONE);
 
+	void AddSidebarTab(const TSharedRef<SDockTab>& InTab);
+
+	float GetTabSidebarSizeCoefficient(const TSharedRef<SDockTab>& InTab);
+	void SetTabSidebarSizeCoefficient(const TSharedRef<SDockTab>& InTab, float InSizeCoefficient);
+
+	/** @return true if the specified tab can be moved to a sidebar */
+	bool CanMoveTabToSideBar(TSharedRef<SDockTab> Tab) const;
+
+	/** Moves a specific tab from this stack into a sidebar */
+	void MoveTabToSidebar(TSharedRef<SDockTab> Tab);
+
 	/** @return All child tabs in this node */
 	const TSlotlessChildren<SDockTab>& GetTabs() const;
-	
-	/** @return How many tabs are in this node */
-	int32 GetNumTabs() const;
 
 	bool HasTab( const struct FTabMatcher& TabMatcher ) const;
 
@@ -76,13 +101,16 @@ public:
 	void BringToFront( const TSharedRef<SDockTab>& TabToBringToFront );
 
 	/** Set the content that the DockNode is presenting. */
-	void SetNodeContent( const TSharedRef<SWidget>& InContent, const TSharedRef<SWidget>& InContentLeft, const TSharedRef<SWidget>& InContentRight, const TSharedRef<SWidget>& InContentBackground );
+	void SetNodeContent( const TSharedRef<SWidget>& InContent, const FDockingStackOptionalContent& OptionalContent);
 
 	virtual FReply OnUserAttemptingDock( SDockingNode::RelativeDirection Direction, const FDragDropEvent& DragDropEvent ) override;
 	
 	/** Recursively searches through all children looking for child tabs */
 	virtual TArray< TSharedRef<SDockTab> > GetAllChildTabs() const override;
 
+	/** Gets the number of tabs in all children recursively */
+	virtual int32 GetNumTabs() const override;
+
 	virtual SSplitter::ESizeRule GetSizeRule() const override;
 
 	void SetTabWellHidden( bool bShouldHideTabWell );
@@ -106,24 +134,16 @@ public:
 	 * Add some extra padding so that the corresponding window chrome element
 	 * does not overlap our tabs.
 	 */
-	void ReserveSpaceForWindowChrome(EChromeElement);
+	void ReserveSpaceForWindowChrome(EChromeElement Element, bool bIncludePaddingForMenuBar, bool bOnlyMinorTabs);
 
 public:
-
-	virtual void Tick( const FGeometry& AllottedGeometry, const double InCurrentTime, const float InDeltaTime ) override;
-
+	/** SWidget interface */
 	virtual FReply OnDragOver( const FGeometry& MyGeometry, const FDragDropEvent& DragDropEvent ) override;
-
 	virtual void OnDragLeave( const FDragDropEvent& DragDropEvent ) override;
-	
 	virtual FReply OnDrop( const FGeometry& MyGeometry, const FDragDropEvent& DragDropEvent ) override;
-
 	virtual void OnFocusChanging( const FWeakWidgetPath& PreviousFocusPath, const FWidgetPath& NewWidgetPath, const FFocusEvent& InFocusEvent ) override;
-
 	virtual FReply OnMouseButtonDown( const FGeometry& MyGeometry, const FPointerEvent& MouseEvent ) override;
-
 	virtual FReply OnKeyDown(const FGeometry& MyGeometry, const FKeyEvent& InKeyEvent) override;
-
 	virtual bool SupportsKeyboardFocus() const override { return true; }
 
 protected:
@@ -198,6 +218,10 @@ private:
 	/** Show/Hide the tab well; do it smoothly with an animation */
 	void ToggleTabWellVisibility();
 
+	/** Moves the foreground tab to a sidebar */
+	void MoveForegroundTabToSidebar();
+
+	/** Unhides the tab well, revealing all tab headers */
 	FReply UnhideTabWell();
 
 	/** Only allow hiding the tab well when there is a single tab in it. */
@@ -212,8 +236,11 @@ private:
 	/** Only allow closing all other tabs when there are more then one tab open and the tab is of type Document or Major. */
 	bool CanCloseAllButForegroundTab() const;
 
-	/** Keep around our geometry from the last frame so that we can resize the preview windows correctly */
-	FGeometry TabStackGeometry;
+	/** @return true if the foreground tab can be moved to a sidebar */
+	bool CanMoveForegroundTabToSidebar() const;
+
+	/** @return true if the specified tab can be moved to a sidebar */
+	bool IsTabAllowedInSidebar(TSharedPtr<SDockTab> Tab) const;
 
 	/** The tab well widget shows all tabs, keeps track of the selected tab, allows tab rearranging, etc. */
 	TSharedPtr<class SDockingTabWell> TabWell;
@@ -221,7 +248,7 @@ private:
 	/** The borders that hold any potential inline content areas. */
 	SHorizontalBox::FSlot* InlineContentAreaLeft;
 	SHorizontalBox::FSlot* InlineContentAreaRight;
-	SOverlay::FOverlaySlot* BackgroundContentArea;
+
 	SVerticalBox::FSlot* TitleBarSlot;
 	TSharedPtr<SWidget> TitleBarContent;
 
@@ -251,8 +278,8 @@ private:
 	/** Get the opacity for the button that unhides the tab well */
 	FSlateColor GetUnhideTabWellButtonOpacity() const;
 
-	/** @return Gets the visibility state for spacers that pad out the tab well to make room for title bar widgets */
-	EVisibility GetTitleAreaSpacerVisibility();
+	/** Gets the background behind the tab stack */
+	const FSlateBrush* GetTabStackBorderImage() const;
 
 	/** Visibility of TitleBar spacer based on maximize/restore status of the window.
 	 ** This gives us a little more space to grab the title bar when the window is not maximized
@@ -280,12 +307,15 @@ private:
 
 	/** Tab command list */
 	TSharedPtr<FUICommandList> ActionList;
+
+	/** Whether or not this tab stack is part of the title bar area */
+	bool bShowingTitleBarArea;
 };
 
 
 struct FTabMatcher
 {
-	FTabMatcher( const FTabId& InTabId, ETabState::Type InTabState = static_cast<ETabState::Type>(ETabState::ClosedTab | ETabState::OpenedTab), const bool InTreatIndexNoneAsWildcard = true )
+	FTabMatcher( const FTabId& InTabId, ETabState::Type InTabState = static_cast<ETabState::Type>(ETabState::ClosedTab | ETabState::OpenedTab | ETabState::SidebarTab), const bool InTreatIndexNoneAsWildcard = true )
 		: TabIdToMatch( InTabId )
 		, RequiredTabState( InTabState )
 		, TreatIndexNoneAsWildcard( InTreatIndexNoneAsWildcard )
diff --git a/Engine/Source/Runtime/Slate/Private/Framework/Docking/SDockingTabWell.cpp b/Engine/Source/Runtime/Slate/Private/Framework/Docking/SDockingTabWell.cpp
--- a/Engine/Source/Runtime/Slate/Private/Framework/Docking/SDockingTabWell.cpp
+++ b/Engine/Source/Runtime/Slate/Private/Framework/Docking/SDockingTabWell.cpp
@@ -6,7 +6,7 @@
 #include "Framework/Docking/FDockingDragOperation.h"
 #include "HAL/PlatformApplicationMisc.h"
 
-const FVector2D FDockingConstants::MaxMinorTabSize(160.f, 50.0f);
+const FVector2D FDockingConstants::MaxMinorTabSize(160.f, 25.0f);
 const FVector2D FDockingConstants::MaxMajorTabSize(210.f, 50.f);
 
 const FVector2D FDockingConstants::GetMaxTabSizeFor( ETabRole TabRole )
@@ -614,13 +614,18 @@ void SDockingTabWell::RemoveAndDestroyTab(const TSharedRef<SDockTab>& TabToRemov
 			{
 				BringTabToFront(OldTabIndex);
 			}
+
+			if (RemovalMethod == SDockingNode::ELayoutModification::TabRemoval_Sidebar && ForegroundTabIndex == INDEX_NONE)
+			{
+				FGlobalTabmanager::Get()->SetActiveTab(nullptr);
+			}
 		}
 		
 		if ( ensure(ParentTabStack.IsValid()) )
 		{
 			TSharedPtr<SDockingArea> DockAreaPtr = ParentTabStack->GetDockArea();
 
-			ParentTabStack->OnTabClosed( TabToRemove );
+			ParentTabStack->OnTabClosed( TabToRemove, RemovalMethod );
 			
 			// We might be closing down an entire dock area, if this is a major tab.
 			// Use this opportunity to save its layout
diff --git a/Engine/Source/Runtime/Slate/Private/Framework/Docking/STabDrawer.cpp b/Engine/Source/Runtime/Slate/Private/Framework/Docking/STabDrawer.cpp
new file mode 100644
--- /dev/null
+++ b/Engine/Source/Runtime/Slate/Private/Framework/Docking/STabDrawer.cpp
@@ -0,0 +1,369 @@
+// Copyright Epic Games, Inc. All Rights Reserved.
+
+#include "Framework/Docking/STabDrawer.h"
+#include "Styling/AppStyle.h"
+#include "Styling/SlateTypes.h"
+#include "Layout/ArrangedChildren.h"
+#include "Framework/Application/SlateApplication.h"
+
+STabDrawer::~STabDrawer()
+{
+	FSlateThrottleManager::Get().LeaveResponsiveMode(AnimationThrottle);
+	FSlateThrottleManager::Get().LeaveResponsiveMode(ResizeThrottleHandle);
+}
+
+void STabDrawer::SetCurrentSize(float InSize)
+{
+	CurrentSize = FMath::Clamp(InSize, MinDrawerSize, TargetDrawerSize);
+}
+
+void STabDrawer::Construct(const FArguments& InArgs, TSharedRef<SDockTab> InTab, ETabDrawerOpenDirection InOpenDirection)
+{
+	Visibility = EVisibility::SelfHitTestInvisible;
+	OpenDirection = InOpenDirection;
+
+	ForTab = InTab;
+	OpenCloseAnimation = FCurveSequence(0.0f, 0.15f, ECurveEaseFunction::QuadOut);
+
+	CurrentSize = 0;
+
+	ShadowOffset = InArgs._ShadowOffset;
+	ExpanderSize = 5.0f;
+
+	SplitterStyle = &FAppStyle::Get().GetWidgetStyle<FSplitterStyle>("Splitter");
+
+	MinDrawerSize = InArgs._MinDrawerSize;
+
+	MaxDrawerSize = InArgs._MaxDrawerSize;
+
+	TargetDrawerSize = FMath::Clamp(InArgs._TargetDrawerSize, MinDrawerSize, MaxDrawerSize);
+
+	OnTargetDrawerSizeChanged = InArgs._OnTargetDrawerSizeChanged;
+	OnDrawerFocusLost = InArgs._OnDrawerFocusLost;
+	OnDrawerClosed = InArgs._OnDrawerClosed;
+
+	BackgroundBrush = FAppStyle::Get().GetBrush("Docking.Sidebar.DrawerBackground");
+	ShadowBrush = FAppStyle::Get().GetBrush("Docking.Sidebar.DrawerShadow");
+	BorderBrush = FAppStyle::Get().GetBrush("Docking.Sidebar.Border");
+
+	FSlateApplication::Get().OnFocusChanging().AddSP(this, &STabDrawer::OnGlobalFocusChanging);
+
+	bIsResizeHandleHovered = false;
+	bIsResizing = false;
+
+	ChildSlot
+	[
+		InArgs._Content.Widget
+	];
+}
+
+void STabDrawer::Open()
+{
+	OpenCloseAnimation.Play(AsShared(), false, OpenCloseAnimation.IsPlaying() ? OpenCloseAnimation.GetSequenceTime() : 0.0f, false);
+
+	if (!OpenCloseTimer.IsValid())
+	{
+		AnimationThrottle = FSlateThrottleManager::Get().EnterResponsiveMode();
+		OpenCloseTimer = RegisterActiveTimer(0.0f, FWidgetActiveTimerDelegate::CreateSP(this, &STabDrawer::UpdateAnimation));
+	}
+}
+
+void STabDrawer::Close()
+{
+	if (OpenCloseAnimation.IsForward())
+	{
+		OpenCloseAnimation.Reverse();
+	}
+
+	if (!OpenCloseTimer.IsValid())
+	{
+		AnimationThrottle = FSlateThrottleManager::Get().EnterResponsiveMode();
+		OpenCloseTimer = RegisterActiveTimer(0.0f, FWidgetActiveTimerDelegate::CreateSP(this, &STabDrawer::UpdateAnimation));
+	}
+
+	//OnContentBrowserDismissedDelegate.ExecuteIfBound(NewlyFocusedWidget);
+}
+
+bool STabDrawer::IsOpen() const
+{
+	return !OpenCloseAnimation.IsAtStart();
+}
+
+const TSharedRef<SDockTab> STabDrawer::GetTab() const
+{
+	return ForTab.ToSharedRef();
+}
+
+bool STabDrawer::SupportsKeyboardFocus() const 
+{
+	return true;
+}
+
+FVector2D STabDrawer::ComputeDesiredSize(float) const
+{
+	if (OpenDirection == ETabDrawerOpenDirection::Bottom)
+	{
+		return FVector2D(1.0f, TargetDrawerSize + ShadowOffset.Y);
+	}
+	else
+	{
+		return FVector2D(TargetDrawerSize + ShadowOffset.X, 1.0f);
+	}
+}
+
+void STabDrawer::OnArrangeChildren(const FGeometry& AllottedGeometry, FArrangedChildren& ArrangedChildren) const 
+{
+	const EVisibility ChildVisibility = ChildSlot.GetWidget()->GetVisibility();
+	if (ArrangedChildren.Accepts(ChildVisibility))
+	{
+		if (OpenDirection == ETabDrawerOpenDirection::Left)
+		{
+			ArrangedChildren.AddWidget(
+				AllottedGeometry.MakeChild(
+					ChildSlot.GetWidget(),
+					FVector2D(0.0f, ShadowOffset.Y),
+					FVector2D(TargetDrawerSize, AllottedGeometry.GetLocalSize().Y - (ShadowOffset.Y * 2))
+				)
+			);
+		}
+		else if (OpenDirection == ETabDrawerOpenDirection::Right)
+		{
+			ArrangedChildren.AddWidget(
+				AllottedGeometry.MakeChild(
+					ChildSlot.GetWidget(),
+					FVector2D(ShadowOffset),
+					FVector2D(TargetDrawerSize, AllottedGeometry.GetLocalSize().Y - (ShadowOffset.Y * 2))
+				)
+			);
+		}
+		else
+		{
+			ArrangedChildren.AddWidget(
+				AllottedGeometry.MakeChild(
+					ChildSlot.GetWidget(),
+					ShadowOffset,
+					FVector2D(AllottedGeometry.GetLocalSize().X - (ShadowOffset.X * 2), TargetDrawerSize)
+				)
+			);
+		}
+		
+	}
+}
+
+FReply STabDrawer::OnMouseButtonDown(const FGeometry& AllottedGeometry, const FPointerEvent& MouseEvent) 
+{
+	FReply Reply = FReply::Unhandled();
+	if (MouseEvent.GetEffectingButton() == EKeys::LeftMouseButton)
+	{
+		const FGeometry RenderTransformedChildGeometry = GetRenderTransformedGeometry(AllottedGeometry);
+		const FGeometry ResizeHandleGeometry = GetResizeHandleGeometry(AllottedGeometry);
+
+		if (ResizeHandleGeometry.IsUnderLocation(MouseEvent.GetScreenSpacePosition()))
+		{
+			bIsResizing = true;
+			InitialResizeGeometry = ResizeHandleGeometry;
+			InitialSizeAtResize = CurrentSize;
+			ResizeThrottleHandle = FSlateThrottleManager::Get().EnterResponsiveMode();
+
+			Reply = FReply::Handled().CaptureMouse(SharedThis(this));
+		}
+	}
+
+	return Reply;
+
+}
+
+FReply STabDrawer::OnMouseButtonUp(const FGeometry& AllottedGeometry, const FPointerEvent& MouseEvent) 
+{
+	if (MouseEvent.GetEffectingButton() == EKeys::LeftMouseButton && bIsResizing == true)
+	{
+		bIsResizing = false;
+		FSlateThrottleManager::Get().LeaveResponsiveMode(ResizeThrottleHandle);
+
+		OnTargetDrawerSizeChanged.ExecuteIfBound(TargetDrawerSize);
+		return FReply::Handled().ReleaseMouseCapture();
+	}
+	return FReply::Unhandled();
+}
+
+FReply STabDrawer::OnMouseMove(const FGeometry& AllottedGeometry, const FPointerEvent& MouseEvent)
+{
+	const FGeometry ResizeHandleGeometry = GetResizeHandleGeometry(AllottedGeometry);
+
+	bIsResizeHandleHovered = ResizeHandleGeometry.IsUnderLocation(MouseEvent.GetScreenSpacePosition());
+
+	if (bIsResizing && this->HasMouseCapture() && !MouseEvent.GetCursorDelta().IsZero())
+	{
+		const FVector2D LocalMousePos = InitialResizeGeometry.AbsoluteToLocal(MouseEvent.GetScreenSpacePosition());
+		const float DeltaHeight = (InitialResizeGeometry.GetLocalPositionAtCoordinates(FVector2D::ZeroVector) - LocalMousePos).Y;
+
+		TargetDrawerSize = FMath::Clamp(InitialSizeAtResize + DeltaHeight, MinDrawerSize, MaxDrawerSize);
+		SetCurrentSize(InitialSizeAtResize + DeltaHeight);
+
+
+		return FReply::Handled();
+	}
+	else
+	{
+		return FReply::Unhandled();
+	}
+}
+
+void STabDrawer::OnMouseLeave(const FPointerEvent& MouseEvent)
+{
+	SCompoundWidget::OnMouseLeave(MouseEvent);
+
+	bIsResizeHandleHovered = false;
+}
+
+FCursorReply STabDrawer::OnCursorQuery(const FGeometry& MyGeometry, const FPointerEvent& CursorEvent) const
+{
+	return bIsResizing || bIsResizeHandleHovered ? FCursorReply::Cursor(EMouseCursor::ResizeUpDown) : FCursorReply::Unhandled();
+}
+
+int32 STabDrawer::OnPaint(const FPaintArgs& Args, const FGeometry& AllottedGeometry, const FSlateRect& MyCullingRect, FSlateWindowElementList& OutDrawElements, int32 LayerId, const FWidgetStyle& InWidgetStyle, bool bParentEnabled) const
+{
+	static FSlateColor ShadowColor = FAppStyle::Get().GetSlateColor("Colors.Foldout");
+
+	const FGeometry RenderTransformedChildGeometry = GetRenderTransformedGeometry(AllottedGeometry);
+	const FGeometry ResizeHandleGeometry = GetResizeHandleGeometry(AllottedGeometry);
+
+	FPaintGeometry OffsetPaintGeom;
+	if (OpenDirection == ETabDrawerOpenDirection::Left)
+	{
+		OffsetPaintGeom = RenderTransformedChildGeometry.ToPaintGeometry(FVector2D(0.0f, ShadowOffset.Y), FVector2D(TargetDrawerSize, AllottedGeometry.GetLocalSize().Y - (ShadowOffset.Y * 2)));
+	}
+	else if (OpenDirection == ETabDrawerOpenDirection::Right)
+	{
+		OffsetPaintGeom = RenderTransformedChildGeometry.ToPaintGeometry(ShadowOffset, FVector2D(TargetDrawerSize, AllottedGeometry.GetLocalSize().Y - (ShadowOffset.Y * 2)));
+	}
+	else
+	{
+		OffsetPaintGeom = RenderTransformedChildGeometry.ToPaintGeometry(ShadowOffset, FVector2D(AllottedGeometry.GetLocalSize().X - (ShadowOffset.X * 2), TargetDrawerSize));
+	}
+
+	// Draw the resize handle
+	if (bIsResizing || bIsResizeHandleHovered)
+	{
+		const FSlateBrush* SplitterBrush = &SplitterStyle->HandleHighlightBrush;
+		FSlateDrawElement::MakeBox(
+			OutDrawElements,
+			LayerId,
+			ResizeHandleGeometry.ToPaintGeometry(),
+			SplitterBrush,
+			ESlateDrawEffect::None,
+			SplitterBrush->GetTint(InWidgetStyle));
+	}
+
+	// Main Shadow
+	FSlateDrawElement::MakeBox(
+		OutDrawElements,
+		LayerId,
+		RenderTransformedChildGeometry.ToPaintGeometry(),
+		ShadowBrush,
+		ESlateDrawEffect::None,
+		ShadowBrush->GetTint(InWidgetStyle));
+
+
+	// Background
+	FSlateDrawElement::MakeBox(
+		OutDrawElements,
+		LayerId,
+		OffsetPaintGeom,
+		BackgroundBrush,
+		ESlateDrawEffect::None,
+		BackgroundBrush->GetTint(InWidgetStyle));
+
+	int32 OutLayerId = SCompoundWidget::OnPaint(Args, RenderTransformedChildGeometry, MyCullingRect, OutDrawElements, LayerId, InWidgetStyle, bParentEnabled);
+
+	// Top border
+	FSlateDrawElement::MakeBox(
+		OutDrawElements,
+		OutLayerId,
+		OffsetPaintGeom,
+		BorderBrush,
+		ESlateDrawEffect::None,
+		BorderBrush->GetTint(InWidgetStyle));
+
+	return OutLayerId+1;
+
+}
+
+FGeometry STabDrawer::GetRenderTransformedGeometry(const FGeometry& AllottedGeometry) const
+{
+	if(OpenDirection == ETabDrawerOpenDirection::Left)
+	{
+		return AllottedGeometry.MakeChild(FSlateRenderTransform(FVector2D(CurrentSize - TargetDrawerSize, 0.0f)));
+	}
+	else if (OpenDirection == ETabDrawerOpenDirection::Right)
+	{
+		return AllottedGeometry.MakeChild(FSlateRenderTransform(FVector2D(TargetDrawerSize - CurrentSize, 0.0f)));
+	}
+	else
+	{
+		return AllottedGeometry.MakeChild(FSlateRenderTransform(FVector2D(0.0f, TargetDrawerSize - CurrentSize)));
+	}
+}
+
+FGeometry STabDrawer::GetResizeHandleGeometry(const FGeometry& AllottedGeometry) const
+{
+	return GetRenderTransformedGeometry(AllottedGeometry).MakeChild(ShadowOffset - FVector2D(0.0f, ExpanderSize), FVector2D(AllottedGeometry.GetLocalSize().X - ShadowOffset.X * 2, ExpanderSize));
+}
+
+EActiveTimerReturnType STabDrawer::UpdateAnimation(double CurrentTime, float DeltaTime)
+{
+	SetCurrentSize(FMath::Lerp(0.0f, TargetDrawerSize, OpenCloseAnimation.GetLerp()));
+
+	if (!OpenCloseAnimation.IsPlaying())
+	{
+		if (OpenCloseAnimation.IsAtStart())
+		{
+			OnDrawerClosed.ExecuteIfBound(SharedThis(this));
+		}
+
+		FSlateThrottleManager::Get().LeaveResponsiveMode(AnimationThrottle);
+		OpenCloseTimer.Reset();
+		return EActiveTimerReturnType::Stop;
+	}
+
+	return EActiveTimerReturnType::Continue;
+}
+
+void STabDrawer::OnGlobalFocusChanging(const FFocusEvent& FocusEvent, const FWeakWidgetPath& OldFocusedWidgetPath, const TSharedPtr<SWidget>& OldFocusedWidget, const FWidgetPath& NewFocusedWidgetPath, const TSharedPtr<SWidget>& NewFocusedWidget)
+{
+	// Sometimes when dismissing focus can change which will trigger this again
+	static bool bIsRentrant = false;
+
+	if (!bIsRentrant)
+	{
+		TGuardValue<bool> RentrancyGuard(bIsRentrant, true);
+
+		TSharedRef<STabDrawer> ThisWidget = SharedThis(this);
+
+		bool bShouldLoseFocus = false;
+		// Do not close due to slow tasks as those opening send window activation events
+		if (!GIsSlowTask && !FSlateApplication::Get().GetActiveModalWindow().IsValid() && (!NewFocusedWidgetPath.ContainsWidget(ChildSlot.GetWidget()) && !NewFocusedWidgetPath.ContainsWidget(ThisWidget)))
+		{
+			if (TSharedPtr<SWidget> MenuHost = FSlateApplication::Get().GetMenuHostWidget())
+			{
+				FWidgetPath MenuHostPath;
+
+				// See if the menu being opened is owned by the drawer contents and if so the menu should not be dismissed
+				FSlateApplication::Get().GeneratePathToWidgetUnchecked(MenuHost.ToSharedRef(), MenuHostPath);
+				if (!MenuHostPath.ContainsWidget(ChildSlot.GetWidget()))
+				{
+					bShouldLoseFocus = true;
+				}
+			}
+			else
+			{
+				bShouldLoseFocus = true;
+			}
+		}
+
+		if (bShouldLoseFocus)
+		{
+			OnDrawerFocusLost.ExecuteIfBound(ThisWidget);
+		}
+	}
+}
diff --git a/Engine/Source/Runtime/Slate/Private/Framework/Docking/STabSidebar.cpp b/Engine/Source/Runtime/Slate/Private/Framework/Docking/STabSidebar.cpp
new file mode 100644
--- /dev/null
+++ b/Engine/Source/Runtime/Slate/Private/Framework/Docking/STabSidebar.cpp
@@ -0,0 +1,581 @@
+// Copyright Epic Games, Inc. All Rights Reserved.
+
+#include "STabSidebar.h"
+#include "Widgets/Text/STextBlock.h"
+#include "Widgets/Input/SButton.h"
+#include "Widgets/SBoxPanel.h"
+#include "SDockingTabWell.h"
+#include "Widgets/Layout/SBox.h"
+#include "Widgets/Images/SImage.h"
+#include "Framework/Application/SlateApplication.h"
+#include "Framework/Docking/STabDrawer.h"
+#include "Framework/MultiBox/MultiBoxBuilder.h"
+#include "Misc/App.h"
+#include "Widgets/Colors/SComplexGradient.h"
+
+#define LOCTEXT_NAMESPACE "TabSidebar"
+
+DECLARE_DELEGATE_OneParam(FOnTabDrawerButtonClicked, TSharedRef<SDockTab>);
+
+class STabDrawerButton : public SCompoundWidget
+{
+
+	SLATE_BEGIN_ARGS(STabDrawerButton)
+	{}
+		SLATE_EVENT(FOnTabDrawerButtonClicked, OnDrawerButtonClicked)
+		SLATE_EVENT(FOnGetContent, OnGetContextMenuContent)
+	SLATE_END_ARGS()
+
+public:
+	void Construct(const FArguments& InArgs, TSharedRef<SDockTab> ForTab, ESidebarLocation InLocation)
+	{
+		const FVector2D Size = FDockingConstants::GetMaxTabSizeFor(ETabRole::PanelTab);
+
+		DockTabStyle = &FAppStyle::Get().GetWidgetStyle<FDockTabStyle>("Docking.Tab");
+
+		OnDrawerButtonClicked = InArgs._OnDrawerButtonClicked;
+		OnGetContextMenuContent = InArgs._OnGetContextMenuContent;
+		Tab = ForTab;
+		Location = InLocation;
+
+		static FLinearColor ActiveBorderColor = FAppStyle::Get().GetSlateColor("Docking.Tab.ActiveTabIndicatorColor").GetSpecifiedColor();
+		static FLinearColor ActiveBorderColorTransparent = FLinearColor(ActiveBorderColor.R, ActiveBorderColor.G, ActiveBorderColor.B, 0.0f);
+		static TArray<FLinearColor> GradientStops{ ActiveBorderColorTransparent, ActiveBorderColor, ActiveBorderColorTransparent };
+
+		ChildSlot
+		.Padding(0, 0, 0, 0)
+		[
+			SNew(SBox)
+			.WidthOverride(Size.X)
+			.HeightOverride(Size.Y)
+			.Clipping(EWidgetClipping::ClipToBounds)
+			[
+				SAssignNew(MainButton, SButton)
+				.ContentPadding(FMargin(0.0f, DockTabStyle->TabPadding.Top, 0.0f, DockTabStyle->TabPadding.Bottom))
+				.OnClicked_Lambda([this](){OnDrawerButtonClicked.ExecuteIfBound(Tab.ToSharedRef()); return FReply::Handled(); })
+				.ForegroundColor(FSlateColor::UseForeground())
+				[
+					SNew(SOverlay)
+					+ SOverlay::Slot()
+					.HAlign(Location == ESidebarLocation::Left ? HAlign_Left : HAlign_Right)
+					[
+						SAssignNew(OpenIndicator, SComplexGradient)
+						.DesiredSizeOverride(FVector2D(1.0f, 1.0f))
+						.GradientColors(GradientStops)
+						.Orientation(EOrientation::Orient_Horizontal)
+					]
+					+ SOverlay::Slot()
+					[
+						SNew(SVerticalBox)
+						+ SVerticalBox::Slot()
+						.AutoHeight()
+						.VAlign(VAlign_Center)
+						.HAlign(HAlign_Center)
+						.Padding(0.0f, 5.0f, 0.0f, 5.0f)
+						[
+							SNew(SImage)
+							.ColorAndOpacity(FSlateColor::UseForeground())
+							.Image(ForTab->GetTabIcon())
+							.DesiredSizeOverride(FVector2D(16,16))
+							//.RenderTransform(Rotate90)
+							//.RenderTransformPivot(FVector2D(.5f, .5f))
+						]
+						+ SVerticalBox::Slot()
+						.Padding(0, 10, 0, 0)
+						.HAlign(HAlign_Left)
+						.VAlign(VAlign_Top)
+						[
+							SAssignNew(Label, STextBlock)
+							.TextStyle(&DockTabStyle->TabTextStyle)
+							.Text(ForTab->GetTabLabel())
+							.Clipping(EWidgetClipping::Inherit)
+							.RenderTransformPivot(FVector2D(0.5f, 0.5f))
+						]
+					]
+				]
+			]
+		];
+
+		UpdateAppearance(nullptr);
+	}
+
+	void UpdateAppearance(const TSharedPtr<SDockTab> OpenedDrawer)
+	{
+		bool bShouldAppearOpened = OpenedDrawer.IsValid();
+
+		switch (Location)
+		{
+		case ESidebarLocation::Left:
+			bShouldAppearOpened ? Rotate90CounterClockwise() : Rotate90Clockwise();
+			break;
+		case ESidebarLocation::Right:
+		default:
+			bShouldAppearOpened ? Rotate90Clockwise() : Rotate90CounterClockwise();
+			break;
+		}
+
+		if (OpenedDrawer == Tab)
+		{
+			// this button is the one with the tab that is actually opened so show the opened indicator
+			OpenIndicator->SetVisibility(EVisibility::HitTestInvisible);
+			MainButton->SetButtonStyle(&FAppStyle::Get().GetWidgetStyle<FButtonStyle>("Docking.SidebarButton.Opened"));
+		}
+		else
+		{
+			OpenIndicator->SetVisibility(EVisibility::Collapsed);
+			MainButton->SetButtonStyle(&FAppStyle::Get().GetWidgetStyle<FButtonStyle>("Docking.SidebarButton.Closed"));
+		}
+	}
+
+	virtual FVector2D ComputeDesiredSize(float LayoutScaleMultiplier) const override
+	{
+		const FVector2D Size = FDockingConstants::GetMaxTabSizeFor(ETabRole::PanelTab);
+
+		FVector2D LocalDesiredSize = SCompoundWidget::ComputeDesiredSize(LayoutScaleMultiplier);
+		return FVector2D(LocalDesiredSize.Y, FMath::Min(LocalDesiredSize.X, Size.X));
+	}
+
+	virtual FReply OnMouseButtonDown(const FGeometry& MyGeometry, const FPointerEvent& MouseEvent)
+	{
+		if (MouseEvent.GetEffectingButton() == EKeys::RightMouseButton && OnGetContextMenuContent.IsBound())
+		{
+			FWidgetPath WidgetPath = MouseEvent.GetEventPath() != nullptr ? *MouseEvent.GetEventPath() : FWidgetPath();
+			FSlateApplication::Get().PushMenu(AsShared(), WidgetPath, OnGetContextMenuContent.Execute(), FSlateApplication::Get().GetCursorPos(), FPopupTransitionEffect::ContextMenu);
+			return FReply::Handled();
+		}
+		else
+		{
+			return FReply::Unhandled();
+		}
+	}
+
+	virtual FSlateColor GetForegroundColor() const
+	{
+		if (OpenIndicator->GetVisibility() != EVisibility::Collapsed)
+		{
+			return DockTabStyle->ActiveForegroundColor;
+		}
+
+		return FSlateColor::UseStyle();
+	}
+private:
+	void Rotate90Clockwise()
+	{
+		Label->SetJustification(ETextJustify::Left);
+		Label->SetRenderTransform(FSlateRenderTransform(FQuat2D(FMath::DegreesToRadians(90))));
+	}
+
+	void Rotate90CounterClockwise()
+	{
+		Label->SetJustification(ETextJustify::Right);
+		Label->SetRenderTransform(FSlateRenderTransform(FQuat2D(FMath::DegreesToRadians(-90))));
+	}
+private:
+	TSharedPtr<SDockTab> Tab;
+	TSharedPtr<STextBlock> Label;
+	TSharedPtr<SWidget> OpenIndicator;
+	TSharedPtr<SButton> MainButton;
+	FOnGetContent OnGetContextMenuContent;
+	FOnTabDrawerButtonClicked OnDrawerButtonClicked;
+	const FDockTabStyle* DockTabStyle;
+	ESidebarLocation Location;
+};
+
+STabSidebar::~STabSidebar()
+{
+	// ensure all drawers are removed when closing a sidebar
+	RemoveAllDrawers();
+}
+
+void STabSidebar::Construct(const FArguments& InArgs)
+{
+	Location = InArgs._Location;
+
+#if WITH_EDITOR
+	FSlateApplication::Get().OnWindowDPIScaleChanged().AddSP(this, &STabSidebar::OnWindowDPIScaleChanged);
+#endif
+
+	FGlobalTabmanager::Get()->OnTabForegrounded_Subscribe(FOnActiveTabChanged::FDelegate::CreateSP(this, &STabSidebar::OnActiveTabChanged));
+
+	ChildSlot
+	.Padding(2.0f,0.0f)
+	[
+		SNew(SBorder)
+		.Padding(0.0f)
+		.BorderImage(FAppStyle::Get().GetBrush("Docking.Sidebar.Background"))
+		[
+			SAssignNew(TabBox, SVerticalBox)
+		]
+	];
+
+}
+
+void STabSidebar::SetOffset(float Offset)
+{
+	ChildSlot.Padding(0.0f, Offset+4, 0.0f, 0.0f);
+}
+
+void STabSidebar::AddTab(TSharedRef<SDockTab> Tab)
+{
+	if(!ContainsTab(Tab))
+	{
+		SetVisibility(EVisibility::SelfHitTestInvisible);
+
+		TSharedRef<STabDrawerButton> TabButton =
+			SNew(STabDrawerButton, Tab, Location)
+			.OnDrawerButtonClicked(this, &STabSidebar::OnTabDrawerButtonClicked)
+			.OnGetContextMenuContent(this, &STabSidebar::OnGetTabDrawerContextMenuWidget, Tab);
+
+		// Figure out the size this tab should be when opened later. We do it now when the tab still has valid geometry.  Once it is moved to the sidebar it will not.
+		float TargetDrawerSizePct = Tab->GetParentDockTabStack()->GetTabSidebarSizeCoefficient(Tab);
+		if (TargetDrawerSizePct == 0)
+		{
+			TSharedPtr<SWindow> MyWindow = FSlateApplication::Get().FindWidgetWindow(AsShared());
+			if (MyWindow.IsValid())
+			{
+				TargetDrawerSizePct = Tab->GetParentDockTabStack()->GetPaintSpaceGeometry().GetLocalSize().X / MyWindow->GetPaintSpaceGeometry().GetLocalSize().X;
+				Tab->GetParentDockTabStack()->SetTabSidebarSizeCoefficient(Tab, TargetDrawerSizePct);
+			}
+		}
+
+		TabBox->AddSlot()
+			.AutoHeight()
+			.HAlign(HAlign_Left)
+			[
+				TabButton
+			];
+
+		Tabs.Emplace(Tab, TabButton);
+	}
+}
+
+bool STabSidebar::RemoveTab(TSharedRef<SDockTab> TabToRemove)
+{
+	int32 FoundIndex = Tabs.IndexOfByPredicate(
+		[TabToRemove](auto TabPair)
+		{
+			return TabPair.Key == TabToRemove;
+		});
+
+
+	if(FoundIndex != INDEX_NONE)
+	{
+		TPair<TSharedRef<SDockTab>, TSharedRef<STabDrawerButton>> TabPair = Tabs[FoundIndex];
+
+		Tabs.RemoveAt(FoundIndex);
+		TabBox->RemoveSlot(TabPair.Value);
+
+		RemoveDrawer(TabToRemove);
+
+		if (Tabs.Num() == 0)
+		{
+			SetVisibility(EVisibility::Collapsed);
+		}
+	}
+
+	return FoundIndex != INDEX_NONE;
+}
+
+bool STabSidebar::RestoreTab(TSharedRef<SDockTab> TabToRestore)
+{
+	if(RemoveTab(TabToRestore))
+	{
+		TabToRestore->GetParentDockTabStack()->SetTabSidebarSizeCoefficient(TabToRestore, 0);
+		TabToRestore->GetParentDockTabStack()->OpenTab(TabToRestore);
+
+		return true;
+	}
+
+	return false;
+}
+
+bool STabSidebar::ContainsTab(TSharedPtr<SDockTab> Tab) const
+{
+	return Tabs.ContainsByPredicate(
+		[Tab](auto TabPair)
+		{
+			return TabPair.Key == Tab;
+		});
+}
+
+TArray<FTabId> STabSidebar::GetAllTabIds() const
+{
+	TArray<FTabId> TabIds;
+	for (auto TabPair : Tabs)
+	{
+		TabIds.Add(TabPair.Key->GetLayoutIdentifier());
+	}
+
+	return TabIds;
+}
+
+bool STabSidebar::TryOpenSidebarDrawer(TSharedRef<SDockTab> ForTab)
+{
+	int32 FoundIndex = Tabs.IndexOfByPredicate(
+		[ForTab](auto TabPair)
+		{
+			return TabPair.Key == ForTab;
+		});
+
+	if (FoundIndex != INDEX_NONE)
+	{
+		OpenDrawerNextFrame(ForTab);
+		return true;
+	}
+
+	return false;
+}
+
+void STabSidebar::OnTabDrawerButtonClicked(TSharedRef<SDockTab> ForTab)
+{
+	OpenDrawerInternal(ForTab);
+}
+
+void STabSidebar::OnTabDrawerFocusLost(TSharedRef<STabDrawer> Drawer)
+{
+	Drawer->Close();
+}
+
+void STabSidebar::OnTabDrawerClosed(TSharedRef<STabDrawer> Drawer)
+{
+	RemoveDrawer(Drawer->GetTab());
+}
+
+void STabSidebar::OnWindowDPIScaleChanged(TSharedRef<SWindow> WindowThatChanged)
+{
+	if (WindowThatChanged == WindowWithOverlayContent)
+	{
+		RemoveAllDrawers();
+	}
+}
+
+void STabSidebar::OnActiveTabChanged(TSharedPtr<SDockTab> PreviouslyActive, TSharedPtr<SDockTab> NewlyActivated)
+{
+	// If a new major tab was activated remove any visible drawer instantly
+	if (NewlyActivated.IsValid() && NewlyActivated->GetVisualTabRole() == ETabRole::MajorTab)
+	{
+		RemoveAllDrawers();
+	}
+}
+
+TSharedRef<SWidget> STabSidebar::OnGetTabDrawerContextMenuWidget(TSharedRef<SDockTab> ForTab)
+{
+	const bool bCloseAfterSelection = true;
+	const bool bCloseSelfOnly = false;
+	FMenuBuilder MenuBuilder(bCloseAfterSelection, nullptr, TSharedPtr<FExtender>(), bCloseSelfOnly, &FAppStyle::Get());
+	{
+		MenuBuilder.BeginSection("RestoreOptions", LOCTEXT("RestoreOptions", "Options"));
+		{
+			MenuBuilder.AddMenuEntry(
+				LOCTEXT("AutoHideTab", "Restore Tab"),
+				LOCTEXT("HideTabWellTooltip", "Moves this tab out of the sidebar and back to a full tab where it previously was before it was added to the sidebar."),
+				FSlateIcon(),
+				FUIAction(
+					FExecuteAction::CreateSP(this, &STabSidebar::OnRestoreTab, ForTab)
+				)
+			);
+
+		}
+		MenuBuilder.EndSection();
+
+		MenuBuilder.BeginSection("CloseOptions");
+		{
+			MenuBuilder.AddMenuEntry(
+				LOCTEXT("CloseTab", "Close Tab"),
+				LOCTEXT("CloseTabTooltip", "Close this tab, removing it from the sidebar and its parent tab well."),
+				FSlateIcon(),
+				FUIAction(
+					FExecuteAction::CreateSP(this, &STabSidebar::OnCloseTab, ForTab)
+				)
+			);
+		}
+		MenuBuilder.EndSection();
+	}
+
+	return MenuBuilder.MakeWidget();
+}
+
+void STabSidebar::OnRestoreTab(TSharedRef<SDockTab> TabToRestore)
+{
+	RestoreTab(TabToRestore);
+}
+
+void STabSidebar::OnCloseTab(TSharedRef<SDockTab> TabToClose)
+{
+	RemoveTab(TabToClose);
+
+	TabToClose->RequestCloseTab();
+	TabToClose->GetParentDockTabStack()->OnTabClosed(TabToClose, SDockingNode::TabRemoval_Closed);
+}
+
+void STabSidebar::RemoveDrawer(TSharedRef<SDockTab> ForTab)
+{
+	TSharedRef<STabDrawer>* OpenedDrawer =
+		OpenedDrawers.FindByPredicate(
+			[&ForTab](TSharedRef<STabDrawer>& Drawer)
+			{
+				return ForTab == Drawer->GetTab();
+			}
+	);
+
+	if(OpenedDrawer)
+	{
+		TSharedRef<STabDrawer> OpenedDrawerRef = *OpenedDrawer;
+
+		if (TSharedPtr<SWindow> MyWindow = WindowWithOverlayContent.Pin())
+		{
+			bool bRemoveSuccessful = MyWindow->RemoveOverlaySlot(OpenedDrawerRef);
+			ensure(bRemoveSuccessful);
+		}
+
+		OpenedDrawers.Remove(OpenedDrawerRef);
+	}
+
+	UpdateDrawerAppearance();
+}
+
+void STabSidebar::CloseAllDrawers()
+{
+	PendingTabToOpen.Reset();
+
+	// Closing drawers can remove them from the opened drawers list so copy the list first
+	TArray<TSharedRef<STabDrawer>> OpenedDrawersCopy = OpenedDrawers;
+
+	for (TSharedRef<STabDrawer>& Drawer : OpenedDrawersCopy)
+	{
+		Drawer->Close();
+	}
+}
+
+void STabSidebar::RemoveAllDrawers()
+{
+	PendingTabToOpen.Reset();
+
+	// Closing drawers can remove them from the opened drawers list so copy the list first
+	TArray<TSharedRef<STabDrawer>> OpenedDrawersCopy = OpenedDrawers;
+
+	for (TSharedRef<STabDrawer>& Drawer : OpenedDrawersCopy)
+	{
+		RemoveDrawer(Drawer->GetTab());
+	}
+}
+
+EActiveTimerReturnType STabSidebar::OnOpenPendingDrawerTimer(double CurrentTime, float DeltaTime)
+{
+	if (TSharedPtr<SDockTab> Tab = PendingTabToOpen.Pin())
+	{
+		OpenDrawerInternal(Tab.ToSharedRef());
+	}
+
+	OpenPendingDrawerTimerHandle.Reset();
+	PendingTabToOpen.Reset();
+
+	return EActiveTimerReturnType::Stop;
+}
+
+void STabSidebar::OpenDrawerNextFrame(TSharedRef<SDockTab> ForTab)
+{
+	PendingTabToOpen = ForTab;
+	if (!OpenPendingDrawerTimerHandle.IsValid())
+	{
+		OpenPendingDrawerTimerHandle = RegisterActiveTimer(0.0f, FWidgetActiveTimerDelegate::CreateSP(this, &STabSidebar::OnOpenPendingDrawerTimer));
+	}
+}
+
+void STabSidebar::OpenDrawerInternal(TSharedRef<SDockTab> ForTab)
+{
+	TSharedRef<STabDrawer>* OpenedDrawer =
+		OpenedDrawers.FindByPredicate(
+			[&ForTab](TSharedRef<STabDrawer>& Drawer)
+			{
+				return ForTab == Drawer->GetTab();
+			}
+	);
+
+	if (OpenedDrawer)
+	{
+		// Drawer already opened close it
+		(*OpenedDrawer)->Close();
+	}
+	else
+	{
+		PendingTabToOpen.Reset();
+
+		TSharedPtr<SWindow> MyWindow = FSlateApplication::Get().FindWidgetWindow(AsShared());
+
+		const FGeometry WindowGeometry = MyWindow->GetTickSpaceGeometry();
+		const FGeometry MyGeometry = GetTickSpaceGeometry();
+
+		// Calculate padding for the drawer itself
+		const float MinDrawerSize = MyGeometry.GetLocalSize().X + MyWindow->GetWindowBorderSize().Left;
+
+		const FVector2D ShadowOffset(8, 8);
+
+		const float DPIScale = MyWindow->GetDPIScaleFactor();
+
+		const float TopOffset = ((ChildSlot.SlotPadding.Get().Top + MyGeometry.GetAbsolutePosition().Y) - WindowGeometry.GetAbsolutePosition().Y);
+
+		const float BottomOffset = ((ChildSlot.SlotPadding.Get().Bottom + WindowGeometry.GetAbsolutePositionAtCoordinates(FVector2D::UnitVector).Y) - MyGeometry.GetAbsolutePositionAtCoordinates(FVector2D::UnitVector).Y);
+
+		FMargin SlotPadding(
+			Location == ESidebarLocation::Left ? MinDrawerSize : 0.0f,
+			TopOffset / DPIScale - ShadowOffset.Y,
+			Location == ESidebarLocation::Right ? MinDrawerSize : 0.0f,
+			BottomOffset / DPIScale - ShadowOffset.Y
+		);
+
+		const float MaxPct = .5f;
+		const float MaxDrawerSize = MyWindow->GetSizeInScreen().X * 0.50f;
+
+		float TargetDrawerSizePct = ForTab->GetParentDockTabStack()->GetTabSidebarSizeCoefficient(ForTab);
+		TargetDrawerSizePct = FMath::Clamp(TargetDrawerSizePct, .0f, .5f);
+
+
+		const float TargetDrawerSize = (MyWindow->GetSizeInScreen().X * TargetDrawerSizePct) / MyWindow->GetDPIScaleFactor();
+
+		TSharedRef<STabDrawer> NewDrawer =
+			SNew(STabDrawer, ForTab, Location == ESidebarLocation::Left ? ETabDrawerOpenDirection::Left : ETabDrawerOpenDirection::Right)
+			.MinDrawerSize(MinDrawerSize)
+			.TargetDrawerSize(TargetDrawerSize)
+			.MaxDrawerSize(MaxDrawerSize)
+			.OnDrawerFocusLost(this, &STabSidebar::OnTabDrawerFocusLost)
+			.OnDrawerClosed(this, &STabSidebar::OnTabDrawerClosed)
+			//.OnTargetHeightChanged(this, &SStatusBar::OnContentBrowserTargetHeightChanged)
+			[
+				ForTab->GetContent()
+			];
+
+		WindowWithOverlayContent = MyWindow;
+
+		MyWindow->AddOverlaySlot()
+			.Padding(SlotPadding)
+			.HAlign(Location == ESidebarLocation::Left ? HAlign_Left : HAlign_Right)
+			[
+				NewDrawer
+			];
+
+		NewDrawer->Open();
+
+		FSlateApplication::Get().SetKeyboardFocus(NewDrawer);
+
+		OpenedDrawers.Add(NewDrawer);
+	}
+
+	UpdateDrawerAppearance();
+}
+
+void STabSidebar::UpdateDrawerAppearance()
+{
+	TSharedPtr<SDockTab> OpenedTab;
+	if (OpenedDrawers.Num() > 0)
+	{
+		OpenedTab = OpenedDrawers.Last()->GetTab();
+	}
+
+	for (auto& TabPair : Tabs)
+	{
+		TabPair.Value->UpdateAppearance(OpenedTab);
+	}
+}
+
+#undef LOCTEXT_NAMESPACE
diff --git a/Engine/Source/Runtime/Slate/Private/Framework/Docking/STabSidebar.h b/Engine/Source/Runtime/Slate/Private/Framework/Docking/STabSidebar.h
new file mode 100644
--- /dev/null
+++ b/Engine/Source/Runtime/Slate/Private/Framework/Docking/STabSidebar.h
@@ -0,0 +1,118 @@
+// Copyright Epic Games, Inc. All Rights Reserved.
+
+#pragma once
+
+#include "CoreMinimal.h"
+#include "Widgets/DeclarativeSyntaxSupport.h"
+#include "Widgets/SCompoundWidget.h"
+
+class SDockTab;
+class STabDrawerButton;
+class SVerticalBox;
+class STabDrawer;
+struct FTabId;
+enum class ESidebarLocation : uint8;
+
+/**
+ * A Sidebar is a widget that contains STabDrawers which can be opened and closed from the drawer to allow temporary access to the tab
+ * A drawer is automatically dismissed when it or any of its children loses focus
+ */
+class STabSidebar : public SCompoundWidget
+{
+	SLATE_BEGIN_ARGS(STabSidebar)
+	{}
+		SLATE_ARGUMENT(ESidebarLocation, Location)
+	SLATE_END_ARGS()
+
+public:
+	~STabSidebar();
+
+	void Construct(const FArguments& InArgs);
+
+	/**
+	 * Sets an offset for the sidebar from the top of the window
+	 */
+	void SetOffset(float Offset);
+
+	/**
+	 * Adds a tab to the sidebar
+	 */
+	void AddTab(TSharedRef<SDockTab> Tab);
+
+	/**
+	 * Removes a tab from the sidebar. Does not restore it to the parent stack
+	 *
+	 * @return true if the tab was found and removal was successful
+	 */
+	bool RemoveTab(TSharedRef<SDockTab> TabToRemove);
+
+	/**
+	 * Restores a tab to the parent tab stack and removes it from this sidebar
+	 *
+	 * @return true if the tab was found and removal was successful
+	 */
+	bool RestoreTab(TSharedRef<SDockTab> TabToRestore);
+
+	/**
+	 * @return true if this sidebar contains the provided tab
+	 */
+	bool ContainsTab(TSharedPtr<SDockTab> Tab) const;
+
+	/**
+	 * Get all layout identifiers for tabs in this sidebar
+	 */
+	TArray<FTabId> GetAllTabIds() const;
+
+	/**
+	 * Attempt to open a drawer in the sidebar for a specified tab
+	 *
+	 * @return true if the passed in tab is contained in this sidebar and was opened
+	 */
+	bool TryOpenSidebarDrawer(TSharedRef<SDockTab> ForTab);
+private:
+	void OnTabDrawerButtonClicked(TSharedRef<SDockTab> ForTab);
+	void OnTabDrawerFocusLost(TSharedRef<STabDrawer> Drawer);
+	void OnTabDrawerClosed(TSharedRef<STabDrawer> Drawer);
+	void OnWindowDPIScaleChanged(TSharedRef<SWindow> WindowThatChanged);
+	/** Called when active tab changes which is used to determine if we should close an opened content browser drawer */
+	void OnActiveTabChanged(TSharedPtr<SDockTab> PreviouslyActive, TSharedPtr<SDockTab> NewlyActivated);
+
+	TSharedRef<SWidget> OnGetTabDrawerContextMenuWidget(TSharedRef<SDockTab> ForTab);
+	void OnRestoreTab(TSharedRef<SDockTab> TabToRestore);
+	void OnCloseTab(TSharedRef<SDockTab> TabToClose);
+	
+	/**
+	 * Removes a single drawer for a specified tab from this sidebar
+	 * Removal is done instantly not waiting for any close animation
+	 */
+	void RemoveDrawer(TSharedRef<SDockTab> ForTab);
+
+	/**
+	 * Removes all drawers instantly
+	 */
+	void RemoveAllDrawers();
+	
+	/**
+	 * Closes all drawers, playing a close animation and waiting to remove until the animation is complete
+	 */
+	void CloseAllDrawers();
+
+	EActiveTimerReturnType OnOpenPendingDrawerTimer(double CurrentTime, float DeltaTime);
+	void OpenDrawerNextFrame(TSharedRef<SDockTab> ForTab);
+	void OpenDrawerInternal(TSharedRef<SDockTab> ForTab);
+
+	/**
+	 * Updates the appearance of open drawers
+	 */
+	void UpdateDrawerAppearance();
+private:
+	TSharedPtr<SVerticalBox> TabBox;
+	TArray<TPair<TSharedRef<SDockTab>, TSharedRef<STabDrawerButton>>> Tabs;
+	TWeakPtr<SWindow> WindowWithOverlayContent;
+	TSharedPtr<FActiveTimerHandle> OpenPendingDrawerTimerHandle;
+	ESidebarLocation Location;
+	/** Generally speaking one drawer is only ever open at once but we animate any previous drawer closing so there could be more than one while an animation is playing */
+	TArray<TSharedRef<STabDrawer>> OpenedDrawers;
+	/** Any pending drawer tab to open */
+	TWeakPtr<SDockTab> PendingTabToOpen;
+}; 
\ No newline at end of file
diff --git a/Engine/Source/Runtime/Slate/Private/Framework/Docking/TabManager.cpp b/Engine/Source/Runtime/Slate/Private/Framework/Docking/TabManager.cpp
--- a/Engine/Source/Runtime/Slate/Private/Framework/Docking/TabManager.cpp
+++ b/Engine/Source/Runtime/Slate/Private/Framework/Docking/TabManager.cpp
@@ -32,6 +32,81 @@ DEFINE_LOG_CATEGORY_STATIC(LogTabManager, Display, All);
 
 #define LOCTEXT_NAMESPACE "TabManager"
 
+static const FString UE4_TABMANAGER_OPENED_TAB_STRING = TEXT("OpenedTab");
+static const FString UE4_TABMANAGER_CLOSED_TAB_STRING = TEXT("ClosedTab");
+static const FString UE4_TABMANAGER_SIDEBAR_TAB_STRING = TEXT("SidebarTab");
+static const FString UE4_TABMANAGER_INVALID_TAB_STRING = TEXT("InvalidTab");
+
+static FString StringFromTabState(ETabState::Type TabState)
+{
+	switch (TabState)
+	{
+	case ETabState::OpenedTab:
+		return UE4_TABMANAGER_OPENED_TAB_STRING;
+	case ETabState::ClosedTab:
+		return UE4_TABMANAGER_CLOSED_TAB_STRING;
+	case ETabState::SidebarTab:
+		return UE4_TABMANAGER_SIDEBAR_TAB_STRING;
+	default:
+		return UE4_TABMANAGER_INVALID_TAB_STRING;
+	}
+}
+
+static FString StringFromSidebarLocation(ESidebarLocation Location)
+{
+	switch (Location)
+	{
+	case ESidebarLocation::Left:
+		return TEXT("Left");
+	case ESidebarLocation::Right:
+		return TEXT("Right");
+	default:
+		return TEXT("None");
+	}
+}
+
+static ESidebarLocation SidebarLocationFromString(const FString& AsString)
+{
+	if (AsString == TEXT("Left"))
+	{
+		return ESidebarLocation::Left;
+	}
+	else if (AsString == TEXT("Right"))
+	{
+		return ESidebarLocation::Right;
+	}
+	else
+	{
+		return ESidebarLocation::None;
+	}
+}
+
+static ETabState::Type TabStateFromString(const FString& AsString)
+{
+	if (AsString == UE4_TABMANAGER_OPENED_TAB_STRING)
+	{
+		return ETabState::OpenedTab;
+	}
+	else if (AsString == UE4_TABMANAGER_CLOSED_TAB_STRING)
+	{
+		return ETabState::ClosedTab;
+	}
+	else if (AsString == UE4_TABMANAGER_INVALID_TAB_STRING)
+	{
+		return ETabState::InvalidTab;
+	}
+	else if (AsString == UE4_TABMANAGER_SIDEBAR_TAB_STRING)
+	{
+		return ETabState::SidebarTab;
+	}
+	else
+	{
+		ensureMsgf(false, TEXT("Invalid tab state."));
+		return ETabState::OpenedTab;
+	}
+}
+
+
 //////////////////////////////////////////////////////////////////////////
 // 
 //////////////////////////////////////////////////////////////////////////
@@ -78,9 +153,7 @@ TSharedPtr<SDockTab> FTabManager::FLastMajorOrNomadTab::Search(const FTabManager
 	return FoundTab;
 }
 
-static const FString UE4_TABMANAGER_OPENED_TAB_STRING = TEXT("OpenedTab");
-static const FString UE4_TABMANAGER_CLOSED_TAB_STRING = TEXT("ClosedTab");
-static const FString UE4_TABMANAGER_INVALID_TAB_STRING = TEXT("InvalidTab");
+
 
 TSharedRef<FTabManager::FLayoutNode> FTabManager::FLayout::NewFromString_Helper( TSharedPtr<FJsonObject> JsonObject )
 {
@@ -132,26 +205,6 @@ TSharedRef<FTabManager::FLayoutNode> FTabManager::FLayout::NewFromString_Helper(
 			
 		}
 
-		static ETabState::Type TabStateFromString( const FString& AsString )
-		{
-			if (AsString == UE4_TABMANAGER_OPENED_TAB_STRING)
-			{
-				return ETabState::OpenedTab;
-			}
-			else if(AsString == UE4_TABMANAGER_CLOSED_TAB_STRING)
-			{
-				return ETabState::ClosedTab;
-			}
-			else if (AsString == UE4_TABMANAGER_INVALID_TAB_STRING)
-			{
-				return ETabState::InvalidTab;
-			}
-			else
-			{
-				ensureMsgf(false, TEXT("Invalid tab state."));
-				return ETabState::OpenedTab;
-			}
-		}
 	};
 
 	const FString NodeType = JsonObject->GetStringField(TEXT("Type"));
@@ -241,7 +294,15 @@ TSharedRef<FTabManager::FLayoutNode> FTabManager::FLayout::NewFromString_Helper(
 			TSharedPtr<FJsonObject> TabAsJson = TabsAsJson[TabIndex]->AsObject();
 			FName TabId = FName( *TabAsJson->GetStringField(TEXT("TabId")) );
 			TabId = FGlobalTabmanager::Get()->GetTabTypeForPotentiallyLegacyTab(TabId);
-			NewStack->AddTab( TabId, local::TabStateFromString( TabAsJson->GetStringField(TEXT("TabState")) ) );			
+
+			FString SidebarLocation;
+			double SidebarSizeCoefficient = .15f;
+			if (TabAsJson->TryGetStringField(TEXT("SidebarLocation"), SidebarLocation))
+			{
+				TabAsJson->TryGetNumberField(TEXT("SidebarCoeff"), SidebarSizeCoefficient);
+			}
+
+			NewStack->AddTab(TabId, TabStateFromString( TabAsJson->GetStringField(TEXT("TabState"))), SidebarLocationFromString(SidebarLocation), SidebarSizeCoefficient);
 		}
 		return NewStack;
 	}
@@ -377,13 +438,18 @@ TSharedRef<FJsonObject> FTabManager::FLayout::PersistToString_Helper(const TShar
 		JsonObj->SetStringField( TEXT("ForegroundTab"), NodeAsStack->ForegroundTabId.ToString() );
 
 		TArray< TSharedPtr<FJsonValue> > TabsAsJson;
-		for ( int32 TabIndex=0; TabIndex < NodeAsStack->Tabs.Num(); ++TabIndex )
+		for(const FTab& Tab : NodeAsStack->Tabs)
 		{
 			TSharedRef<FJsonObject> TabAsJson = MakeShareable( new FJsonObject() );
-			TabAsJson->SetStringField( TEXT("TabId"), NodeAsStack->Tabs[TabIndex].TabId.ToString() );
-			TabAsJson->SetStringField( TEXT("TabState"), (NodeAsStack->Tabs[TabIndex].TabState == ETabState::OpenedTab)
-				? UE4_TABMANAGER_OPENED_TAB_STRING : NodeAsStack->Tabs[TabIndex].TabState == ETabState::ClosedTab
-				? UE4_TABMANAGER_CLOSED_TAB_STRING : UE4_TABMANAGER_INVALID_TAB_STRING);
+			TabAsJson->SetStringField( TEXT("TabId"), Tab.TabId.ToString() );
+			TabAsJson->SetStringField(TEXT("TabState"), StringFromTabState(Tab.TabState));
+
+			if (Tab.TabState == ETabState::SidebarTab && Tab.SidebarLocation != ESidebarLocation::None)
+			{
+				TabAsJson->SetStringField(TEXT("SidebarLocation"), StringFromSidebarLocation(Tab.SidebarLocation));
+				TabAsJson->SetNumberField(TEXT("SidebarCoeff"), Tab.SidebarSizeCoefficient);
+			}
+
 			TabsAsJson.Add( MakeShareable( new FJsonValueObject(TabAsJson) ) );
 		}
 		JsonObj->SetArrayField( TEXT("Tabs"), TabsAsJson );
@@ -1103,30 +1169,34 @@ void FTabManager::DrawAttention( const TSharedRef<SDockTab>& TabToHighlight )
 {
 	// Bring the tab to front.
 	const TSharedPtr<SDockingArea> DockingArea = TabToHighlight->GetDockArea();
-	if ( DockingArea.IsValid() )
+	if (DockingArea.IsValid())
 	{
 		const TSharedRef<FTabManager> ManagerOfTabToHighlight = DockingArea->GetTabManager();
 
-		if ( ManagerOfTabToHighlight != FGlobalTabmanager::Get() )
+		if (ManagerOfTabToHighlight != FGlobalTabmanager::Get())
 		{
 			FGlobalTabmanager::Get()->DrawAttentionToTabManager(ManagerOfTabToHighlight);
 		}
 
 		TSharedPtr<SWindow> OwnerWindow = DockingArea->GetParentWindow();
 
-		if ( SWindow* OwnerWindowPtr = OwnerWindow.Get() )
+		if (SWindow* OwnerWindowPtr = OwnerWindow.Get())
 		{
 			// When should we force a window to the front?
 			// 1) The owner window is already active, so we know the user is using this screen.
 			// 2) This window is a child window of another already active window (same as 1).
 			// 3) Slate is currently processing input, which would imply we got this request at the behest of a user's click or press.
-			if ( OwnerWindowPtr->IsActive() || OwnerWindowPtr->HasActiveParent() || FSlateApplication::Get().IsProcessingInput() )
+			if (OwnerWindowPtr->IsActive() || OwnerWindowPtr->HasActiveParent() || FSlateApplication::Get().IsProcessingInput())
 			{
 				OwnerWindowPtr->BringToFront();
 			}
 		}
 
-		TabToHighlight->GetParentDockTabStack()->BringToFront(TabToHighlight);
+		if (!DockingArea->TryOpenSidebarDrawer(TabToHighlight))
+		{
+			TabToHighlight->GetParentDockTabStack()->BringToFront(TabToHighlight);
+		}
+
 		TabToHighlight->FlashTab();
 
 		FGlobalTabmanager::Get()->UpdateMainMenu(TabToHighlight, true);
@@ -1240,12 +1310,12 @@ TSharedPtr<SDockTab> FTabManager::InvokeTab_Internal( const FTabId& TabId )
 			? Spawner->OnFindTabToReuse.Execute( TabId )
 			: Spawner->SpawnedTabPtr.Pin();
 
-		if ( ExistingTab.IsValid() )
+		if (ExistingTab.IsValid())
 		{
-			if ( !ExistingTab->IsActive() )
+			if (!ExistingTab->IsActive())
 			{
 				// Draw attention to this tab if it didn't already have focus
-				DrawAttention( ExistingTab.ToSharedRef() );
+				DrawAttention(ExistingTab.ToSharedRef());
 			}
 			return ExistingTab.ToSharedRef();
 		}
@@ -1388,13 +1458,18 @@ FTabManager::FTabManager( const TSharedPtr<SDockTab>& InOwnerTab, const TSharedR
 	LocalWorkspaceMenuRoot = FWorkspaceItem::NewGroup(LOCTEXT("LocalWorkspaceRoot", "Local Workspace Root"));
 }
 
-TSharedPtr<SDockingArea> FTabManager::RestoreArea(
-	const TSharedRef<FArea>& AreaToRestore, const TSharedPtr<SWindow>& InParentWindow, const bool bEmbedTitleAreaContent, const EOutputCanBeNullptr OutputCanBeNullptr)
+TSharedPtr<SDockingArea> FTabManager::RestoreArea(const TSharedRef<FArea>& AreaToRestore, const TSharedPtr<SWindow>& InParentWindow, const bool bEmbedTitleAreaContent, const EOutputCanBeNullptr OutputCanBeNullptr)
 {
-	if (TSharedPtr<SDockingNode> RestoredNode = RestoreArea_Helper(AreaToRestore, InParentWindow, bEmbedTitleAreaContent, OutputCanBeNullptr))
+	// Sidebar tabs for this area
+	FSidebarTabLists SidebarTabs;
+
+	if (TSharedPtr<SDockingNode> RestoredNode = RestoreArea_Helper(AreaToRestore, InParentWindow, bEmbedTitleAreaContent, SidebarTabs, OutputCanBeNullptr))
 	{
 		TSharedRef<SDockingArea> RestoredArea = StaticCastSharedRef<SDockingArea>(RestoredNode->AsShared());
+
 		RestoredArea->CleanUp(SDockingNode::TabRemoval_None);
+
+		RestoredArea->AddSidebarTabsFromRestoredLayout(SidebarTabs);
 		return RestoredArea;
 	}
 	else
@@ -1404,8 +1479,7 @@ TSharedPtr<SDockingArea> FTabManager::RestoreArea(
 	}
 }
 
-TSharedPtr<SDockingNode> FTabManager::RestoreArea_Helper(
-	const TSharedRef<FLayoutNode>& LayoutNode, const TSharedPtr<SWindow>& ParentWindow, const bool bEmbedTitleAreaContent, const EOutputCanBeNullptr OutputCanBeNullptr)
+TSharedPtr<SDockingNode> FTabManager::RestoreArea_Helper(const TSharedRef<FLayoutNode>& LayoutNode, const TSharedPtr<SWindow>& ParentWindow, const bool bEmbedTitleAreaContent, FSidebarTabLists& OutSidebarTabs, const EOutputCanBeNullptr OutputCanBeNullptr)
 {
 	TSharedPtr<FTabManager::FStack> NodeAsStack = LayoutNode->AsStack();
 	TSharedPtr<FTabManager::FSplitter> NodeAsSplitter = LayoutNode->AsSplitter();
@@ -1452,7 +1526,7 @@ TSharedPtr<SDockingNode> FTabManager::RestoreArea_Helper(
 		// Open Tabs
 		for (const FTab& SomeTab : NodeAsStack->Tabs)
 		{
-			if (SomeTab.TabState == ETabState::OpenedTab && IsValidTabForSpawning(SomeTab))
+			if ((SomeTab.TabState == ETabState::OpenedTab || SomeTab.TabState == ETabState::SidebarTab) && IsValidTabForSpawning(SomeTab))
 			{
 				const TSharedPtr<SDockTab> NewTabWidget = SpawnTab(SomeTab.TabId, ParentWindow, bCanOutputBeNullptr);
 
@@ -1460,6 +1534,7 @@ TSharedPtr<SDockingNode> FTabManager::RestoreArea_Helper(
 				{
 					if (SomeTab.TabId == NodeAsStack->ForegroundTabId)
 					{
+						ensure(SomeTab.TabState == ETabState::OpenedTab);
 						WidgetToActivate = NewTabWidget;
 					}
 
@@ -1469,7 +1544,25 @@ TSharedPtr<SDockingNode> FTabManager::RestoreArea_Helper(
 						NewStackWidget = SNew(SDockingTabStack, NodeAsStack.ToSharedRef());
 						NewStackWidget->SetSizeCoefficient(LayoutNode->GetSizeCoefficient());
 					}
-					NewStackWidget->AddTabWidget(NewTabWidget.ToSharedRef());
+
+					if (SomeTab.TabState == ETabState::OpenedTab)
+					{
+						NewStackWidget->AddTabWidget(NewTabWidget.ToSharedRef());
+					}
+					else
+					{
+						// Let the stack know we have a tab that belongs in its stack that is currently in a sidebar
+						NewStackWidget->AddSidebarTab(NewTabWidget.ToSharedRef());
+						if (SomeTab.SidebarLocation == ESidebarLocation::Left)
+						{
+							OutSidebarTabs.LeftSidebarTabs.Add(NewTabWidget.ToSharedRef());
+						}
+						else
+						{
+							ensure(SomeTab.SidebarLocation == ESidebarLocation::Right);
+							OutSidebarTabs.RightSidebarTabs.Add(NewTabWidget.ToSharedRef());
+						}
+					}
 				}
 			}
 		}
@@ -1520,7 +1613,7 @@ TSharedPtr<SDockingNode> FTabManager::RestoreArea_Helper(
 			NewWindow->SetTitle(FGlobalTabmanager::Get()->GetApplicationTitle());
 
 			TArray<TSharedRef<SDockingNode>> DockingNodes;
-			if (CanRestoreSplitterContent(DockingNodes, NodeAsArea.ToSharedRef(), NewWindow, OutputCanBeNullptr))
+			if (CanRestoreSplitterContent(DockingNodes, NodeAsArea.ToSharedRef(), NewWindow, OutSidebarTabs, OutputCanBeNullptr))
 			{
 				// Create SplitterWidget only if it will be filled with at least 1 DockingNodes
 				// Any windows that were "pulled out" of a dock area should be children of the window in which the parent dock area resides.
@@ -1539,7 +1632,7 @@ TSharedPtr<SDockingNode> FTabManager::RestoreArea_Helper(
 				// Restore content
 				if (!bCanOutputBeNullptr)
 				{
-					RestoreSplitterContent(NodeAsArea.ToSharedRef(), NewDockAreaWidget.ToSharedRef(), NewWindow);
+					RestoreSplitterContent(NodeAsArea.ToSharedRef(), NewDockAreaWidget.ToSharedRef(), NewWindow, OutSidebarTabs);
 				}
 				else
 				{
@@ -1550,7 +1643,7 @@ TSharedPtr<SDockingNode> FTabManager::RestoreArea_Helper(
 		else
 		{
 			TArray<TSharedRef<SDockingNode>> DockingNodes;
-			if (CanRestoreSplitterContent(DockingNodes, NodeAsArea.ToSharedRef(), ParentWindow, OutputCanBeNullptr))
+			if (CanRestoreSplitterContent(DockingNodes, NodeAsArea.ToSharedRef(), ParentWindow, OutSidebarTabs, OutputCanBeNullptr))
 			{
 				SAssignNew(NewDockAreaWidget, SDockingArea, SharedThis(this), NodeAsArea.ToSharedRef())
 					// We only want to set a parent window on this dock area, if we need to have title area content
@@ -1564,7 +1657,7 @@ TSharedPtr<SDockingNode> FTabManager::RestoreArea_Helper(
 				// Restore content
 				if (!bCanOutputBeNullptr)
 				{
-					RestoreSplitterContent(NodeAsArea.ToSharedRef(), NewDockAreaWidget.ToSharedRef(), ParentWindow);
+					RestoreSplitterContent(NodeAsArea.ToSharedRef(), NewDockAreaWidget.ToSharedRef(), ParentWindow, OutSidebarTabs);
 				}
 				else
 				{
@@ -1578,14 +1671,14 @@ TSharedPtr<SDockingNode> FTabManager::RestoreArea_Helper(
 	else if ( NodeAsSplitter.IsValid() ) 
 	{
 		TArray<TSharedRef<SDockingNode>> DockingNodes;
-		if (CanRestoreSplitterContent(DockingNodes, NodeAsSplitter.ToSharedRef(), ParentWindow, OutputCanBeNullptr))
+		if (CanRestoreSplitterContent(DockingNodes, NodeAsSplitter.ToSharedRef(), ParentWindow, OutSidebarTabs, OutputCanBeNullptr))
 		{
 			TSharedRef<SDockingSplitter> NewSplitterWidget = SNew( SDockingSplitter, NodeAsSplitter.ToSharedRef() );
 			NewSplitterWidget->SetSizeCoefficient(LayoutNode->GetSizeCoefficient());
 			// Restore content
 			if (!bCanOutputBeNullptr)
 			{
-				RestoreSplitterContent(NodeAsSplitter.ToSharedRef(), NewSplitterWidget, ParentWindow);
+				RestoreSplitterContent(NodeAsSplitter.ToSharedRef(), NewSplitterWidget, ParentWindow, OutSidebarTabs);
 			}
 			else
 			{
@@ -1607,7 +1700,7 @@ TSharedPtr<SDockingNode> FTabManager::RestoreArea_Helper(
 	}
 }
 
-bool FTabManager::CanRestoreSplitterContent(TArray<TSharedRef<SDockingNode>>& DockingNodes, const TSharedRef<FSplitter>& SplitterNode, const TSharedPtr<SWindow>& ParentWindow, const EOutputCanBeNullptr OutputCanBeNullptr)
+bool FTabManager::CanRestoreSplitterContent(TArray<TSharedRef<SDockingNode>>& DockingNodes, const TSharedRef<FSplitter>& SplitterNode, const TSharedPtr<SWindow>& ParentWindow, FSidebarTabLists& OutSidebarTabs, const EOutputCanBeNullptr OutputCanBeNullptr)
 {
 	if (OutputCanBeNullptr == EOutputCanBeNullptr::Never)
 	{
@@ -1620,7 +1713,7 @@ bool FTabManager::CanRestoreSplitterContent(TArray<TSharedRef<SDockingNode>>& Do
 		const TSharedRef<FLayoutNode> ThisChildNode = SplitterNode->ChildNodes[ChildNodeIndex];
 
 		const bool bEmbedTitleAreaContent = false;
-		const TSharedPtr<SDockingNode> ThisChildNodeWidget = RestoreArea_Helper(ThisChildNode, ParentWindow, bEmbedTitleAreaContent, OutputCanBeNullptr);
+		const TSharedPtr<SDockingNode> ThisChildNodeWidget = RestoreArea_Helper(ThisChildNode, ParentWindow, bEmbedTitleAreaContent, OutSidebarTabs, OutputCanBeNullptr);
 		if (ThisChildNodeWidget)
 		{
 			const TSharedRef<SDockingNode> ThisChildNodeWidgetRef = StaticCastSharedRef<SDockingNode>(ThisChildNodeWidget->AsShared());
@@ -1630,7 +1723,7 @@ bool FTabManager::CanRestoreSplitterContent(TArray<TSharedRef<SDockingNode>>& Do
 	return (DockingNodes.Num() > 0);
 }
 
-void FTabManager::RestoreSplitterContent( const TArray<TSharedRef<SDockingNode>>& DockingNodes, const TSharedRef<SDockingSplitter>& SplitterWidget )
+void FTabManager::RestoreSplitterContent( const TArray<TSharedRef<SDockingNode>>& DockingNodes, const TSharedRef<SDockingSplitter>& SplitterWidget)
 {
 	for (const TSharedRef<SDockingNode>& DockingNode : DockingNodes)
 	{
@@ -1638,7 +1731,7 @@ void FTabManager::RestoreSplitterContent( const TArray<TSharedRef<SDockingNode>>
 	}
 }
 
-void FTabManager::RestoreSplitterContent(const TSharedRef<FSplitter>& SplitterNode, const TSharedRef<SDockingSplitter>& SplitterWidget, const TSharedPtr<SWindow>& ParentWindow)
+void FTabManager::RestoreSplitterContent(const TSharedRef<FSplitter>& SplitterNode, const TSharedRef<SDockingSplitter>& SplitterWidget, const TSharedPtr<SWindow>& ParentWindow, FSidebarTabLists& OutSidebarTabs)
 {
 	// Restore the contents of this splitter.
 	for ( int32 ChildNodeIndex = 0; ChildNodeIndex < SplitterNode->ChildNodes.Num(); ++ChildNodeIndex )
@@ -1646,7 +1739,7 @@ void FTabManager::RestoreSplitterContent(const TSharedRef<FSplitter>& SplitterNo
 		TSharedRef<FLayoutNode> ThisChildNode = SplitterNode->ChildNodes[ChildNodeIndex];
 
 		const bool bEmbedTitleAreaContent = false;
-		TSharedPtr<SDockingNode> ThisChildNodeWidget = RestoreArea_Helper(ThisChildNode, ParentWindow, bEmbedTitleAreaContent);
+		TSharedPtr<SDockingNode> ThisChildNodeWidget = RestoreArea_Helper(ThisChildNode, ParentWindow, bEmbedTitleAreaContent, OutSidebarTabs);
 		check(ThisChildNodeWidget.IsValid());
 		if (ThisChildNodeWidget)
 		{
@@ -1695,6 +1788,54 @@ bool FTabManager::IsAllowedTabType(const FName TabType) const
 	return TabType == NAME_None || TabBlacklist->PassesFilter(TabType);
 }
 
+bool FTabManager::IsTabAllowedInSidebar(const FTabId TabId) const
+{
+	if (const TSharedPtr<const FTabSpawnerEntry> Spawner = FindTabSpawnerFor(TabId.TabType))
+	{
+		return Spawner->CanSidebarTab();
+	}
+
+	return false;
+}
+
+void FTabManager::ToggleSidebarOpenTabs()
+{
+	if(TemporarilySidebaredTabs.Num() == 0)
+	{
+		// Sidebar opened tabs not in a sidebar already
+		for (int32 AreaIndex = 0; AreaIndex < DockAreas.Num(); ++AreaIndex)
+		{
+			TSharedPtr<SDockingArea> SomeDockArea = DockAreas[AreaIndex].Pin();
+
+			if (SomeDockArea.IsValid() && SomeDockArea->CanHaveSidebar())
+			{
+				TArray<TSharedRef<SDockTab>> AllTabs = SomeDockArea->GetAllChildTabs();
+				for (TSharedRef<SDockTab>& Tab : AllTabs)
+				{
+					if (IsTabAllowedInSidebar(Tab->GetLayoutIdentifier()) && !SomeDockArea->IsTabInSidebar(Tab) && Tab->GetParentDockTabStack()->CanMoveTabToSideBar(Tab))
+					{
+						Tab->GetParentDockTabStack()->MoveTabToSidebar(Tab);
+						TemporarilySidebaredTabs.Add(Tab);
+					}
+				}
+
+			}
+		}
+	}
+	else
+	{
+		for (TWeakPtr<SDockTab>& TabPtr : TemporarilySidebaredTabs)
+		{
+			if (TSharedPtr<SDockTab> Tab = TabPtr.Pin())
+			{
+				Tab->GetParentDockTabStack()->GetDockArea()->RestoreTabFromSidebar(Tab.ToSharedRef());
+			}
+		}
+
+		TemporarilySidebaredTabs.Empty();
+	}
+}
+
 TSharedPtr<SDockTab> FTabManager::SpawnTab(const FTabId& TabId, const TSharedPtr<SWindow>& ParentWindow, const bool bCanOutputBeNullptr)
 {
 	TSharedPtr<SDockTab> NewTabWidget;
@@ -2082,6 +2223,20 @@ TSharedPtr<FTabSpawnerEntry> FTabManager::FindTabSpawnerFor(FName TabId)
 		: TSharedPtr<FTabSpawnerEntry>();
 }
 
+const TSharedPtr<const FTabSpawnerEntry> FTabManager::FindTabSpawnerFor(FName TabId) const
+{
+	// Look for a spawner in this tab manager.
+	const TSharedRef<FTabSpawnerEntry>* Spawner = TabSpawner.Find(TabId);
+	if (Spawner == nullptr)
+	{
+		Spawner = NomadTabSpawner->Find(TabId);
+	}
+
+	return (Spawner != nullptr)
+		? TSharedPtr<const FTabSpawnerEntry>(*Spawner)
+		: TSharedPtr<FTabSpawnerEntry>();
+}
+
 int32 FTabManager::FindTabInCollapsedAreas( const FTabMatcher& Matcher )
 {
 	for ( int32 CollapsedDockAreaIndex=0; CollapsedDockAreaIndex < CollapsedDockAreas.Num(); ++CollapsedDockAreaIndex )
diff --git a/Engine/Source/Runtime/Slate/Private/Widgets/Docking/SDockTab.cpp b/Engine/Source/Runtime/Slate/Private/Widgets/Docking/SDockTab.cpp
--- a/Engine/Source/Runtime/Slate/Private/Widgets/Docking/SDockTab.cpp
+++ b/Engine/Source/Runtime/Slate/Private/Widgets/Docking/SDockTab.cpp
@@ -328,7 +328,7 @@ void SDockTab::SetLayoutIdentifier( const FTabId& TabId )
 	LayoutIdentifier = TabId;
 }
 
-const FTabId& SDockTab::GetLayoutIdentifier() const
+const FTabId SDockTab::GetLayoutIdentifier() const
 {
 	return LayoutIdentifier;
 }
@@ -494,7 +494,7 @@ void SDockTab::SetDraggedOverDockArea( const TSharedPtr<SDockingArea>& Area )
 bool SDockTab::HasSiblingTab(const FTabId& SiblingTabId, const bool TreatIndexNoneAsWildcard) const
 {
 	TSharedPtr<SDockingTabStack> ParentTabStack = GetParentDockTabStack();
-	return (ParentTabStack.IsValid()) ? ParentTabStack->HasTab(FTabMatcher(SiblingTabId, static_cast<ETabState::Type>(ETabState::ClosedTab | ETabState::OpenedTab), TreatIndexNoneAsWildcard)) : false;
+	return (ParentTabStack.IsValid()) ? ParentTabStack->HasTab(FTabMatcher(SiblingTabId, static_cast<ETabState::Type>(ETabState::ClosedTab | ETabState::OpenedTab | ETabState::SidebarTab), TreatIndexNoneAsWildcard)) : false;
 }
 
 void SDockTab::Construct( const FArguments& InArgs )
diff --git a/Engine/Source/Runtime/Slate/Public/Framework/Docking/STabDrawer.h b/Engine/Source/Runtime/Slate/Public/Framework/Docking/STabDrawer.h
new file mode 100644
--- /dev/null
+++ b/Engine/Source/Runtime/Slate/Public/Framework/Docking/STabDrawer.h
@@ -0,0 +1,119 @@
+// Copyright Epic Games, Inc. All Rights Reserved.
+
+#pragma once
+
+#include "Widgets/SCompoundWidget.h"
+#include "Widgets/DeclarativeSyntaxSupport.h"
+#include "Application/ThrottleManager.h"
+#include "Animation/CurveSequence.h"
+
+struct FSplitterStyle;
+
+class STabDrawer;
+class SDockTab;
+
+/** The direction that a tab drawer opens relative to the location of the sidebar it is in */
+enum class ETabDrawerOpenDirection : uint8
+{
+	/** Open from the left to the right */
+	Left,
+
+	/** Open from the right to the left */
+	Right,
+
+	/** Open from the top to the bottom */
+	Bottom
+};
+
+DECLARE_DELEGATE_OneParam(FOnDrawerTargetSizeChanged, float);
+DECLARE_DELEGATE_OneParam(FOnDrawerFocusLost, TSharedRef<STabDrawer>);
+DECLARE_DELEGATE_OneParam(FOnDrawerClosed, TSharedRef<STabDrawer>);
+
+/**
+ * A tab drawer is a widget that contains the contents of a widget when that widget is in a sidebar
+ */
+class STabDrawer : public SCompoundWidget
+{
+	SLATE_BEGIN_ARGS(STabDrawer)
+	{
+		_Clipping = EWidgetClipping::ClipToBounds;
+		_ShadowOffset = FVector2D(8.0f, 8.0f);
+	}
+	SLATE_DEFAULT_SLOT(FArguments, Content)
+		/** The minimum size a drawer can be when opened. This unit is in window space */
+		SLATE_ARGUMENT(float, MinDrawerSize)
+		/** The maximum size a drawer can be when opened. This unit is in window space */
+		SLATE_ARGUMENT(float, MaxDrawerSize)
+		/** The size that the drawer should be, clamped to the above min/max values. This unit is in window space */
+		SLATE_ARGUMENT(float, TargetDrawerSize)
+		/** Called when the target drawer size is changed by the user */
+		SLATE_EVENT(FOnDrawerTargetSizeChanged, OnTargetDrawerSizeChanged)
+		/** Called when the drawer loses focus */
+		SLATE_EVENT(FOnDrawerFocusLost, OnDrawerFocusLost)
+		/** Called when the drawer is completely closed (i.e will be called once the close animation completes */
+		SLATE_EVENT(FOnDrawerClosed, OnDrawerClosed)
+		/** The side of the drop shadow surrounding the drawer */
+		SLATE_ARGUMENT(FVector2D, ShadowOffset)
+	SLATE_END_ARGS()
+
+	~STabDrawer();
+	void Construct(const FArguments& InArgs, TSharedRef<SDockTab> InTab, ETabDrawerOpenDirection InOpenDirection);
+
+	/** Sets the current size of the drawer, ignoring any open/close animation */
+	void SetCurrentSize(float InSize);
+
+	/** Begins an animation which opens the drawer */
+	void Open();
+
+	/** Begins an animation which closes the drawer */
+	void Close();
+
+	/** @return true if the drawer is open */
+	bool IsOpen() const;
+
+	/** @return the tab whose contents is being shown in the drawer */
+	const TSharedRef<SDockTab> GetTab() const;
+
+	/** SWidget interface */
+	virtual bool SupportsKeyboardFocus() const override;
+	virtual FVector2D ComputeDesiredSize(float LayoutScaleMultiplier) const override;
+	virtual void OnArrangeChildren(const FGeometry& AllottedGeometry, FArrangedChildren& ArrangedChildren) const override;
+	virtual FReply OnMouseButtonDown(const FGeometry& AllottedGeometry, const FPointerEvent& MouseEvent) override;
+	virtual FReply OnMouseButtonUp(const FGeometry& AllottedGeometry, const FPointerEvent& MouseEvent) override;
+	virtual FReply OnMouseMove(const FGeometry& AllottedGeometry, const FPointerEvent& MouseEvent) override;
+	virtual void OnMouseLeave(const FPointerEvent& MouseEvent) override;
+	virtual FCursorReply OnCursorQuery(const FGeometry& MyGeometry, const FPointerEvent& CursorEvent) const override;
+	virtual int32 OnPaint(const FPaintArgs& Args, const FGeometry& AllottedGeometry, const FSlateRect& MyCullingRect, FSlateWindowElementList& OutDrawElements, int32 LayerId, const FWidgetStyle& InWidgetStyle, bool bParentEnabled) const override;
+
+private:
+	FGeometry GetRenderTransformedGeometry(const FGeometry& AllottedGeometry) const;
+	FGeometry GetResizeHandleGeometry(const FGeometry& AllottedGeometry) const;
+	EActiveTimerReturnType UpdateAnimation(double CurrentTime, float DeltaTime);
+	void OnGlobalFocusChanging(const FFocusEvent& FocusEvent, const FWeakWidgetPath& OldFocusedWidgetPath, const TSharedPtr<SWidget>& OldFocusedWidget, const FWidgetPath& NewFocusedWidgetPath, const TSharedPtr<SWidget>& NewFocusedWidget);
+
+private:
+	FGeometry InitialResizeGeometry;
+	FOnDrawerTargetSizeChanged OnTargetDrawerSizeChanged;
+	FOnDrawerFocusLost OnDrawerFocusLost;
+	FOnDrawerClosed OnDrawerClosed;
+	TSharedPtr<FActiveTimerHandle> OpenCloseTimer;
+	TSharedPtr<SDockTab> ForTab;
+	FThrottleRequest ResizeThrottleHandle;
+	FThrottleRequest AnimationThrottle;
+	FCurveSequence OpenCloseAnimation;
+	const FSlateBrush* BackgroundBrush;
+	const FSlateBrush* ShadowBrush;
+	const FSlateBrush* BorderBrush;
+	const FSplitterStyle* SplitterStyle;
+	FVector2D ShadowOffset;
+	float ExpanderSize;
+	float CurrentSize;
+	float MinDrawerSize;
+	float MaxDrawerSize;
+	float TargetDrawerSize;
+	float InitialSizeAtResize;
+	ETabDrawerOpenDirection OpenDirection;
+	bool bIsResizing;
+	bool bIsResizeHandleHovered;
+};
+
diff --git a/Engine/Source/Runtime/Slate/Public/Framework/Docking/TabManager.h b/Engine/Source/Runtime/Slate/Public/Framework/Docking/TabManager.h
--- a/Engine/Source/Runtime/Slate/Public/Framework/Docking/TabManager.h
+++ b/Engine/Source/Runtime/Slate/Public/Framework/Docking/TabManager.h
@@ -18,6 +18,7 @@ class SDockingArea;
 class SDockingTabStack;
 class FLayoutExtender;
 struct FTabMatcher;
+struct FSidebarTabLists;
 
 DECLARE_MULTICAST_DELEGATE_TwoParams(
 	FOnActiveTabChanged,
@@ -27,18 +28,29 @@ DECLARE_MULTICAST_DELEGATE_TwoParams(
 	TSharedPtr<SDockTab> );
 	
 
+enum class ESidebarLocation : uint8
+{
+	/** Tab is in a sidebar on the left side of its parent area */
+	Left,
+	/** Tab is in a sidebar on the right side of its parent area */
+	Right,
+
+	/** Tab is not in a sidebar */
+	None,
+};
+
 struct FTabId
 {
-	FTabId( )
+	FTabId()
 		: InstanceId(INDEX_NONE)
 	{ }
 
-	FTabId( const FName& InTabType, const int32 InInstanceId )
+	FTabId(const FName InTabType, const int32 InInstanceId)
 		: TabType(InTabType)
 		, InstanceId(InInstanceId)
 	{ }
 
-	FTabId( const FName& InTabType )
+	FTabId(const FName InTabType)
 		: TabType(InTabType)
 		, InstanceId(INDEX_NONE)
 	{ }
@@ -49,7 +61,7 @@ struct FTabId
 		return InstanceId == INDEX_NONE;
 	}
 
-	bool operator==( const FTabId& Other ) const
+	bool operator==(const FTabId& Other) const
 	{
 		return TabType == Other.TabType && (InstanceId == INDEX_NONE || Other.InstanceId == INDEX_NONE || InstanceId == Other.InstanceId) ;
 	}
@@ -146,6 +158,7 @@ struct FTabSpawnerEntry : public FWorkspaceItem
 		, OnFindTabToReuse()
 		, MenuType(ETabSpawnerMenuType::Enabled)
 		, bAutoGenerateMenuEntry(true)
+		, bCanSidebarTab(true)
 		, SpawnedTabPtr()
 	{
 	}
@@ -192,6 +205,17 @@ struct FTabSpawnerEntry : public FWorkspaceItem
 		return *this;
 	}
 
+	FTabSpawnerEntry& SetCanSidebarTab(bool bInCanSidebarTab)
+	{
+		bCanSidebarTab = bInCanSidebarTab;
+		return *this;
+	}
+
+	bool CanSidebarTab() const
+	{
+		return bCanSidebarTab;
+	}
+
 	virtual TSharedPtr<FTabSpawnerEntry> AsSpawnerEntry() override
 	{
 		return SharedThis(this);
@@ -207,6 +231,8 @@ private:
 	TAttribute<ETabSpawnerMenuType::Type> MenuType;
 	/** Whether to automatically generate a menu entry for this tab spawner */
 	bool bAutoGenerateMenuEntry;
+	/** Whether or not this tab can ever be in a sidebar */
+	bool bCanSidebarTab;
 
 	TWeakPtr<SDockTab> SpawnedTabPtr;
 
@@ -226,10 +252,12 @@ namespace ETabState
 	{
 		OpenedTab = 0x1 << 0,
 		ClosedTab = 0x1 << 1,
+		SidebarTab = 0x1 << 2,
+
 		/**
 		 * InvalidTab refers to tabs that were not recognized by the Editor (e.g., LiveLink when its plugin its disabled).
 		 */
-		InvalidTab = 0x1 << 2
+		InvalidTab = 0x1 << 3
 	};
 }
 
@@ -296,19 +324,33 @@ class SLATE_API FTabManager : public TSharedFromThis<FTabManager>
 
 		struct FTab
 		{
-			FTab( const FTabId& InTabId, ETabState::Type InTabState )
-			: TabId(InTabId)
-			, TabState(InTabState)
+			FTab(const FTabId& InTabId, ETabState::Type InTabState)
+				: TabId(InTabId)
+				, TabState(InTabState)
+				, SidebarLocation(ESidebarLocation::None)
+				, SidebarSizeCoefficient(0.0f)
+			{
+				check(InTabState != ETabState::SidebarTab);
+			}
+
+			FTab(const FTabId& InTabId, ETabState::Type InTabState, ESidebarLocation InSidebarLocation, float InSidebarSizeCoefficient)
+				: TabId(InTabId)
+				, TabState(InTabState)
+				, SidebarLocation(InSidebarLocation)
+				, SidebarSizeCoefficient(InSidebarSizeCoefficient)
 			{
+				check(InTabState != ETabState::SidebarTab || InSidebarLocation != ESidebarLocation::None);
 			}
 
 			bool operator==( const FTab& Other ) const
 			{
-				return this->TabId == Other.TabId && this->TabState == Other.TabState;
+				return this->TabId == Other.TabId && this->TabState == Other.TabState && this->SidebarLocation == Other.SidebarLocation;
 			}
 
 			FTabId TabId;
 			ETabState::Type TabState;
+			ESidebarLocation SidebarLocation;
+			float SidebarSizeCoefficient;
 		};
 
 		class SLATE_API FStack : public FLayoutNode
@@ -318,16 +360,35 @@ class SLATE_API FTabManager : public TSharedFromThis<FTabManager>
 				friend class SDockingTabStack;
 
 			public:				
+				TSharedRef<FStack> AddTab(const FName TabType, ETabState::Type InTabState)
+				{
+					check(InTabState != ETabState::SidebarTab);
+					Tabs.Add(FTab( FTabId(TabType), InTabState));
+					return SharedThis(this);
+				}
 
-				TSharedRef<FStack> AddTab( const FName& TabType, ETabState::Type TabState )
+				TSharedRef<FStack> AddTab(const FTabId TabId, ETabState::Type InTabState)
 				{
-					Tabs.Add( FTab( FTabId(TabType), TabState ) );
+					check(InTabState != ETabState::SidebarTab);
+
+					Tabs.Add(FTab(TabId, InTabState));
+
 					return SharedThis(this);
 				}
 
-				TSharedRef<FStack> AddTab( const FTabId& TabId, ETabState::Type TabState )
+				TSharedRef<FStack> AddTab(const FName TabType, ETabState::Type InTabState, ESidebarLocation InSidebarLocation, float SidebarSizeCoefficient)
 				{
-					Tabs.Add( FTab(TabId, TabState) );
+					check(InTabState != ETabState::SidebarTab || InSidebarLocation != ESidebarLocation::None);
+					Tabs.Add(FTab(FTabId(TabType), InTabState, InSidebarLocation, SidebarSizeCoefficient));
+					return SharedThis(this);
+				}
+
+				TSharedRef<FStack> AddTab(const FTabId TabId, ETabState::Type InTabState, ESidebarLocation InSidebarLocation, float SidebarSizeCoefficient)
+				{
+					check(InTabState != ETabState::SidebarTab || InSidebarLocation != ESidebarLocation::None);
+
+					Tabs.Add(FTab(TabId, InTabState, InSidebarLocation, SidebarSizeCoefficient));
+
 					return SharedThis(this);
 				}
 
@@ -797,6 +858,14 @@ class SLATE_API FTabManager : public TSharedFromThis<FTabManager>
 		/** @return if the provided tab can be closed. */
 		bool IsTabCloseable(const TSharedRef<const SDockTab>& InTab) const;
 
+		/** @return true if a tab is ever allowed in a sidebar */
+		bool IsTabAllowedInSidebar(const FTabId TabId) const;
+
+		/**
+		 * Temporarily moves all open tabs in this tab manager to a sidebar or restores them from a temporary state
+		 */
+		void ToggleSidebarOpenTabs();
+
 		/** @return The local workspace menu root */
 		const TSharedRef<FWorkspaceItem> GetLocalWorkspaceMenuRoot() const;
 
@@ -850,8 +919,7 @@ class SLATE_API FTabManager : public TSharedFromThis<FTabManager>
 		TSharedPtr<SDockingArea> RestoreArea(
 			const TSharedRef<FArea>& AreaToRestore, const TSharedPtr<SWindow>& InParentWindow, const bool bEmbedTitleAreaContent = false, const EOutputCanBeNullptr OutputCanBeNullptr = EOutputCanBeNullptr::Never);
 
-		TSharedPtr<class SDockingNode> RestoreArea_Helper(
-			const TSharedRef<FLayoutNode>& LayoutNode, const TSharedPtr<SWindow>& ParentWindow, const bool bEmbedTitleAreaContent, const EOutputCanBeNullptr OutputCanBeNullptr = EOutputCanBeNullptr::Never);
+		TSharedPtr<class SDockingNode> RestoreArea_Helper(const TSharedRef<FLayoutNode>& LayoutNode, const TSharedPtr<SWindow>& ParentWindow, const bool bEmbedTitleAreaContent, FSidebarTabLists& OutSidebarTabs, const EOutputCanBeNullptr OutputCanBeNullptr = EOutputCanBeNullptr::Never);
 
 		/**
 		 * Use CanRestoreSplitterContent + RestoreSplitterContent when the output of its internal RestoreArea_Helper can be a nullptr.
@@ -865,17 +933,18 @@ class SLATE_API FTabManager : public TSharedFromThis<FTabManager>
 		 *			RestoreSplitterContent(DockingNodes, SplitterWidget);
 		 *		}
 		 */
-		bool CanRestoreSplitterContent(TArray<TSharedRef<class SDockingNode>>& DockingNodes, const TSharedRef<FSplitter>& SplitterNode, const TSharedPtr<SWindow>& ParentWindow, const EOutputCanBeNullptr OutputCanBeNullptr);
+		bool CanRestoreSplitterContent(TArray<TSharedRef<class SDockingNode>>& DockingNodes, const TSharedRef<FSplitter>& SplitterNode, const TSharedPtr<SWindow>& ParentWindow, FSidebarTabLists& OutSidebarTabs, const EOutputCanBeNullptr OutputCanBeNullptr);
 		void RestoreSplitterContent(const TArray<TSharedRef<class SDockingNode>>& DockingNodes, const TSharedRef<class SDockingSplitter>& SplitterWidget);
 
 		/**
 		 * Use this standalone RestoreSplitterContent when the output of its internal RestoreArea_Helper cannot be a nullptr.
 		 */
-		void RestoreSplitterContent(const TSharedRef<FSplitter>& SplitterNode, const TSharedRef<class SDockingSplitter>& SplitterWidget, const TSharedPtr<SWindow>& ParentWindow);
+		void RestoreSplitterContent(const TSharedRef<FSplitter>& SplitterNode, const TSharedRef<class SDockingSplitter>& SplitterWidget, const TSharedPtr<SWindow>& ParentWindow, FSidebarTabLists& OutSidebarTabs);
 		
 		bool IsValidTabForSpawning( const FTab& SomeTab ) const;
 		bool IsAllowedTab(const FTabId& TabId) const;
 		bool IsAllowedTabType(const FName TabType) const;
+
 		TSharedPtr<SDockTab> SpawnTab(const FTabId& TabId, const TSharedPtr<SWindow>& ParentWindow, const bool bCanOutputBeNullptr = false);
 
 		TSharedPtr<class SDockingTabStack> FindTabInLiveAreas( const FTabMatcher& TabMatcher ) const;
@@ -884,6 +953,7 @@ class SLATE_API FTabManager : public TSharedFromThis<FTabManager>
 		template<typename MatchFunctorType> static bool HasAnyMatchingTabs( const TSharedRef<FTabManager::FLayoutNode>& SomeNode, const MatchFunctorType& Matcher );
 		bool HasOpenTabs( const TSharedRef<FTabManager::FLayoutNode>& SomeNode ) const;
 		bool HasValidTabs( const TSharedRef<FTabManager::FLayoutNode>& SomeNode ) const;
+
 		/**
 		 * It sets the desired (or all) tabs in the FTabManager::FLayoutNode to the desired value.
 		 * @param SomeNode The area whose tabs will be modified.
@@ -922,6 +992,8 @@ class SLATE_API FTabManager : public TSharedFromThis<FTabManager>
 		FTabSpawner TabSpawner;
 		TSharedRef<FTabSpawner> NomadTabSpawner;
 		TSharedPtr<FTabSpawnerEntry> FindTabSpawnerFor(FName TabId);
+		const TSharedPtr<const FTabSpawnerEntry> FindTabSpawnerFor(FName TabId) const;
+
 		bool HasTabSpawnerFor(FName TabId) const;
 
 		TArray< TWeakPtr<SDockingArea> > DockAreas;
@@ -992,6 +1064,9 @@ class SLATE_API FTabManager : public TSharedFromThis<FTabManager>
 
 		/** Allow systems to dynamically hide tabs */
 		TSharedRef<FBlacklistNames> TabBlacklist;
+
+		/** Tabs which have been temporarily put in the a sidebar */
+		TArray<TWeakPtr<SDockTab>> TemporarilySidebaredTabs;
 };
 
 
diff --git a/Engine/Source/Runtime/Slate/Public/Widgets/Docking/SDockTab.h b/Engine/Source/Runtime/Slate/Public/Widgets/Docking/SDockTab.h
--- a/Engine/Source/Runtime/Slate/Public/Widgets/Docking/SDockTab.h
+++ b/Engine/Source/Runtime/Slate/Public/Widgets/Docking/SDockTab.h
@@ -156,7 +156,7 @@ public:
 	FMargin GetContentPadding() const;
 
 	/** Gets this tab's layout identifier */
-	const FTabId& GetLayoutIdentifier() const;
+	const FTabId GetLayoutIdentifier() const;
 
 	/** Sets the tab's tab well parent, or resets it if nothing is passed in */
 	void SetParent(TSharedPtr<SDockingTabWell> Parent = TSharedPtr<SDockingTabWell>());
@@ -164,6 +164,9 @@ public:
 	/** Gets the tab's tab well parent, or nothing, if it has none */
 	TSharedPtr<SDockingTabWell> GetParent() const;
 
+	/** Gets the dock tab stack this dockable tab resides within, if any */
+	TSharedPtr<SDockingTabStack> GetParentDockTabStack() const;
+
 	/** Gets the dock area that this resides in */
 	TSharedPtr<class SDockingArea> GetDockArea() const;
 
@@ -278,10 +281,8 @@ public:
 	{
 		return LastActivationTime;
 	}
+
 protected:
-	
-	/** Gets the dock tab stack this dockable tab resides within, if any */
-	TSharedPtr<SDockingTabStack> GetParentDockTabStack() const;
 
 	/** @return the style currently applied to the dock tab */
 	const FDockTabStyle& GetCurrentStyle() const;
diff --git a/Engine/Source/Runtime/SlateCore/Private/Styling/StarshipCoreStyle.cpp b/Engine/Source/Runtime/SlateCore/Private/Styling/StarshipCoreStyle.cpp
--- a/Engine/Source/Runtime/SlateCore/Private/Styling/StarshipCoreStyle.cpp
+++ b/Engine/Source/Runtime/SlateCore/Private/Styling/StarshipCoreStyle.cpp
@@ -1353,6 +1353,7 @@ void FStarshipCoreStyle::SetupCheckboxStyles(TSharedRef<FStyle>& Style)
 void FStarshipCoreStyle::SetupDockingStyles(TSharedRef<FStyle>& Style)
 {
 	const FButtonStyle& Button = Style->GetWidgetStyle<FButtonStyle>("Button");
+	const FButtonStyle& NoBorder = Style->GetWidgetStyle<FButtonStyle>("NoBorder");
 	const FTextBlockStyle& NormalText = Style->GetWidgetStyle<FTextBlockStyle>("NormalText");
 	const FSlateColor SelectionColor = Style->GetSlateColor("SelectionColor");
 
@@ -1381,8 +1382,8 @@ void FStarshipCoreStyle::SetupDockingStyles(TSharedRef<FStyle>& Style)
 	FLinearColor DockColor_Hovered(FColor(54, 54, 54));
 	FLinearColor DockColor_Active(FColor(62, 62, 62));
 
-	// Panel Tab
-	Style->Set("Docking.Tab", FDockTabStyle()
+	FDockTabStyle MinorTabStyle =
+		FDockTabStyle()
 		.SetCloseButtonStyle(CloseButton)
 		.SetNormalBrush(FSlateNoResource())
 		.SetHoveredBrush(BOX_BRUSH("/Starship/Docking/DockTab_Hover", 4.0f / 20.0f, FStyleColors::Background))
@@ -1401,8 +1402,10 @@ void FStarshipCoreStyle::SetupDockingStyles(TSharedRef<FStyle>& Style)
 		.SetActiveForegroundColor(FStyleColors::ForegroundHover)
 		.SetForegroundForegroundColor(FStyleColors::Foreground)
 		.SetHoveredForegroundColor(FStyleColors::ForegroundHover)
-		.SetTabTextStyle(NormalText)
-	);
+		.SetTabTextStyle(NormalText);
+
+	// Panel Tab
+	Style->Set("Docking.Tab", MinorTabStyle);
 
 	// App Tab
 	Style->Set("Docking.MajorTab", FDockTabStyle()
@@ -1434,6 +1437,32 @@ void FStarshipCoreStyle::SetupDockingStyles(TSharedRef<FStyle>& Style)
 	Style->Set("Docking.Tab.ActiveTabIndicatorColor", FStyleColors::Primary);
 
 
+	FButtonStyle SidebarTabButtonOpened =
+		FButtonStyle(NoBorder)
+		.SetNormal(MinorTabStyle.ForegroundBrush)
+		.SetHovered(MinorTabStyle.ForegroundBrush)
+		.SetNormalForeground(MinorTabStyle.NormalForegroundColor)
+		.SetPressedForeground(MinorTabStyle.NormalForegroundColor)
+		.SetHoveredForeground(MinorTabStyle.HoveredForegroundColor);
+
+
+	FButtonStyle SidebarTabButtonClosed =
+		FButtonStyle(NoBorder)
+		.SetNormal(MinorTabStyle.NormalBrush)
+		.SetHovered(MinorTabStyle.HoveredBrush)
+		.SetNormalForeground(MinorTabStyle.NormalForegroundColor)
+		.SetPressedForeground(MinorTabStyle.NormalForegroundColor)
+		.SetHoveredForeground(MinorTabStyle.HoveredForegroundColor);
+			
+	Style->Set("Docking.SidebarButton.Closed", SidebarTabButtonClosed);
+
+	Style->Set("Docking.SidebarButton.Opened", SidebarTabButtonOpened);
+
+	Style->Set("Docking.Sidebar.DrawerShadow", new BOX_BRUSH("/Starship/Docking/drawer-shadow", FMargin(8/64.f), FLinearColor(0, 0, 0, 1)));
+	Style->Set("Docking.Sidebar.DrawerBackground", new FSlateColorBrush(FStyleColors::Background));
+	Style->Set("Docking.Sidebar.Background", new FSlateColorBrush(FStyleColors::Input));
+	Style->Set("Docking.Sidebar.Border", new FSlateRoundedBoxBrush(FSlateColor(FLinearColor::Transparent), 5.0f, FStyleColors::Hover, 1.0f) );
+
 	// Dock Cross
 	Style->Set("Docking.Cross.DockLeft", new IMAGE_BRUSH("/Docking/OuterDockingIndicator", FVector2D(6, 6), FLinearColor(1.0f, 0.35f, 0.0f, 0.25f)));
 	Style->Set("Docking.Cross.DockLeft_Hovered", new IMAGE_BRUSH("/Docking/OuterDockingIndicator", FVector2D(6, 6), FLinearColor(1.0f, 0.35f, 0.0f)));
diff --git a/Engine/Source/Runtime/SlateCore/Private/Widgets/SWindow.cpp b/Engine/Source/Runtime/SlateCore/Private/Widgets/SWindow.cpp
--- a/Engine/Source/Runtime/SlateCore/Private/Widgets/SWindow.cpp
+++ b/Engine/Source/Runtime/SlateCore/Private/Widgets/SWindow.cpp
@@ -582,7 +582,7 @@ void SWindow::ConstructWindowInternals()
 		WindowBorder =
 			FSlateApplicationBase::Get().MakeImage(
 				&Style->BorderBrush,
-				FLinearColor::White,
+				Style->BorderColor,
 				WindowContentVisibility
 			);
 
@@ -595,7 +595,6 @@ void SWindow::ConstructWindowInternals()
 		this->ChildSlot
 		[
 			SAssignNew(WindowOverlay, SOverlay)
-			.Visibility(EVisibility::SelfHitTestInvisible)
 			// window background
 			+ SOverlay::Slot()
 			[
@@ -1285,12 +1284,14 @@ SOverlay::FOverlaySlot& SWindow::AddOverlaySlot( const int32 ZOrder )
 	return WindowOverlay->AddSlot(ZOrder);
 }
 
-void SWindow::RemoveOverlaySlot( const TSharedRef<SWidget>& InContent )
+bool SWindow::RemoveOverlaySlot(const TSharedRef<SWidget>& InContent)
 {
 	if(WindowOverlay.IsValid())
 	{
-		WindowOverlay->RemoveSlot( InContent );
+		return WindowOverlay->RemoveSlot(InContent);
 	}
+
+	return false;
 }
 
 TSharedPtr<FPopupLayer> SWindow::OnVisualizePopup(const TSharedRef<SWidget>& PopupContent)
@@ -1718,7 +1719,7 @@ void SWindow::Minimize()
 
 int32 SWindow::GetCornerRadius()
 {
-	return IsRegularWindow() ? SWindowDefs::CornerRadius : 0;
+	return IsRegularWindow() ? Style->WindowCornerRadius : 0;
 }
 
 bool SWindow::SupportsKeyboardFocus() const
diff --git a/Engine/Source/Runtime/SlateCore/Public/Layout/SlateRect.h b/Engine/Source/Runtime/SlateCore/Public/Layout/SlateRect.h
--- a/Engine/Source/Runtime/SlateCore/Public/Layout/SlateRect.h
+++ b/Engine/Source/Runtime/SlateCore/Public/Layout/SlateRect.h
@@ -72,7 +72,7 @@ public:
 	}
 
 	/**
-	 * Returns the size of the rectangle.
+	 * Returns the size of the rectangle in each dimension.
 	 *
 	 * @return The size as a vector.
 	 */
@@ -81,6 +81,14 @@ public:
 		return FVector2D( Right - Left, Bottom - Top );
 	}
 
+	/**
+	 * @return the area of the rectangle
+	 */
+	FORCEINLINE float GetArea() const
+	{
+		return (Right - Left) * (Bottom - Top);
+	}
+
 	/**
 	 * Returns the center of the rectangle
 	 * 
diff --git a/Engine/Source/Runtime/SlateCore/Public/Widgets/SWindow.h b/Engine/Source/Runtime/SlateCore/Public/Widgets/SWindow.h
--- a/Engine/Source/Runtime/SlateCore/Public/Widgets/SWindow.h
+++ b/Engine/Source/Runtime/SlateCore/Public/Widgets/SWindow.h
@@ -95,9 +95,6 @@ namespace SWindowDefs
 {
 	/** Height of a Slate window title bar, in pixels */
 	static const float DefaultTitleBarSize = 24.0f;
-
-	/** Size of the corner rounding radius.  Used for regular, non-maximized windows only (not tool-tips or decorators.) */
-	static const int32 CornerRadius = 6;
 }
 
 /** Proxy structure to handle deprecated construction from bool */
@@ -159,7 +156,7 @@ public:
 		, _ShouldPreserveAspectRatio( false )
 		, _CreateTitleBar( true )
 		, _SaneWindowPlacement( true )
-		, _LayoutBorder(FMargin(5, 5, 5, 5))
+		, _LayoutBorder( _Style->BorderPadding )
 		, _UserResizeBorder(FMargin(5, 5, 5, 5))
 		, _bManualManageDPI( false )
 
@@ -531,8 +528,9 @@ public:
 	 * Removes a widget that is being drawn over the entire window
 	 *
 	 * @param	InContent	The widget to remove
+	 * @return	true if successful
 	 */
-	void RemoveOverlaySlot( const TSharedRef<SWidget>& InContent );
+	bool RemoveOverlaySlot(const TSharedRef<SWidget>& InContent);
 
 	/**
 	 * Visualize a new pop-up if possible.  If it's not possible for this widget to host the pop-up
@@ -776,6 +774,11 @@ public:
 		TitleBar = InTitleBar;
 	}
 
+	TSharedPtr<IWindowTitleBar> GetTitleBar() const
+	{
+		return TitleBar;
+	}
+
 	// Events
 	virtual FCursorReply OnCursorQuery( const FGeometry& MyGeometry, const FPointerEvent& CursorEvent ) const override;
 
