CIS Fixes for SSegmentedControl

diff --git a/Engine/Source/Runtime/Slate/Public/Widgets/Input/SSegmentedControl.h b/Engine/Source/Runtime/Slate/Public/Widgets/Input/SSegmentedControl.h
--- a/Engine/Source/Runtime/Slate/Public/Widgets/Input/SSegmentedControl.h
+++ b/Engine/Source/Runtime/Slate/Public/Widgets/Input/SSegmentedControl.h
@@ -27,42 +27,36 @@ class SSegmentedControl : public SCompoundWidget
 public:
 
 	/** Stores the per-child info for this panel type */
-	template< typename OptionType >
-	struct FSlot : public TSlotBase<FSlot<OptionType>> 
+	template< typename SlotOptionType >
+	struct FSlot : public TSlotBase<FSlot<SlotOptionType>> 
 	{
-		FSlot(const OptionType& InValue) 
-		: TSlotBase<FSlot<OptionType>>() 
+		FSlot(const SlotOptionType& InValue) 
+		: TSlotBase<FSlot<SlotOptionType>>() 
 		, _Text()
 		, _Icon(nullptr)
 		, _Value(InValue)
 		{ }
 
-		FSlot<OptionType>& Text(const TAttribute<FText>& InText)
+		FSlot<SlotOptionType>& Text(const TAttribute<FText>& InText)
 		{
 			_Text = InText;
 			return *this;
 		}
 
-		FSlot<OptionType>& Text(FText& InText)
-		{
-			_Text.SetValue(InText);
-			return *this;
-		}
-
-		FSlot<OptionType>& Icon(const TAttribute<const FSlateBrush*>& InBrush)
+		FSlot<SlotOptionType>& Icon(const TAttribute<const FSlateBrush*>& InBrush)
 		{
 			_Icon= InBrush;
 			return *this;
 		}
 
-	friend class SSegmentedControl<OptionType>;
+	friend class SSegmentedControl<SlotOptionType>;
 
 	protected:
 
 		TAttribute<FText> _Text;
 		TAttribute<const FSlateBrush*> _Icon;
 
-		OptionType _Value;
+		SlotOptionType _Value;
 	};
 
 	static FSlot<OptionType>& Slot(const OptionType& InValue)
