Made tab drawers resizable

diff --git a/Engine/Source/Runtime/Slate/Private/Framework/Docking/STabDrawer.cpp b/Engine/Source/Runtime/Slate/Private/Framework/Docking/STabDrawer.cpp
--- a/Engine/Source/Runtime/Slate/Private/Framework/Docking/STabDrawer.cpp
+++ b/Engine/Source/Runtime/Slate/Private/Framework/Docking/STabDrawer.cpp
@@ -19,7 +19,6 @@ void STabDrawer::SetCurrentSize(float InSize)
 
 void STabDrawer::Construct(const FArguments& InArgs, TSharedRef<SDockTab> InTab, ETabDrawerOpenDirection InOpenDirection)
 {
-	Visibility = EVisibility::SelfHitTestInvisible;
 	OpenDirection = InOpenDirection;
 
 	ForTab = InTab;
@@ -180,7 +179,7 @@ FReply STabDrawer::OnMouseButtonUp(const FGeometry& AllottedGeometry, const FPoi
 		bIsResizing = false;
 		FSlateThrottleManager::Get().LeaveResponsiveMode(ResizeThrottleHandle);
 
-		OnTargetDrawerSizeChanged.ExecuteIfBound(TargetDrawerSize);
+		OnTargetDrawerSizeChanged.ExecuteIfBound(SharedThis(this), TargetDrawerSize);
 		return FReply::Handled().ReleaseMouseCapture();
 	}
 	return FReply::Unhandled();
@@ -194,11 +193,26 @@ FReply STabDrawer::OnMouseMove(const FGeometry& AllottedGeometry, const FPointer
 
 	if (bIsResizing && this->HasMouseCapture() && !MouseEvent.GetCursorDelta().IsZero())
 	{
-		const FVector2D LocalMousePos = InitialResizeGeometry.AbsoluteToLocal(MouseEvent.GetScreenSpacePosition());
-		const float DeltaHeight = (InitialResizeGeometry.GetLocalPositionAtCoordinates(FVector2D::ZeroVector) - LocalMousePos).Y;
+		const FVector2D MousePos = MouseEvent.GetScreenSpacePosition();
+		float DeltaSize = 0.0f;
+
+		if (OpenDirection == ETabDrawerOpenDirection::Left)
+		{
+			DeltaSize = (MousePos - InitialResizeGeometry.GetAbsolutePositionAtCoordinates(FVector2D::ZeroVector)).X;
+		}
+		else if (OpenDirection == ETabDrawerOpenDirection::Right)
+		{
+			DeltaSize = (InitialResizeGeometry.GetAbsolutePositionAtCoordinates(FVector2D::ZeroVector) - MousePos).X;
+		}
+		else
+		{
+			DeltaSize = (InitialResizeGeometry.GetAbsolutePositionAtCoordinates(FVector2D::ZeroVector) - MousePos).Y;
+		}
+
 
-		TargetDrawerSize = FMath::Clamp(InitialSizeAtResize + DeltaHeight, MinDrawerSize, MaxDrawerSize);
-		SetCurrentSize(InitialSizeAtResize + DeltaHeight);
+
+		TargetDrawerSize = FMath::Clamp(InitialSizeAtResize + DeltaSize, MinDrawerSize, MaxDrawerSize);
+		SetCurrentSize(InitialSizeAtResize + DeltaSize);
 
 
 		return FReply::Handled();
@@ -218,7 +232,7 @@ void STabDrawer::OnMouseLeave(const FPointerEvent& MouseEvent)
 
 FCursorReply STabDrawer::OnCursorQuery(const FGeometry& MyGeometry, const FPointerEvent& CursorEvent) const
 {
-	return bIsResizing || bIsResizeHandleHovered ? FCursorReply::Cursor(EMouseCursor::ResizeUpDown) : FCursorReply::Unhandled();
+	return bIsResizing || bIsResizeHandleHovered ? FCursorReply::Cursor(OpenDirection == ETabDrawerOpenDirection::Bottom ? EMouseCursor::ResizeUpDown : EMouseCursor::ResizeLeftRight) : FCursorReply::Unhandled();
 }
 
 int32 STabDrawer::OnPaint(const FPaintArgs& Args, const FGeometry& AllottedGeometry, const FSlateRect& MyCullingRect, FSlateWindowElementList& OutDrawElements, int32 LayerId, const FWidgetStyle& InWidgetStyle, bool bParentEnabled) const
@@ -307,7 +321,20 @@ FGeometry STabDrawer::GetRenderTransformedGeometry(const FGeometry& AllottedGeom
 
 FGeometry STabDrawer::GetResizeHandleGeometry(const FGeometry& AllottedGeometry) const
 {
-	return GetRenderTransformedGeometry(AllottedGeometry).MakeChild(ShadowOffset - FVector2D(0.0f, ExpanderSize), FVector2D(AllottedGeometry.GetLocalSize().X - ShadowOffset.X * 2, ExpanderSize));
+	FGeometry RenderTransformedGeometry = GetRenderTransformedGeometry(AllottedGeometry);
+
+	if (OpenDirection == ETabDrawerOpenDirection::Left)
+	{
+		return RenderTransformedGeometry.MakeChild((FVector2D(RenderTransformedGeometry.GetLocalSize().X-ShadowOffset.X, ShadowOffset.Y)), FVector2D(ExpanderSize, AllottedGeometry.GetLocalSize().Y - ShadowOffset.Y * 2));
+	}
+	else if (OpenDirection == ETabDrawerOpenDirection::Right)
+	{
+		return RenderTransformedGeometry.MakeChild(ShadowOffset - FVector2D(ExpanderSize, 0.0f), FVector2D(ExpanderSize, AllottedGeometry.GetLocalSize().Y - ShadowOffset.Y * 2));
+	}
+	else
+	{
+		return RenderTransformedGeometry.MakeChild(ShadowOffset - FVector2D(0.0f, ExpanderSize), FVector2D(AllottedGeometry.GetLocalSize().X - ShadowOffset.X * 2, ExpanderSize));
+	}
 }
 
 EActiveTimerReturnType STabDrawer::UpdateAnimation(double CurrentTime, float DeltaTime)
diff --git a/Engine/Source/Runtime/Slate/Private/Framework/Docking/STabSidebar.cpp b/Engine/Source/Runtime/Slate/Private/Framework/Docking/STabSidebar.cpp
--- a/Engine/Source/Runtime/Slate/Private/Framework/Docking/STabSidebar.cpp
+++ b/Engine/Source/Runtime/Slate/Private/Framework/Docking/STabSidebar.cpp
@@ -341,6 +341,17 @@ void STabSidebar::OnTabDrawerClosed(TSharedRef<STabDrawer> Drawer)
 	RemoveDrawer(Drawer->GetTab());
 }
 
+void STabSidebar::OnTargetDrawerSizeChanged(TSharedRef<STabDrawer> Drawer, float NewSize)
+{
+	TSharedRef<SDockTab> Tab = Drawer->GetTab();
+	TSharedPtr<SWindow> MyWindow = FSlateApplication::Get().FindWidgetWindow(AsShared());
+	if (MyWindow.IsValid())
+	{
+		const float TargetDrawerSizePct = NewSize / MyWindow->GetPaintSpaceGeometry().GetLocalSize().X;
+		Tab->GetParentDockTabStack()->SetTabSidebarSizeCoefficient(Tab, TargetDrawerSizePct);
+	}
+}
+
 void STabSidebar::OnWindowDPIScaleChanged(TSharedRef<SWindow> WindowThatChanged)
 {
 	if (WindowThatChanged == WindowWithOverlayContent)
@@ -540,7 +551,7 @@ void STabSidebar::OpenDrawerInternal(TSharedRef<SDockTab> ForTab)
 			.MaxDrawerSize(MaxDrawerSize)
 			.OnDrawerFocusLost(this, &STabSidebar::OnTabDrawerFocusLost)
 			.OnDrawerClosed(this, &STabSidebar::OnTabDrawerClosed)
-			//.OnTargetHeightChanged(this, &SStatusBar::OnContentBrowserTargetHeightChanged)
+			.OnTargetDrawerSizeChanged(this, &STabSidebar::OnTargetDrawerSizeChanged)
 			[
 				ForTab->GetContent()
 			];
diff --git a/Engine/Source/Runtime/Slate/Private/Framework/Docking/STabSidebar.h b/Engine/Source/Runtime/Slate/Private/Framework/Docking/STabSidebar.h
--- a/Engine/Source/Runtime/Slate/Private/Framework/Docking/STabSidebar.h
+++ b/Engine/Source/Runtime/Slate/Private/Framework/Docking/STabSidebar.h
@@ -73,6 +73,7 @@ private:
 	void OnTabDrawerButtonClicked(TSharedRef<SDockTab> ForTab);
 	void OnTabDrawerFocusLost(TSharedRef<STabDrawer> Drawer);
 	void OnTabDrawerClosed(TSharedRef<STabDrawer> Drawer);
+	void OnTargetDrawerSizeChanged(TSharedRef<STabDrawer> Drawer, float NewSize);
 	void OnWindowDPIScaleChanged(TSharedRef<SWindow> WindowThatChanged);
 	/** Called when active tab changes which is used to determine if we should close an opened content browser drawer */
 	void OnActiveTabChanged(TSharedPtr<SDockTab> PreviouslyActive, TSharedPtr<SDockTab> NewlyActivated);
diff --git a/Engine/Source/Runtime/Slate/Public/Framework/Docking/STabDrawer.h b/Engine/Source/Runtime/Slate/Public/Framework/Docking/STabDrawer.h
--- a/Engine/Source/Runtime/Slate/Public/Framework/Docking/STabDrawer.h
+++ b/Engine/Source/Runtime/Slate/Public/Framework/Docking/STabDrawer.h
@@ -25,7 +25,7 @@ enum class ETabDrawerOpenDirection : uint8
 	Bottom
 };
 
-DECLARE_DELEGATE_OneParam(FOnDrawerTargetSizeChanged, float);
+DECLARE_DELEGATE_TwoParams(FOnDrawerTargetSizeChanged, TSharedRef<STabDrawer>, float);
 DECLARE_DELEGATE_OneParam(FOnDrawerFocusLost, TSharedRef<STabDrawer>);
 DECLARE_DELEGATE_OneParam(FOnDrawerClosed, TSharedRef<STabDrawer>);
 
