Fix details panel in a drawer closing when color picker is opened

diff --git a/Engine/Source/Editor/DetailCustomizations/Private/ColorStructCustomization.cpp b/Engine/Source/Editor/DetailCustomizations/Private/ColorStructCustomization.cpp
--- a/Engine/Source/Editor/DetailCustomizations/Private/ColorStructCustomization.cpp
+++ b/Engine/Source/Editor/DetailCustomizations/Private/ColorStructCustomization.cpp
@@ -76,6 +76,14 @@ void FColorStructCustomization::MakeHeaderRow(TSharedRef<class IPropertyHandle>&
 }
 
 
+FColorStructCustomization::~FColorStructCustomization()
+{
+	if (ActiveTransaction != INDEX_NONE && GEditor)
+	{
+		GEditor->CancelTransaction(ActiveTransaction);
+	}
+}
+
 TSharedRef<SWidget> FColorStructCustomization::CreateColorWidget(TWeakPtr<IPropertyHandle> StructWeakHandlePtr)
 {
 	FSlateFontInfo NormalText = IDetailLayoutBuilder::GetDetailFont();
@@ -173,7 +181,7 @@ void FColorStructCustomization::GetSortedChildren(TSharedRef<IPropertyHandle> In
 
 void FColorStructCustomization::CreateColorPicker(bool bUseAlpha)
 {
-	GEditor->BeginTransaction(FText::Format(LOCTEXT("SetColorProperty", "Edit {0}"), StructPropertyHandle->GetPropertyDisplayName()));
+	ActiveTransaction = GEditor->BeginTransaction(FText::Format(LOCTEXT("SetColorProperty", "Edit {0}"), StructPropertyHandle->GetPropertyDisplayName()));
 
 	int32 NumObjects = StructPropertyHandle->GetNumOuterObjects();
 
@@ -231,7 +239,7 @@ void FColorStructCustomization::CreateColorPicker(bool bUseAlpha)
 
 TSharedRef<SColorPicker> FColorStructCustomization::CreateInlineColorPicker(TWeakPtr<IPropertyHandle> StructWeakHandlePtr)
 {
-	GEditor->BeginTransaction(FText::Format(LOCTEXT("SetColorProperty", "Edit {0}"), StructPropertyHandle->GetPropertyDisplayName()));
+	ActiveTransaction = GEditor->BeginTransaction(FText::Format(LOCTEXT("SetColorProperty", "Edit {0}"), StructPropertyHandle->GetPropertyDisplayName()));
 
 	int32 NumObjects = StructPropertyHandle->GetNumOuterObjects();
 
@@ -318,7 +326,12 @@ void FColorStructCustomization::OnColorPickerCancelled(FLinearColor OriginalColo
 		StructPropertyHandle->SetPerObjectValues(PerObjectColors);
 	}
 
-	GEditor->CancelTransaction(0);
+	if (ActiveTransaction != INDEX_NONE)
+	{
+		GEditor->CancelTransaction(ActiveTransaction);
+		ActiveTransaction = INDEX_NONE;
+	}
+
 }
 
 void FColorStructCustomization::OnColorPickerWindowClosed(const TSharedRef<SWindow>& Window)
@@ -328,6 +341,7 @@ void FColorStructCustomization::OnColorPickerWindowClosed(const TSharedRef<SWind
 	StructPropertyHandle->GetValueAsFormattedString(ColorString);
 	StructPropertyHandle->SetValueFromFormattedString(ColorString);
 	GEditor->EndTransaction();
+	ActiveTransaction = INDEX_NONE;
 }
 
 
diff --git a/Engine/Source/Editor/DetailCustomizations/Public/Customizations/ColorStructCustomization.h b/Engine/Source/Editor/DetailCustomizations/Public/Customizations/ColorStructCustomization.h
--- a/Engine/Source/Editor/DetailCustomizations/Public/Customizations/ColorStructCustomization.h
+++ b/Engine/Source/Editor/DetailCustomizations/Public/Customizations/ColorStructCustomization.h
@@ -22,7 +22,8 @@ class DETAILCUSTOMIZATIONS_API FColorStructCustomization
 public:
 
 	static TSharedRef<IPropertyTypeCustomization> MakeInstance();
-	
+	~FColorStructCustomization();
+
 protected:
 
 	FColorStructCustomization()
@@ -30,6 +31,7 @@ protected:
 		, bIsInlineColorPickerVisible(false)
 		, bIsInteractive(false)
 		, bDontUpdateWhileEditing(false)
+		, ActiveTransaction(INDEX_NONE)
 	{}
 
 protected:
@@ -157,4 +159,6 @@ protected:
 
 	/** Overrides the default state of the sRGB check box */
 	TOptional<bool> sRGBOverride;
+
+	int32 ActiveTransaction;
 };
diff --git a/Engine/Source/Editor/LevelEditor/LevelEditor.Build.cs b/Engine/Source/Editor/LevelEditor/LevelEditor.Build.cs
--- a/Engine/Source/Editor/LevelEditor/LevelEditor.Build.cs
+++ b/Engine/Source/Editor/LevelEditor/LevelEditor.Build.cs
@@ -69,6 +69,7 @@ public class LevelEditor : ModuleRules
 				"CommonMenuExtensions",
 				"ToolMenus",
 				"StatusBar",
+				"AppFramework"
 			}
 		);
 
diff --git a/Engine/Source/Editor/LevelEditor/Private/SLevelEditor.cpp b/Engine/Source/Editor/LevelEditor/Private/SLevelEditor.cpp
--- a/Engine/Source/Editor/LevelEditor/Private/SLevelEditor.cpp
+++ b/Engine/Source/Editor/LevelEditor/Private/SLevelEditor.cpp
@@ -53,6 +53,7 @@
 #include "IPlacementModeModule.h"
 #include "Classes/EditorStyleSettings.h"
 #include "StatusBarSubsystem.h"
+#include "Widgets/Colors/SColorPicker.h"
 
 static const FName MainFrameModuleName("MainFrame");
 static const FName LevelEditorModuleName("LevelEditor");
@@ -587,12 +588,25 @@ TSharedRef<SDockTab> SLevelEditor::SummonDetailsPanel( FName TabIdentifier )
 {
 	TSharedRef<SActorDetails> ActorDetails = StaticCastSharedRef<SActorDetails>( CreateActorDetails( TabIdentifier ) );
 
+	TWeakPtr<SActorDetails> ActorDetailsWeakPtr = ActorDetails;
+
 	const FText Label = NSLOCTEXT( "LevelEditor", "DetailsTabTitle", "Details" );
 
 	TSharedRef<SDockTab> DocTab = SNew(SDockTab)
 		.Icon( FEditorStyle::GetBrush( "LevelEditor.Tabs.Details" ) )
 		.Label( Label )
 		.ToolTip( IDocumentation::Get()->CreateToolTip( Label, nullptr, "Shared/LevelEditor", "DetailsTab" ) )
+		.OnTabDrawerClosed_Lambda([ActorDetailsWeakPtr]()
+			{
+				// Close the color picker if a details panel is put into a drawer since you cant visually see the properties anymore
+				if (TSharedPtr<SColorPicker> ColorPicker = GetColorPicker())
+				{
+					if (ColorPicker->GetOptionalOwningDetailsView().IsValid())
+					{
+						DestroyColorPicker();
+					}
+				}
+			})
 		[
 			SNew( SBorder)
 			.BorderImage(FAppStyle::Get().GetBrush("Brushes.Background"))
diff --git a/Engine/Source/Runtime/Slate/Private/Framework/Docking/STabDrawer.cpp b/Engine/Source/Runtime/Slate/Private/Framework/Docking/STabDrawer.cpp
--- a/Engine/Source/Runtime/Slate/Private/Framework/Docking/STabDrawer.cpp
+++ b/Engine/Source/Runtime/Slate/Private/Framework/Docking/STabDrawer.cpp
@@ -356,6 +356,19 @@ EActiveTimerReturnType STabDrawer::UpdateAnimation(double CurrentTime, float Del
 	return EActiveTimerReturnType::Continue;
 }
 
+static bool IsLegalWidgetFocused(const FWidgetPath& FocusPath, const TArrayView<TSharedRef<SWidget>> LegalFocusWidgets)
+{
+	for (const TSharedRef<const SWidget>& Widget : LegalFocusWidgets)
+	{
+		if (FocusPath.ContainsWidget(Widget))
+		{
+			return true;
+		}
+	}
+
+	return false;
+}
+
 void STabDrawer::OnGlobalFocusChanging(const FFocusEvent& FocusEvent, const FWeakWidgetPath& OldFocusedWidgetPath, const TSharedPtr<SWidget>& OldFocusedWidget, const FWidgetPath& NewFocusedWidgetPath, const TSharedPtr<SWidget>& NewFocusedWidget)
 {
 	// Sometimes when dismissing focus can change which will trigger this again
@@ -366,26 +379,35 @@ void STabDrawer::OnGlobalFocusChanging(const FFocusEvent& FocusEvent, const FWea
 		TGuardValue<bool> RentrancyGuard(bIsRentrant, true);
 
 		TSharedRef<STabDrawer> ThisWidget = SharedThis(this);
+		TArray<TSharedRef<SWidget>, TInlineAllocator<4>> LegalFocusWidgets;
+		LegalFocusWidgets.Add(ThisWidget);
+		LegalFocusWidgets.Add(ChildSlot.GetWidget());
 
 		bool bShouldLoseFocus = false;
 		// Do not close due to slow tasks as those opening send window activation events
-		if (!GIsSlowTask && !FSlateApplication::Get().GetActiveModalWindow().IsValid() && (!NewFocusedWidgetPath.ContainsWidget(ChildSlot.GetWidget()) && !NewFocusedWidgetPath.ContainsWidget(ThisWidget)))
+		if (!GIsSlowTask && !FSlateApplication::Get().GetActiveModalWindow().IsValid() && !IsLegalWidgetFocused(NewFocusedWidgetPath, MakeArrayView(LegalFocusWidgets)))
 		{
-			if (TSharedPtr<SWidget> MenuHost = FSlateApplication::Get().GetMenuHostWidget())
-			{
-				FWidgetPath MenuHostPath;
+			TSharedRef<SWindow> NewWindow = NewFocusedWidgetPath.GetWindow();
+			TSharedPtr<SWindow> MyWindow = FSlateApplication::Get().FindWidgetWindow(ThisWidget);
 
-				// See if the menu being opened is owned by the drawer contents and if so the menu should not be dismissed
-				FSlateApplication::Get().GeneratePathToWidgetUnchecked(MenuHost.ToSharedRef(), MenuHostPath);
-				if (!MenuHostPath.ContainsWidget(ChildSlot.GetWidget()))
+			if (!NewWindow->IsDescendantOf(MyWindow))
+			{
+				if (TSharedPtr<SWidget> MenuHost = FSlateApplication::Get().GetMenuHostWidget())
+				{
+					FWidgetPath MenuHostPath;
+
+					// See if the menu being opened is owned by the drawer contents and if so the menu should not be dismissed
+					FSlateApplication::Get().GeneratePathToWidgetUnchecked(MenuHost.ToSharedRef(), MenuHostPath);
+					if (!MenuHostPath.ContainsWidget(ChildSlot.GetWidget()))
+					{
+						bShouldLoseFocus = true;
+					}
+				}
+				else
 				{
 					bShouldLoseFocus = true;
 				}
 			}
-			else
-			{
-				bShouldLoseFocus = true;
-			}
 		}
 
 		if (bShouldLoseFocus)
diff --git a/Engine/Source/Runtime/Slate/Private/Framework/Docking/STabSidebar.cpp b/Engine/Source/Runtime/Slate/Private/Framework/Docking/STabSidebar.cpp
--- a/Engine/Source/Runtime/Slate/Private/Framework/Docking/STabSidebar.cpp
+++ b/Engine/Source/Runtime/Slate/Private/Framework/Docking/STabSidebar.cpp
@@ -360,7 +360,7 @@ void STabSidebar::OnWindowDPIScaleChanged(TSharedRef<SWindow> WindowThatChanged)
 	}
 }
 
-void STabSidebar::OnActiveTabChanged(TSharedPtr<SDockTab> PreviouslyActive, TSharedPtr<SDockTab> NewlyActivated)
+void STabSidebar::OnActiveTabChanged(TSharedPtr<SDockTab> NewlyActivated, TSharedPtr<SDockTab> PreviouslyActive)
 {
 	// If a new major tab was activated remove any visible drawer instantly
 	if (NewlyActivated.IsValid() && NewlyActivated->GetVisualTabRole() == ETabRole::MajorTab)
@@ -442,6 +442,8 @@ void STabSidebar::RemoveDrawer(TSharedRef<SDockTab> ForTab)
 		OpenedDrawers.Remove(OpenedDrawerRef);
 	}
 
+	ForTab->OnTabDrawerClosed();
+
 	UpdateDrawerAppearance();
 }
 
diff --git a/Engine/Source/Runtime/Slate/Private/Framework/Docking/STabSidebar.h b/Engine/Source/Runtime/Slate/Private/Framework/Docking/STabSidebar.h
--- a/Engine/Source/Runtime/Slate/Private/Framework/Docking/STabSidebar.h
+++ b/Engine/Source/Runtime/Slate/Private/Framework/Docking/STabSidebar.h
@@ -76,7 +76,7 @@ private:
 	void OnTargetDrawerSizeChanged(TSharedRef<STabDrawer> Drawer, float NewSize);
 	void OnWindowDPIScaleChanged(TSharedRef<SWindow> WindowThatChanged);
 	/** Called when active tab changes which is used to determine if we should close an opened content browser drawer */
-	void OnActiveTabChanged(TSharedPtr<SDockTab> PreviouslyActive, TSharedPtr<SDockTab> NewlyActivated);
+	void OnActiveTabChanged(TSharedPtr<SDockTab> NewlyActivated, TSharedPtr<SDockTab> PreviouslyActive);
 
 	TSharedRef<SWidget> OnGetTabDrawerContextMenuWidget(TSharedRef<SDockTab> ForTab);
 	void OnRestoreTab(TSharedRef<SDockTab> TabToRestore);
diff --git a/Engine/Source/Runtime/Slate/Private/Widgets/Docking/SDockTab.cpp b/Engine/Source/Runtime/Slate/Private/Widgets/Docking/SDockTab.cpp
--- a/Engine/Source/Runtime/Slate/Private/Widgets/Docking/SDockTab.cpp
+++ b/Engine/Source/Runtime/Slate/Private/Widgets/Docking/SDockTab.cpp
@@ -515,6 +515,8 @@ void SDockTab::Construct( const FArguments& InArgs )
 	this->TabIcon = InArgs._Icon;
 	this->bShouldAutosize = InArgs._ShouldAutosize;
 	this->TabColorScale = InArgs._TabColorScale;
+	
+	OnTabDrawerClosedEvent = InArgs._OnTabDrawerClosed;
 
 	MajorTabStyle = &FAppStyle::Get().GetWidgetStyle<FDockTabStyle>("Docking.MajorTab");
 	GenericTabStyle = &FAppStyle::Get().GetWidgetStyle<FDockTabStyle>("Docking.Tab");
@@ -685,6 +687,11 @@ void SDockTab::UpdateTabStyle()
 	LabelWidget->SetTextStyle(&TabTextStyle);
 }
 
+void SDockTab::OnTabDrawerClosed()
+{
+	OnTabDrawerClosedEvent.ExecuteIfBound();
+}
+
 const FDockTabStyle& SDockTab::GetCurrentStyle() const
 {
 	if ( GetVisualTabRole() == ETabRole::MajorTab )
diff --git a/Engine/Source/Runtime/Slate/Public/Framework/Docking/TabManager.h b/Engine/Source/Runtime/Slate/Public/Framework/Docking/TabManager.h
--- a/Engine/Source/Runtime/Slate/Public/Framework/Docking/TabManager.h
+++ b/Engine/Source/Runtime/Slate/Public/Framework/Docking/TabManager.h
@@ -22,9 +22,9 @@ struct FSidebarTabLists;
 
 DECLARE_MULTICAST_DELEGATE_TwoParams(
 	FOnActiveTabChanged,
-	/** Previously active tab */
-	TSharedPtr<SDockTab>,
 	/** Newly active tab */
+	TSharedPtr<SDockTab>,
+	/** Previously active tab */
 	TSharedPtr<SDockTab> );
 	
 
diff --git a/Engine/Source/Runtime/Slate/Public/Widgets/Docking/SDockTab.h b/Engine/Source/Runtime/Slate/Public/Widgets/Docking/SDockTab.h
--- a/Engine/Source/Runtime/Slate/Public/Widgets/Docking/SDockTab.h
+++ b/Engine/Source/Runtime/Slate/Public/Widgets/Docking/SDockTab.h
@@ -47,6 +47,7 @@ enum ETabActivationCause : uint8
 class SLATE_API SDockTab : public SBorder
 {
 	friend class FTabManager;
+	friend class STabSidebar;
 public:
 
 	/** Invoked when a tab is closing */
@@ -92,6 +93,8 @@ public:
 		SLATE_ARGUMENT( bool, ShouldAutosize )
 		SLATE_EVENT( FCanCloseTab, OnCanCloseTab )
 		SLATE_EVENT( FOnPersistVisualState, OnPersistVisualState )
+		/** Invoked when a tab is closed from a drawer. This does not mean the tab or its contents is destroyed, just hidden. Use OnTabClosed for that */
+		SLATE_EVENT( FSimpleDelegate, OnTabDrawerClosed)
 		SLATE_ATTRIBUTE( FLinearColor, TabColorScale )
 		SLATE_ATTRIBUTE( FSlateColor, ForegroundColor )
 	SLATE_END_ARGS()
@@ -333,6 +336,8 @@ private:
 
 	void UpdateTabStyle();
 
+	void OnTabDrawerClosed();
+
 	/** The handle to the active tab activation tick */
 	TWeakPtr<FActiveTimerHandle> DragDropTimerHandle;
 	TWeakPtr<FActiveTimerHandle> UpdateStyleTimerHandle;
@@ -374,6 +379,8 @@ protected:
 	/** Delegate to execute to determine if we can close this tab */
 	FCanCloseTab OnCanCloseTab;
 
+	FSimpleDelegate OnTabDrawerClosedEvent;
+
 	/**
 	 * Invoked during the Save Visual State pass; gives this tab a chance to save misc info about visual state.
 	 * e.g. Content Browser might save the current filters, current folder, whether some panel is collapsed, etc.
