Fix mode tab not closing properly if in a drawer. Tabs now also remember their spot in a drawer if they were closed from said drawer.

diff --git a/Engine/Source/Runtime/Slate/Private/Framework/Docking/SDockingArea.cpp b/Engine/Source/Runtime/Slate/Private/Framework/Docking/SDockingArea.cpp
--- a/Engine/Source/Runtime/Slate/Private/Framework/Docking/SDockingArea.cpp
+++ b/Engine/Source/Runtime/Slate/Private/Framework/Docking/SDockingArea.cpp
@@ -409,6 +409,11 @@ bool SDockingArea::IsTabInSidebar(TSharedRef<SDockTab> Tab) const
 	return LeftSidebar->ContainsTab(Tab) || RightSidebar->ContainsTab(Tab);
 }
 
+bool SDockingArea::RemoveTabFromSidebar(TSharedRef<SDockTab> Tab)
+{
+	return LeftSidebar->RemoveTab(Tab) || RightSidebar->RemoveTab(Tab);
+}
+
 bool SDockingArea::TryOpenSidebarDrawer(TSharedRef<SDockTab> TabToOpen) const
 {
 	return LeftSidebar->TryOpenSidebarDrawer(TabToOpen) || RightSidebar->TryOpenSidebarDrawer(TabToOpen);
@@ -427,6 +432,15 @@ void SDockingArea::AddSidebarTabsFromRestoredLayout(const FSidebarTabLists& Side
 	}
 }
 
+TArray<TSharedRef<SDockTab>> SDockingArea::GetAllSidebarTabs() const
+{
+	TArray<TSharedRef<SDockTab>> AllSidebarTabs;
+	AllSidebarTabs.Append(LeftSidebar->GetAllTabs());
+	AllSidebarTabs.Append(RightSidebar->GetAllTabs());
+
+	return AllSidebarTabs;
+}
+
 SDockingNode::ECleanupRetVal SDockingArea::CleanUpNodes()
 {
 	SDockingNode::ECleanupRetVal ReturnValue = SDockingSplitter::CleanUpNodes();
diff --git a/Engine/Source/Runtime/Slate/Private/Framework/Docking/SDockingArea.h b/Engine/Source/Runtime/Slate/Private/Framework/Docking/SDockingArea.h
--- a/Engine/Source/Runtime/Slate/Private/Framework/Docking/SDockingArea.h
+++ b/Engine/Source/Runtime/Slate/Private/Framework/Docking/SDockingArea.h
@@ -112,6 +112,12 @@ public:
 	 */
 	bool IsTabInSidebar(TSharedRef<SDockTab> Tab) const;
 
+	/**
+ 	 * Removes a tab from a sidebar
+	 * @return true if the specified tab was found and removed
+	 */
+	bool RemoveTabFromSidebar(TSharedRef<SDockTab> Tab);
+
 	/**
 	 * Attempts to open a sidebar drawer that may the tab to open
 	 *
@@ -124,6 +130,11 @@ public:
 	 */
 	void AddSidebarTabsFromRestoredLayout(const FSidebarTabLists& SidebarTabs);
 
+	/**
+	 * Gets all tabs in all sidebars in this dock area
+	 */
+	TArray<TSharedRef<SDockTab>> GetAllSidebarTabs() const;
+
 	bool CanHaveSidebar() const { return bCanHaveSidebar; }
 protected:
 	virtual SDockingNode::ECleanupRetVal CleanUpNodes() override;
diff --git a/Engine/Source/Runtime/Slate/Private/Framework/Docking/SDockingTabStack.cpp b/Engine/Source/Runtime/Slate/Private/Framework/Docking/SDockingTabStack.cpp
--- a/Engine/Source/Runtime/Slate/Private/Framework/Docking/SDockingTabStack.cpp
+++ b/Engine/Source/Runtime/Slate/Private/Framework/Docking/SDockingTabStack.cpp
@@ -273,11 +273,34 @@ void SDockingTabStack::OnTabRemoved( const FTabId& TabId )
 void SDockingTabStack::OpenTab( const TSharedRef<SDockTab>& InTab, int32 InsertLocationAmongActiveTabs )
 {
 	const int32 TabIndex = OpenPersistentTab( InTab->GetLayoutIdentifier(), InsertLocationAmongActiveTabs );
+
 	// The tab may be a nomad tab, in which case it should inherit whichever tab manager it is being put into!
-	InTab->SetTabManager( GetDockArea()->GetTabManager() );
-	AddTabWidget( InTab, TabIndex );
-	OnLiveTabAdded();
-	TabWell->RefreshParentContent();
+	InTab->SetTabManager(GetDockArea()->GetTabManager());
+
+	const FTabManager::FTab& TabInfo = Tabs[TabIndex];
+	if (TabInfo.TabState == ETabState::SidebarTab)
+	{
+		FSidebarTabLists SidebarLists;
+		if (TabInfo.SidebarLocation == ESidebarLocation::Left)
+		{
+			SidebarLists.LeftSidebarTabs.Add(InTab);
+		}
+		else
+		{
+			ensure(TabInfo.SidebarLocation == ESidebarLocation::Right);
+			SidebarLists.RightSidebarTabs.Add(InTab);
+		}
+
+		AddSidebarTab(InTab);
+		GetDockArea()->AddSidebarTabsFromRestoredLayout(SidebarLists);
+	}
+	else
+	{
+		AddTabWidget(InTab, TabIndex);
+		OnLiveTabAdded();
+		TabWell->RefreshParentContent();
+	}
+
 }
 
 void SDockingTabStack::AddTabWidget( const TSharedRef<SDockTab>& InTab, int32 AtLocation )
@@ -948,6 +971,21 @@ void SDockingTabStack::MoveTabToSidebar(TSharedRef<SDockTab> Tab)
 	}
 }
 
+void SDockingTabStack::RestoreTabFromSidebar(TSharedRef<SDockTab> Tab)
+{
+	const int32 TabIndex = Tabs.IndexOfByPredicate(FTabMatcher(Tab->GetLayoutIdentifier(), ETabState::SidebarTab));
+	if (TabIndex != INDEX_NONE)
+	{
+		FTabManager::FTab& TabInfo = Tabs[TabIndex];
+
+		TabInfo.SidebarSizeCoefficient = 0;
+		// Set the sate to closed so its reopened by OpenTab
+		TabInfo.TabState = ETabState::ClosedTab;
+		TabInfo.SidebarLocation = ESidebarLocation::None;
+		OpenTab(Tab);
+	}
+}
+
 FReply SDockingTabStack::UnhideTabWell()
 {
 	SetTabWellHidden(false);
@@ -1099,8 +1137,8 @@ int32 SDockingTabStack::OpenPersistentTab( const FTabId& TabId, int32 OpenLocati
 	{						
 		if (ExistingClosedTabIndex != INDEX_NONE)
 		{
-			// There's already a tab with that name; open it.
-			Tabs[ExistingClosedTabIndex].TabState = ETabState::OpenedTab;
+			FTabManager::FTab& Tab = Tabs[ExistingClosedTabIndex];
+			Tab.TabState = Tab.SidebarLocation == ESidebarLocation::None ? ETabState::OpenedTab : ETabState::SidebarTab;
 			return ExistingClosedTabIndex;
 		}
 		else
diff --git a/Engine/Source/Runtime/Slate/Private/Framework/Docking/SDockingTabStack.h b/Engine/Source/Runtime/Slate/Private/Framework/Docking/SDockingTabStack.h
--- a/Engine/Source/Runtime/Slate/Private/Framework/Docking/SDockingTabStack.h
+++ b/Engine/Source/Runtime/Slate/Private/Framework/Docking/SDockingTabStack.h
@@ -86,6 +86,9 @@ public:
 	/** Moves a specific tab from this stack into a sidebar */
 	void MoveTabToSidebar(TSharedRef<SDockTab> Tab);
 
+	/** Restores a specific tab from this stack to its original position */
+	void RestoreTabFromSidebar(TSharedRef<SDockTab> Tab);
+
 	/** @return All child tabs in this node */
 	const TSlotlessChildren<SDockTab>& GetTabs() const;
 
diff --git a/Engine/Source/Runtime/Slate/Private/Framework/Docking/STabSidebar.cpp b/Engine/Source/Runtime/Slate/Private/Framework/Docking/STabSidebar.cpp
--- a/Engine/Source/Runtime/Slate/Private/Framework/Docking/STabSidebar.cpp
+++ b/Engine/Source/Runtime/Slate/Private/Framework/Docking/STabSidebar.cpp
@@ -280,9 +280,7 @@ bool STabSidebar::RestoreTab(TSharedRef<SDockTab> TabToRestore)
 {
 	if(RemoveTab(TabToRestore))
 	{
-		TabToRestore->GetParentDockTabStack()->SetTabSidebarSizeCoefficient(TabToRestore, 0);
-		TabToRestore->GetParentDockTabStack()->OpenTab(TabToRestore);
-
+		TabToRestore->GetParentDockTabStack()->RestoreTabFromSidebar(TabToRestore);
 		return true;
 	}
 
@@ -309,6 +307,17 @@ TArray<FTabId> STabSidebar::GetAllTabIds() const
 	return TabIds;
 }
 
+TArray<TSharedRef<SDockTab>> STabSidebar::GetAllTabs() const
+{
+	TArray<TSharedRef<SDockTab>> DockTabs;
+	for (auto TabPair : Tabs)
+	{
+		DockTabs.Add(TabPair.Key);
+	}
+
+	return DockTabs;
+}
+
 bool STabSidebar::TryOpenSidebarDrawer(TSharedRef<SDockTab> ForTab)
 {
 	int32 FoundIndex = Tabs.IndexOfByPredicate(
@@ -413,10 +422,11 @@ void STabSidebar::OnRestoreTab(TSharedRef<SDockTab> TabToRestore)
 
 void STabSidebar::OnCloseTab(TSharedRef<SDockTab> TabToClose)
 {
-	RemoveTab(TabToClose);
-
-	TabToClose->RequestCloseTab();
-	TabToClose->GetParentDockTabStack()->OnTabClosed(TabToClose, SDockingNode::TabRemoval_Closed);
+	if(TabToClose->RequestCloseTab())
+	{
+		RemoveTab(TabToClose);
+		TabToClose->GetParentDockTabStack()->OnTabClosed(TabToClose, SDockingNode::TabRemoval_Closed);
+	}
 }
 
 void STabSidebar::RemoveDrawer(TSharedRef<SDockTab> ForTab)
diff --git a/Engine/Source/Runtime/Slate/Private/Framework/Docking/STabSidebar.h b/Engine/Source/Runtime/Slate/Private/Framework/Docking/STabSidebar.h
--- a/Engine/Source/Runtime/Slate/Private/Framework/Docking/STabSidebar.h
+++ b/Engine/Source/Runtime/Slate/Private/Framework/Docking/STabSidebar.h
@@ -42,6 +42,7 @@ public:
 	/**
 	 * Removes a tab from the sidebar. Does not restore it to the parent stack
 	 *
+	 * Note it is not sufficient to call this to clean up the tab completely. Call RequestCloseTab on the dock tab to do that.
 	 * @return true if the tab was found and removal was successful
 	 */
 	bool RemoveTab(TSharedRef<SDockTab> TabToRemove);
@@ -63,6 +64,11 @@ public:
 	 */
 	TArray<FTabId> GetAllTabIds() const;
 
+	/**
+	 * Get all tabs in this sidebar
+	 */
+	TArray<TSharedRef<SDockTab>> GetAllTabs() const;
+
 	/**
 	 * Attempt to open a drawer in the sidebar for a specified tab
 	 *
diff --git a/Engine/Source/Runtime/Slate/Private/Framework/Docking/TabManager.cpp b/Engine/Source/Runtime/Slate/Private/Framework/Docking/TabManager.cpp
--- a/Engine/Source/Runtime/Slate/Private/Framework/Docking/TabManager.cpp
+++ b/Engine/Source/Runtime/Slate/Private/Framework/Docking/TabManager.cpp
@@ -1927,6 +1927,7 @@ TSharedPtr<SDockTab> FTabManager::FindExistingLiveTab( const FTabId& TabId ) con
 		if ( SomeDockArea.IsValid() )
 		{
 			TArray< TSharedRef<SDockTab> > ChildTabs = SomeDockArea->GetAllChildTabs();
+			ChildTabs.Append(SomeDockArea->GetAllSidebarTabs());
 			for (int32 ChildTabIndex=0; ChildTabIndex < ChildTabs.Num(); ++ChildTabIndex)
 			{
 				if ( TabId == ChildTabs[ChildTabIndex]->GetLayoutIdentifier() )
diff --git a/Engine/Source/Runtime/Slate/Private/Widgets/Docking/SDockTab.cpp b/Engine/Source/Runtime/Slate/Private/Widgets/Docking/SDockTab.cpp
--- a/Engine/Source/Runtime/Slate/Private/Widgets/Docking/SDockTab.cpp
+++ b/Engine/Source/Runtime/Slate/Private/Widgets/Docking/SDockTab.cpp
@@ -866,11 +866,22 @@ bool SDockTab::CanCloseTab() const
 bool SDockTab::RequestCloseTab()
 {	
 	this->PersistVisualState();
+
 	// The tab can be closed if the delegate is not bound or if the delegate call indicates we cannot close it
 	const bool bCanCloseTabNow = CanCloseTab();
-	if( bCanCloseTabNow )
+
+	if(bCanCloseTabNow)
 	{
-		RemoveTabFromParent();
+		bool bRemovedFromSidebar = false;
+		if (GetParentDockTabStack() && GetParentDockTabStack()->GetDockArea())
+		{
+			bRemovedFromSidebar = GetParentDockTabStack()->GetDockArea()->RemoveTabFromSidebar(SharedThis(this));
+		}
+
+		if (bCanCloseTabNow && !bRemovedFromSidebar)
+		{
+			RemoveTabFromParent();
+		}
 	}
 	return bCanCloseTabNow;
 }
