Make sure tab drawer window overlay pointer is cleaned up when all drawers are closed.

diff --git a/Engine/Source/Runtime/Slate/Private/Framework/Docking/STabSidebar.cpp b/Engine/Source/Runtime/Slate/Private/Framework/Docking/STabSidebar.cpp
--- a/Engine/Source/Runtime/Slate/Private/Framework/Docking/STabSidebar.cpp
+++ b/Engine/Source/Runtime/Slate/Private/Framework/Docking/STabSidebar.cpp
@@ -452,6 +452,12 @@ void STabSidebar::RemoveDrawer(TSharedRef<SDockTab> ForTab)
 		OpenedDrawers.Remove(OpenedDrawerRef);
 	}
 
+	if (OpenedDrawers.Num() == 0)
+	{
+		WindowWithOverlayContent.Reset();
+	}
+
+
 	ForTab->OnTabDrawerClosed();
 
 	UpdateDrawerAppearance();
@@ -568,6 +574,7 @@ void STabSidebar::OpenDrawerInternal(TSharedRef<SDockTab> ForTab)
 				ForTab->GetContent()
 			];
 
+		check(!WindowWithOverlayContent.IsValid() || WindowWithOverlayContent == MyWindow);
 		WindowWithOverlayContent = MyWindow;
 
 		MyWindow->AddOverlaySlot()
