Added a tooltip to segmented control slots

diff --git a/Engine/Source/Runtime/Slate/Public/Widgets/Input/SSegmentedControl.h b/Engine/Source/Runtime/Slate/Public/Widgets/Input/SSegmentedControl.h
--- a/Engine/Source/Runtime/Slate/Public/Widgets/Input/SSegmentedControl.h
+++ b/Engine/Source/Runtime/Slate/Public/Widgets/Input/SSegmentedControl.h
@@ -33,6 +33,7 @@ public:
 		FSlot(const SlotOptionType& InValue) 
 		: TSlotBase<FSlot<SlotOptionType>>() 
 		, _Text()
+		, _Tooltip()
 		, _Icon(nullptr)
 		, _Value(InValue)
 		{ }
@@ -49,11 +50,18 @@ public:
 			return *this;
 		}
 
+		FSlot<SlotOptionType>& ToolTip(const TAttribute<FText>& InTooltip)
+		{
+			_Tooltip = InTooltip;
+			return *this;
+		}
+
 	friend class SSegmentedControl<SlotOptionType>;
 
 	protected:
 
 		TAttribute<FText> _Text;
+		TAttribute<FText> _Tooltip;
 		TAttribute<const FSlateBrush*> _Icon;
 
 		SlotOptionType _Value;
@@ -140,6 +148,7 @@ public:
 			.AutoWidth()
 			[
 				SNew(SCheckBox)
+				.ToolTipText(ChildSlotPtr->_Tooltip)
 				.Style(SlotIndex == 0 ? &Style->FirstControlStyle : SlotIndex == (NumSlots - 1) ? &Style->LastControlStyle : &Style->ControlStyle)
 				.IsChecked(this, &SSegmentedControl::IsCurrentValue, ChildValue)
 				.OnCheckStateChanged(this, &SSegmentedControl::CommitValue, ChildValue)
