Fixed up default highlight colors for text widgets

diff --git a/Engine/Source/Runtime/Slate/Private/Framework/Text/SlateTextHighlightRunRenderer.cpp b/Engine/Source/Runtime/Slate/Private/Framework/Text/SlateTextHighlightRunRenderer.cpp
--- a/Engine/Source/Runtime/Slate/Private/Framework/Text/SlateTextHighlightRunRenderer.cpp
+++ b/Engine/Source/Runtime/Slate/Private/Framework/Text/SlateTextHighlightRunRenderer.cpp
@@ -21,13 +21,13 @@ int32 FSlateTextHighlightRunRenderer::OnPaint( const FPaintArgs& Args, const FTe
 	FSlateDrawElement::MakeBox(
 		OutDrawElements,
 		++LayerId,
-		AllottedGeometry.ToPaintGeometry(TransformVector(InverseScale, FVector2D( Block->GetSize().X, Line.Size.Y )), FSlateLayoutTransform(TransformPoint(InverseScale, Location))), 
+		AllottedGeometry.ToPaintGeometry(TransformVector(InverseScale, FVector2D(Block->GetSize().X, Line.Size.Y)), FSlateLayoutTransform(TransformPoint(InverseScale, Location))),
 		&DefaultStyle.HighlightShape,
 		bParentEnabled ? ESlateDrawEffect::None : ESlateDrawEffect::DisabledEffect,
-		InWidgetStyle.GetColorAndOpacityTint() * DefaultStyle.HighlightColor
+		InWidgetStyle.GetColorAndOpacityTint() * DefaultStyle.HighlightShape.TintColor.GetColor(InWidgetStyle)
 		);
 
-	FLinearColor InvertedHighlightColor = FLinearColor::White - DefaultStyle.HighlightColor;
+	FLinearColor InvertedHighlightColor = DefaultStyle.HighlightColor.GetColor(InWidgetStyle);
 	InvertedHighlightColor.A = InWidgetStyle.GetForegroundColor().A;
 
 	FWidgetStyle WidgetStyle( InWidgetStyle );
diff --git a/Engine/Source/Runtime/Slate/Private/Widgets/Text/STextBlock.cpp b/Engine/Source/Runtime/Slate/Private/Widgets/Text/STextBlock.cpp
--- a/Engine/Source/Runtime/Slate/Private/Widgets/Text/STextBlock.cpp
+++ b/Engine/Source/Runtime/Slate/Private/Widgets/Text/STextBlock.cpp
@@ -95,7 +95,7 @@ FLinearColor STextBlock::GetShadowColorAndOpacity() const
 	return ShadowColorAndOpacity.IsSet() ? ShadowColorAndOpacity.Get() : TextStyle.ShadowColorAndOpacity;
 }
 
-FLinearColor STextBlock::GetHighlightColor() const
+FSlateColor STextBlock::GetHighlightColor() const
 {
 	return HighlightColor.IsSet() ? HighlightColor.Get() : TextStyle.HighlightColor;
 }
diff --git a/Engine/Source/Runtime/Slate/Private/Widgets/Text/SlateEditableTextTypes.cpp b/Engine/Source/Runtime/Slate/Private/Widgets/Text/SlateEditableTextTypes.cpp
--- a/Engine/Source/Runtime/Slate/Private/Widgets/Text/SlateEditableTextTypes.cpp
+++ b/Engine/Source/Runtime/Slate/Private/Widgets/Text/SlateEditableTextTypes.cpp
@@ -242,7 +242,7 @@ int32 FTextSearchHighlighter::OnPaint(const FPaintArgs& Args, const FTextLayout:
 	const FVector2D Location(Line.Offset.X + OffsetX, Line.Offset.Y);
 
 	// If we've not been set to an explicit color, calculate a suitable one from the linked color
-	FLinearColor SelectionBackgroundColorAndOpacity = DefaultStyle.HighlightColor * InWidgetStyle.GetColorAndOpacityTint();
+	FLinearColor SelectionBackgroundColorAndOpacity = DefaultStyle.HighlightColor.GetColor(InWidgetStyle) * InWidgetStyle.GetColorAndOpacityTint();
 	SelectionBackgroundColorAndOpacity.A *= 0.2f;
 
 	// The block size and offset values are pre-scaled, so we need to account for that when converting the block offsets into paint geometry
diff --git a/Engine/Source/Runtime/Slate/Public/Widgets/Text/STextBlock.h b/Engine/Source/Runtime/Slate/Public/Widgets/Text/STextBlock.h
--- a/Engine/Source/Runtime/Slate/Public/Widgets/Text/STextBlock.h
+++ b/Engine/Source/Runtime/Slate/Public/Widgets/Text/STextBlock.h
@@ -275,7 +275,7 @@ private:
 	FLinearColor GetShadowColorAndOpacity() const;
 
 	/** Gets the current highlight color */
-	FLinearColor GetHighlightColor() const;
+	FSlateColor GetHighlightColor() const;
 
 	/** Gets the current highlight shape */
 	const FSlateBrush* GetHighlightShape() const;
diff --git a/Engine/Source/Runtime/SlateCore/Private/Styling/SlateTypes.cpp b/Engine/Source/Runtime/SlateCore/Private/Styling/SlateTypes.cpp
--- a/Engine/Source/Runtime/SlateCore/Private/Styling/SlateTypes.cpp
+++ b/Engine/Source/Runtime/SlateCore/Private/Styling/SlateTypes.cpp
@@ -78,12 +78,12 @@ void FCheckBoxStyle::PostSerialize(const FArchive& Ar)
 #endif
 
 FTextBlockStyle::FTextBlockStyle()
-: Font()
-, ColorAndOpacity()
-, ShadowOffset(FVector2D::ZeroVector)
-, ShadowColorAndOpacity(FLinearColor::Black)
-, SelectedBackgroundColor(FSlateColor::UseForeground())
-, HighlightColor(ForceInitToZero)
+	: Font()
+	, ColorAndOpacity()
+	, ShadowOffset(FVector2D::ZeroVector)
+	, ShadowColorAndOpacity(FLinearColor::Black)
+	, SelectedBackgroundColor(FSlateColor::UseForeground())
+	, HighlightColor(FLinearColor::Black)
 {
 }
 
diff --git a/Engine/Source/Runtime/SlateCore/Private/Styling/StarshipCoreStyle.cpp b/Engine/Source/Runtime/SlateCore/Private/Styling/StarshipCoreStyle.cpp
--- a/Engine/Source/Runtime/SlateCore/Private/Styling/StarshipCoreStyle.cpp
+++ b/Engine/Source/Runtime/SlateCore/Private/Styling/StarshipCoreStyle.cpp
@@ -893,7 +893,8 @@ void FStarshipCoreStyle::SetupTextStyles(TSharedRef<FStyle>& Style)
 		.SetColorAndOpacity(FSlateColor::UseForeground())
 		.SetShadowOffset(FVector2D::ZeroVector)
 		.SetSelectedBackgroundColor(FStyleColors::Highlight)
-		.SetHighlightShape(FSlateColorBrush(FStyleColors::White));
+		.SetHighlightColor(FStyleColors::Black)
+		.SetHighlightShape(FSlateColorBrush(FStyleColors::AccentGreen));
 
 	Style->Set("NormalFont", StyleFonts.Normal);
 	Style->Set("SmallFont", StyleFonts.Small);
diff --git a/Engine/Source/Runtime/SlateCore/Public/Styling/SlateTypes.h b/Engine/Source/Runtime/SlateCore/Public/Styling/SlateTypes.h
--- a/Engine/Source/Runtime/SlateCore/Public/Styling/SlateTypes.h
+++ b/Engine/Source/Runtime/SlateCore/Public/Styling/SlateTypes.h
@@ -322,8 +322,8 @@ struct SLATECORE_API FTextBlockStyle : public FSlateWidgetStyle
 
 	/** The color of highlighted text */
 	UPROPERTY(EditAnywhere, BlueprintReadWrite, Category=Appearance, AdvancedDisplay)
-	FLinearColor HighlightColor;
-	FTextBlockStyle& SetHighlightColor(const FLinearColor& InHighlightColor) { HighlightColor = InHighlightColor; return *this; }
+	FSlateColor HighlightColor;
+	FTextBlockStyle& SetHighlightColor(const FSlateColor& InHighlightColor) { HighlightColor = InHighlightColor; return *this; }
 
 	/** The shape of highlighted text */
 	UPROPERTY(EditAnywhere, BlueprintReadWrite, Category=Appearance, AdvancedDisplay)
