CIS Fix for TextLayout

diff --git a/Engine/Source/Runtime/Slate/Public/Framework/Text/TextLayout.h b/Engine/Source/Runtime/Slate/Public/Framework/Text/TextLayout.h
--- a/Engine/Source/Runtime/Slate/Public/Framework/Text/TextLayout.h
+++ b/Engine/Source/Runtime/Slate/Public/Framework/Text/TextLayout.h
@@ -8,6 +8,8 @@
 #include "Framework/Text/TextRunRenderer.h"
 #include "Framework/Text/TextLineHighlight.h"
 #include "Framework/Text/IRun.h"
+#include "Styling/SlateTypes.h"
+
 #include "TextLayout.generated.h"
 
 #define TEXT_LAYOUT_DEBUG 0
