Fix crash docking tabs

diff --git a/Engine/Source/Runtime/Slate/Private/Framework/Docking/SDockingTabStack.cpp b/Engine/Source/Runtime/Slate/Private/Framework/Docking/SDockingTabStack.cpp
--- a/Engine/Source/Runtime/Slate/Private/Framework/Docking/SDockingTabStack.cpp
+++ b/Engine/Source/Runtime/Slate/Private/Framework/Docking/SDockingTabStack.cpp
@@ -270,14 +270,18 @@ void SDockingTabStack::OnTabRemoved( const FTabId& TabId )
 	RemovePersistentTab( TabId );
 }
 
-void SDockingTabStack::OpenTab( const TSharedRef<SDockTab>& InTab, int32 InsertLocationAmongActiveTabs )
+void SDockingTabStack::OpenTab(const TSharedRef<SDockTab>& InTab, int32 InsertLocationAmongActiveTabs)
 {
-	const int32 TabIndex = OpenPersistentTab( InTab->GetLayoutIdentifier(), InsertLocationAmongActiveTabs );
+	const int32 InsertIndex = OpenPersistentTab(InTab->GetLayoutIdentifier(), InsertLocationAmongActiveTabs);
 
 	// The tab may be a nomad tab, in which case it should inherit whichever tab manager it is being put into!
 	InTab->SetTabManager(GetDockArea()->GetTabManager());
 
-	const FTabManager::FTab& TabInfo = Tabs[TabIndex];
+	const FTabId TabId = InTab->GetLayoutIdentifier();
+
+	// the insert index is not the same as the tab index in the array for new tabs so find the tab again to check the tab state.
+	const FTabManager::FTab& TabInfo = *Tabs.FindByPredicate([TabId](const FTabManager::FTab& TestTab) {return TestTab.TabId == TabId; });
+
 	if (TabInfo.TabState == ETabState::SidebarTab)
 	{
 		FSidebarTabLists SidebarLists;
@@ -296,7 +300,7 @@ void SDockingTabStack::OpenTab( const TSharedRef<SDockTab>& InTab, int32 InsertL
 	}
 	else
 	{
-		AddTabWidget(InTab, TabIndex);
+		AddTabWidget(InTab, InsertIndex);
 		OnLiveTabAdded();
 		TabWell->RefreshParentContent();
 	}
