Fix tabs in drawers not showing new name if the tab is renamed

diff --git a/Engine/Source/Runtime/Slate/Private/Framework/Docking/STabSidebar.cpp b/Engine/Source/Runtime/Slate/Private/Framework/Docking/STabSidebar.cpp
--- a/Engine/Source/Runtime/Slate/Private/Framework/Docking/STabSidebar.cpp
+++ b/Engine/Source/Runtime/Slate/Private/Framework/Docking/STabSidebar.cpp
@@ -33,6 +33,9 @@ public:
 
 		DockTabStyle = &FAppStyle::Get().GetWidgetStyle<FDockTabStyle>("Docking.Tab");
 
+		// Sometimes tabs can be renamed so ensure that we pick up the rename
+		ForTab->SetOnTabRenamed(SDockTab::FOnTabRenamed::CreateSP(this, &STabDrawerButton::OnTabRenamed));
+
 		OnDrawerButtonClicked = InArgs._OnDrawerButtonClicked;
 		OnGetContextMenuContent = InArgs._OnGetContextMenuContent;
 		Tab = ForTab;
@@ -127,6 +130,14 @@ public:
 		}
 	}
 
+	void OnTabRenamed(TSharedRef<SDockTab> ForTab)
+	{
+		if (ensure(ForTab == Tab))
+		{
+			Label->SetText(ForTab->GetTabLabel());
+		}
+	}
+
 	virtual FVector2D ComputeDesiredSize(float LayoutScaleMultiplier) const override
 	{
 		const FVector2D Size = FDockingConstants::GetMaxTabSizeFor(ETabRole::PanelTab);
diff --git a/Engine/Source/Runtime/Slate/Private/Widgets/Docking/SDockTab.cpp b/Engine/Source/Runtime/Slate/Private/Widgets/Docking/SDockTab.cpp
--- a/Engine/Source/Runtime/Slate/Private/Widgets/Docking/SDockTab.cpp
+++ b/Engine/Source/Runtime/Slate/Private/Widgets/Docking/SDockTab.cpp
@@ -430,6 +430,11 @@ void SDockTab::SetOnTabActivated( const FOnTabActivatedCallback& InDelegate )
 	OnTabActivated = InDelegate;
 }
 
+void SDockTab::SetOnTabRenamed(const FOnTabRenamed& InDelegate)
+{
+	OnTabRenamed = InDelegate;
+}
+
 TSharedRef<FTabManager> SDockTab::GetTabManager() const
 {
 	return MyTabManager.Pin().ToSharedRef();
@@ -785,6 +790,7 @@ FText SDockTab::GetTabLabel() const
 void SDockTab::SetLabel( const TAttribute<FText>& InTabLabel )
 {
 	TabLabel = InTabLabel;
+	OnTabRenamed.ExecuteIfBound(SharedThis(this));
 }
 
 FText SDockTab::GetTabLabelSuffix() const
diff --git a/Engine/Source/Runtime/Slate/Public/Widgets/Docking/SDockTab.h b/Engine/Source/Runtime/Slate/Public/Widgets/Docking/SDockTab.h
--- a/Engine/Source/Runtime/Slate/Public/Widgets/Docking/SDockTab.h
+++ b/Engine/Source/Runtime/Slate/Public/Widgets/Docking/SDockTab.h
@@ -56,7 +56,10 @@ public:
 	/** Invoked when a tab is activated */
 	DECLARE_DELEGATE_TwoParams(FOnTabActivatedCallback, TSharedRef<SDockTab>, ETabActivationCause);
 
-	/** Invoked when this tab should save some information about its content. */
+	/** Invoked when a tab is renamed */
+	DECLARE_DELEGATE_OneParam(FOnTabRenamed, TSharedRef<SDockTab>);
+
+	/** Invoked w`en this tab should save some information about its content. */
 	DECLARE_DELEGATE(FOnPersistVisualState);
 
 	/** Delegate called before a tab is closed.  Returning false will prevent the tab from closing */
@@ -245,6 +248,9 @@ public:
 	/** Set the handler that will be invoked when the tab is activated */
 	void SetOnTabActivated( const FOnTabActivatedCallback& InDelegate );
 
+	/** Set the handler that will be invoked when the tab is renamed */
+	void SetOnTabRenamed(const FOnTabRenamed& InDelegate);
+
 	/** Get the tab manager currently managing this tab. Note that a user move the tab between Tab Managers, so this return value may change. */
 	TSharedRef<FTabManager> GetTabManager() const;
 
@@ -387,6 +393,9 @@ protected:
 	 */
 	FOnPersistVisualState OnPersistVisualState;
 
+	/** Invoked when the tab is renamed */
+	FOnTabRenamed OnTabRenamed;
+
 	/** The styles used to draw the tab in its various states */
 	const FDockTabStyle* MajorTabStyle;
 	const FDockTabStyle* GenericTabStyle;
