Spin box values no longer get separated from their property values in detail views.

# This was a result of accidental removal of the check for value attribute being bound in the CommitValue function in SSpinBox from CL 13906773

#rnx

diff --git a/Engine/Source/Runtime/Slate/Public/Widgets/Input/SSpinBox.h b/Engine/Source/Runtime/Slate/Public/Widgets/Input/SSpinBox.h
--- a/Engine/Source/Runtime/Slate/Public/Widgets/Input/SSpinBox.h
+++ b/Engine/Source/Runtime/Slate/Public/Widgets/Input/SSpinBox.h
@@ -910,6 +910,10 @@ protected:
 			if( CurrentDelta != 0 )
 			{
 				ValueToCommit = FMath::GridSnap((double)ValueToCommit, (double)CurrentDelta); // snap numeric point value to nearest Delta
+				if (CachedExternalValue != ValueToCommit)
+				{
+					IntermediateDragFractionalValue = 0.0;
+				}
 			}
 		}		
 
@@ -930,7 +934,10 @@ protected:
 
 		OnValueChanged.ExecuteIfBound(ValueToCommit);
 
-		ValueAttribute.Set(ValueToCommit);
+		if (!ValueAttribute.IsBound())
+		{
+			ValueAttribute.Set(ValueToCommit);
+		}
 
 		// Update the cache of the external value to what the user believes the value is now.
 		CachedExternalValue = ValueAttribute.Get();
@@ -1042,7 +1049,6 @@ private:
 		else
 		{
 			NumericType TmpValue = CurrentValue + IntermediateDragFractionalValue;
-			IntermediateDragFractionalValue = 0.0;
 			return TmpValue;
 		}
 	}
