Fix implicit type conversion warnings in Slate

diff --git a/Engine/Source/Runtime/Slate/Private/Framework/Application/MenuStack.cpp b/Engine/Source/Runtime/Slate/Private/Framework/Application/MenuStack.cpp
--- a/Engine/Source/Runtime/Slate/Private/Framework/Application/MenuStack.cpp
+++ b/Engine/Source/Runtime/Slate/Private/Framework/Application/MenuStack.cpp
@@ -120,7 +120,7 @@ namespace MenuStackInternal
 					+ SOverlay::Slot()
 					[
 						SNew(SBorder)
-						.Padding(0)
+						.Padding(0.f)
 						.BorderImage(FStyleDefaults::GetNoBrush())
 						.ForegroundColor(FCoreStyle::Get().GetSlateColor("DefaultForeground"))
 						[
diff --git a/Engine/Source/Runtime/Slate/Private/Framework/Application/SlateApplication.cpp b/Engine/Source/Runtime/Slate/Private/Framework/Application/SlateApplication.cpp
--- a/Engine/Source/Runtime/Slate/Private/Framework/Application/SlateApplication.cpp
+++ b/Engine/Source/Runtime/Slate/Private/Framework/Application/SlateApplication.cpp
@@ -1597,7 +1597,7 @@ void FSlateApplication::ThrottleApplicationBasedOnMouseMovement()
 
 		// After a key or mouse button is pressed, we'll leave the throttle disengaged for awhile so the
 		// user can use the keys to navigate in a viewport, for example.
-		const float MinTimeSinceButtonPressToThrottle = 1.0f;
+		const double MinTimeSinceButtonPressToThrottle = 1.0;
 
 		// Use a small movement threshold to avoid engaging the throttle when the user bumps the mouse
 		const float MinMouseMovePixelsBeforeThrottle = 2.0f;
@@ -1617,8 +1617,8 @@ void FSlateApplication::ThrottleApplicationBasedOnMouseMovement()
 			}
 		}
 
-		const float TimeSinceLastUserInteraction = CurrentTime - LastUserInteractionTimeForThrottling;
-		const float TimeSinceLastMouseMove = CurrentTime - LastMouseMoveTime;
+		const double TimeSinceLastUserInteraction = CurrentTime - LastUserInteractionTimeForThrottling;
+		const double TimeSinceLastMouseMove = CurrentTime - LastMouseMoveTime;
 		if( TimeSinceLastMouseMove < TimeToThrottleAfterMouseStops )
 		{
 			// Only throttle if a Slate window is currently active.  If a Wx window (such as Matinee) is
@@ -1942,15 +1942,15 @@ void FSlateApplication::AddModalWindow( TSharedRef<SWindow> InSlateWindow, const
 		Renderer->EndFrame();
 
 		//Throttle loop data
-		float LastLoopTime = (float)FPlatformTime::Seconds();
-		const float MinThrottlePeriod = (1.0f / 60.0f); //Throttle the loop to a maximum of 60Hz
+		double LastLoopTime = FPlatformTime::Seconds();
+		const double MinThrottlePeriod = (1.0 / 60.0); //Throttle the loop to a maximum of 60Hz
 
 		// Tick slate from here in the event that we should not return until the modal window is closed.
 		while( InSlateWindow == GetActiveModalWindow() )
 		{
 			//Throttle the loop
-			const float CurrentLoopTime = FPlatformTime::Seconds();
-			const float SleepTime = MinThrottlePeriod - (CurrentLoopTime-LastLoopTime);
+			const double CurrentLoopTime = FPlatformTime::Seconds();
+			const float SleepTime = static_cast<float>(MinThrottlePeriod - (CurrentLoopTime-LastLoopTime));
 			LastLoopTime = CurrentLoopTime;
 			if (SleepTime > 0.0f)
 			{
diff --git a/Engine/Source/Runtime/Slate/Private/Framework/MultiBox/MultiBox.cpp b/Engine/Source/Runtime/Slate/Private/Framework/MultiBox/MultiBox.cpp
--- a/Engine/Source/Runtime/Slate/Private/Framework/MultiBox/MultiBox.cpp
+++ b/Engine/Source/Runtime/Slate/Private/Framework/MultiBox/MultiBox.cpp
@@ -795,8 +795,8 @@ void SMultiBoxWidget::BuildMultiBoxWidget()
 
 	// Create a box panel that the various multiblocks will resides within
 	// @todo Slate MultiBox: Expose margins and other useful bits
-	TSharedPtr< SVerticalBox > VerticalBox;
-	TSharedPtr< SWidget > MainWidget;
+	TSharedPtr<SVerticalBox> VerticalBox;
+	TSharedPtr<SWidget> MainWidget;
 	TSharedPtr<SHorizontalBox> HorizontalBox;
 
 	/** The current row of buttons for if the multibox type is a button row */
@@ -834,11 +834,11 @@ void SMultiBoxWidget::BuildMultiBoxWidget()
 			[
 				SAssignNew(UniformToolbarPanel, SUniformWrapPanel)
 				.HAlign(HAlign_Left)
-				.MinDesiredSlotWidth(68)
-				.MinDesiredSlotHeight(52)
-				.MaxDesiredSlotWidth(68)
-				.MaxDesiredSlotHeight(52)
-				.SlotPadding( 0.0 )
+				.MinDesiredSlotWidth(68.f)
+				.MinDesiredSlotHeight(52.f)
+				.MaxDesiredSlotWidth(68.f)
+				.MaxDesiredSlotHeight(52.f)
+				.SlotPadding(0.f)
 			];
 		}
 		break;
diff --git a/Engine/Source/Runtime/Slate/Private/Framework/MultiBox/SUniformToolbarPanel.cpp b/Engine/Source/Runtime/Slate/Private/Framework/MultiBox/SUniformToolbarPanel.cpp
--- a/Engine/Source/Runtime/Slate/Private/Framework/MultiBox/SUniformToolbarPanel.cpp
+++ b/Engine/Source/Runtime/Slate/Private/Framework/MultiBox/SUniformToolbarPanel.cpp
@@ -42,7 +42,7 @@ void SUniformToolbarPanel::Construct( const FArguments& InArgs )
 		SAssignNew(Dropdown, SComboButton)
 		.HasDownArrow(false)
 		.ButtonStyle(&ToolBarStyle.ButtonStyle)
-		.ContentPadding(0)
+		.ContentPadding(0.f)
 		.ToolTipText(NSLOCTEXT("Slate", "ExpandToolbar", "Click to expand toolbar"))
 		.OnGetMenuContent(InArgs._OnDropdownOpened)
 		.Cursor(EMouseCursor::Default)
@@ -120,8 +120,8 @@ void SUniformToolbarPanel::OnArrangeChildren( const FGeometry& AllottedGeometry,
 							FVector2D(XAxisResult.Size, YAxisResult.Size)
 						);
 
-					WidgetExtents = FMath::TruncToFloat(ArrangedChild.Geometry.AbsolutePosition.X + ArrangedChild.Geometry.GetLocalSize().X * ArrangedChild.Geometry.Scale);
-					AllottedGeometryExtents = FMath::TruncToFloat((AllottedGeometry.AbsolutePosition.X + AllottedGeometry.GetLocalSize().X * AllottedGeometry.Scale) - ClippedChildrenDropDownDesiredSize.X);
+					WidgetExtents = FMath::TruncToInt(ArrangedChild.Geometry.AbsolutePosition.X + ArrangedChild.Geometry.GetLocalSize().X * ArrangedChild.Geometry.Scale);
+					AllottedGeometryExtents = FMath::TruncToInt((AllottedGeometry.AbsolutePosition.X + AllottedGeometry.GetLocalSize().X * AllottedGeometry.Scale) - ClippedChildrenDropDownDesiredSize.X);
 
 				}
 				else
@@ -132,8 +132,8 @@ void SUniformToolbarPanel::OnArrangeChildren( const FGeometry& AllottedGeometry,
 							FVector2D(XAxisResult.Size, YAxisResult.Size)
 						);
 
-					WidgetExtents = FMath::TruncToFloat(ArrangedChild.Geometry.AbsolutePosition.Y + ArrangedChild.Geometry.GetLocalSize().Y * ArrangedChild.Geometry.Scale);
-					AllottedGeometryExtents = FMath::TruncToFloat((AllottedGeometry.AbsolutePosition.Y + AllottedGeometry.GetLocalSize().Y * AllottedGeometry.Scale) - ClippedChildrenDropDownDesiredSize.Y);
+					WidgetExtents = FMath::TruncToInt(ArrangedChild.Geometry.AbsolutePosition.Y + ArrangedChild.Geometry.GetLocalSize().Y * ArrangedChild.Geometry.Scale);
+					AllottedGeometryExtents = FMath::TruncToInt((AllottedGeometry.AbsolutePosition.Y + AllottedGeometry.GetLocalSize().Y * AllottedGeometry.Scale) - ClippedChildrenDropDownDesiredSize.Y);
 
 				}
 
diff --git a/Engine/Source/Runtime/Slate/Private/Widgets/Input/SComboButton.cpp b/Engine/Source/Runtime/Slate/Private/Widgets/Input/SComboButton.cpp
--- a/Engine/Source/Runtime/Slate/Private/Widgets/Input/SComboButton.cpp
+++ b/Engine/Source/Runtime/Slate/Private/Widgets/Input/SComboButton.cpp
@@ -61,7 +61,7 @@ void SComboButton::Construct( const FArguments& InArgs )
 				.AutoWidth()
 				.HAlign( HAlign_Center )
 				.VAlign( VAlign_Center )
-				.Padding( InArgs._HasDownArrow ? 2 : 0 )
+				.Padding( InArgs._HasDownArrow ? 2.f : 0.f )
 				[
 					SNew(SOverlay)
 					// drop shadow
diff --git a/Engine/Source/Runtime/Slate/Private/Widgets/Input/SEditableTextBox.cpp b/Engine/Source/Runtime/Slate/Private/Widgets/Input/SEditableTextBox.cpp
--- a/Engine/Source/Runtime/Slate/Private/Widgets/Input/SEditableTextBox.cpp
+++ b/Engine/Source/Runtime/Slate/Private/Widgets/Input/SEditableTextBox.cpp
@@ -40,7 +40,7 @@ void SEditableTextBox::Construct( const FArguments& InArgs )
 		.BorderImage( this, &SEditableTextBox::GetBorderImage )
 		.BorderBackgroundColor( this, &SEditableTextBox::DetermineBackgroundColor )
 		.ForegroundColor( this, &SEditableTextBox::DetermineForegroundColor )
-		.Padding( 0 )
+		.Padding(0.f)
 		[
 			SAssignNew( Box, SHorizontalBox)
 
diff --git a/Engine/Source/Runtime/Slate/Private/Widgets/Layout/SScrollBox.cpp b/Engine/Source/Runtime/Slate/Private/Widgets/Layout/SScrollBox.cpp
--- a/Engine/Source/Runtime/Slate/Private/Widgets/Layout/SScrollBox.cpp
+++ b/Engine/Source/Runtime/Slate/Private/Widgets/Layout/SScrollBox.cpp
@@ -350,7 +350,7 @@ float SScrollBox::GetViewOffsetFraction() const
 	const float ContentSize = GetScrollComponentFromVector(ScrollPanel->GetDesiredSize());
 
 	const float ViewFraction = GetViewFraction();
-	return FMath::Clamp<float>( DesiredScrollOffset/ContentSize, 0.0, 1.0 - ViewFraction );
+	return FMath::Clamp( DesiredScrollOffset/ContentSize, 0.f, 1.f - ViewFraction );
 }
 
 void SScrollBox::SetScrollOffset( float NewScrollOffset )
diff --git a/Engine/Source/Runtime/Slate/Private/Widgets/Layout/SSplitter.cpp b/Engine/Source/Runtime/Slate/Private/Widgets/Layout/SSplitter.cpp
--- a/Engine/Source/Runtime/Slate/Private/Widgets/Layout/SSplitter.cpp
+++ b/Engine/Source/Runtime/Slate/Private/Widgets/Layout/SSplitter.cpp
@@ -181,7 +181,7 @@ int32 SSplitter::OnPaint( const FPaintArgs& Args, const FGeometry& AllottedGeome
 	{
 		const FGeometry& GeometryAfterSplitter = ArrangedChildren[ FMath::Clamp(ChildIndex + 1, 0, ArrangedChildren.Num()-1) ].Geometry;
 
-		const float HalfHitDetectionSplitterHandleSize = FMath::RoundToInt(HitDetectionSplitterHandleSize / 2.0f);
+		const float HalfHitDetectionSplitterHandleSize = FMath::RoundToFloat(HitDetectionSplitterHandleSize / 2.0f);
 		const float HalfPhysicalSplitterHandleSize = PhysicalSplitterHandleSize;// FMath::RoundToInt(PhysicalSplitterHandleSize / 2.0f);
 
 		FVector2D HandleSize;		
diff --git a/Engine/Source/Runtime/Slate/Private/Widgets/Layout/SWindowTitleBarArea.cpp b/Engine/Source/Runtime/Slate/Private/Widgets/Layout/SWindowTitleBarArea.cpp
--- a/Engine/Source/Runtime/Slate/Private/Widgets/Layout/SWindowTitleBarArea.cpp
+++ b/Engine/Source/Runtime/Slate/Private/Widgets/Layout/SWindowTitleBarArea.cpp
@@ -72,7 +72,7 @@ void SWindowTitleBarArea::Construct( const FArguments& InArgs )
 	MinimizeButton = SNew(SButton)
 		.IsFocusable(false)
 		.IsEnabled(bIsMinimizeButtonEnabled)
-		.ContentPadding(0)
+		.ContentPadding(0.f)
 		.OnClicked(this, &SWindowTitleBarArea::MinimizeButton_OnClicked)
 		.Cursor(EMouseCursor::Default)
 		.ButtonStyle(FCoreStyle::Get(), "NoBorder")
@@ -85,7 +85,7 @@ void SWindowTitleBarArea::Construct( const FArguments& InArgs )
 	MaximizeRestoreButton = SNew(SButton)
 		.IsFocusable(false)
 		.IsEnabled(bIsMaximizeRestoreButtonEnabled)
-		.ContentPadding(0.0f)
+		.ContentPadding(0.f)
 		.OnClicked(this, &SWindowTitleBarArea::MaximizeRestoreButton_OnClicked)
 		.Cursor(EMouseCursor::Default)
 		.ButtonStyle(FCoreStyle::Get(), "NoBorder")
diff --git a/Engine/Source/Runtime/Slate/Private/Widgets/Notifications/SNotificationList.cpp b/Engine/Source/Runtime/Slate/Private/Widgets/Notifications/SNotificationList.cpp
--- a/Engine/Source/Runtime/Slate/Private/Widgets/Notifications/SNotificationList.cpp
+++ b/Engine/Source/Runtime/Slate/Private/Widgets/Notifications/SNotificationList.cpp
@@ -653,7 +653,7 @@ public:
 		ChildSlot
 			[
 				SNew(SBorder)
-				.Padding(0)
+				.Padding(0.f)
 				.BorderImage(FCoreStyle::Get().GetBrush("NoBorder"))
 				.BorderBackgroundColor(this, &SNotificationItemExternalImpl::GetContentColor)
 				.ColorAndOpacity(this, &SNotificationItemExternalImpl::GetContentColorRaw)
diff --git a/Engine/Source/Runtime/Slate/Private/Widgets/Notifications/SProgressBar.cpp b/Engine/Source/Runtime/Slate/Private/Widgets/Notifications/SProgressBar.cpp
--- a/Engine/Source/Runtime/Slate/Private/Widgets/Notifications/SProgressBar.cpp
+++ b/Engine/Source/Runtime/Slate/Private/Widgets/Notifications/SProgressBar.cpp
@@ -221,7 +221,7 @@ int32 SProgressBar::OnPaint( const FPaintArgs& Args, const FGeometry& AllottedGe
 						OutDrawElements,
 						RetLayerId++,
 						AllottedGeometry.ToPaintGeometry(
-							FVector2D( (AllottedGeometry.GetLocalSize().X * 0.5f) - ((AllottedGeometry.GetLocalSize().X * ( ClampedFraction ))*0.5), 0.0f),
+							FVector2D( (AllottedGeometry.GetLocalSize().X * 0.5f) - ((AllottedGeometry.GetLocalSize().X * ( ClampedFraction ))*0.5f), 0.0f),
 							FVector2D( AllottedGeometry.GetLocalSize().X * ( ClampedFraction ) , AllottedGeometry.GetLocalSize().Y )),
 						CurrentFillImage,
 						DrawEffects,
@@ -368,7 +368,7 @@ void SProgressBar::UpdateMarqueeActiveTimer()
 
 EActiveTimerReturnType SProgressBar::ActiveTick(double InCurrentTime, float InDeltaTime)
 {
-	MarqueeOffset = InCurrentTime - FMath::FloorToDouble(InCurrentTime);
+	MarqueeOffset = (float)(InCurrentTime - FMath::FloorToDouble(InCurrentTime));
 	
 	TOptional<float> PercentFraction = Percent.Get();
 	if (PercentFraction.IsSet())
diff --git a/Engine/Source/Runtime/Slate/Private/Widgets/Views/SListPanel.cpp b/Engine/Source/Runtime/Slate/Private/Widgets/Views/SListPanel.cpp
--- a/Engine/Source/Runtime/Slate/Private/Widgets/Views/SListPanel.cpp
+++ b/Engine/Source/Runtime/Slate/Private/Widgets/Views/SListPanel.cpp
@@ -55,7 +55,7 @@ void SListPanel::OnArrangeChildren( const FGeometry& AllottedGeometry, FArranged
 			// This is a tile view list - arrange items side by side along the line axis until there is no more room, then create a new line along the scroll axis.
 			const EListItemAlignment ListItemAlignment = ItemAlignment.Get();
 			const float ItemPadding = GetItemPadding(AllottedGeometry, ListItemAlignment);
-			const float HalfItemPadding = ItemPadding * 0.5;
+			const float HalfItemPadding = ItemPadding * 0.5f;
 
 			const FTableViewDimensions LocalItemSize = GetItemSize(AllottedGeometry, ListItemAlignment);
 			DimensionsSoFar.ScrollAxis = -FMath::FloorToInt(FirstLineScrollOffset * LocalItemSize.ScrollAxis) - OverscrollAmount;
@@ -75,7 +75,7 @@ void SListPanel::OnArrangeChildren( const FGeometry& AllottedGeometry, FArranged
 						}
 						else
 						{
-							const float HalfLinePadding = LinePadding * 0.5;
+							const float HalfLinePadding = LinePadding * 0.5f;
 							DimensionsSoFar.LineAxis += IsVisible ? HalfLinePadding : 0;
 						}
 					}
diff --git a/Engine/Source/Runtime/Slate/Public/Framework/Application/SWindowTitleBar.h b/Engine/Source/Runtime/Slate/Public/Framework/Application/SWindowTitleBar.h
--- a/Engine/Source/Runtime/Slate/Public/Framework/Application/SWindowTitleBar.h
+++ b/Engine/Source/Runtime/Slate/Public/Framework/Application/SWindowTitleBar.h
@@ -266,7 +266,7 @@ protected:
 			MinimizeButton = SNew(SButton)
 					.IsFocusable(false)
 					.IsEnabled(OwnerWindow->HasMinimizeBox())
-					.ContentPadding(0)
+					.ContentPadding(0.f)
 					.OnClicked(this, &SWindowTitleBar::MinimizeButton_OnClicked)
 					.Cursor(EMouseCursor::Default)
 					.ButtonStyle(FCoreStyle::Get(), "NoBorder")
diff --git a/Engine/Source/Runtime/Slate/Public/Widgets/Layout/SLinkedBox.h b/Engine/Source/Runtime/Slate/Public/Widgets/Layout/SLinkedBox.h
--- a/Engine/Source/Runtime/Slate/Public/Widgets/Layout/SLinkedBox.h
+++ b/Engine/Source/Runtime/Slate/Public/Widgets/Layout/SLinkedBox.h
@@ -43,7 +43,7 @@ public:
 
 protected:
 	
-	mutable uint32 FrameCounterLastCached = 0;
+	mutable uint64 FrameCounterLastCached = 0;
 	mutable FVector2D CachedUniformSize;
 
 private:
diff --git a/Engine/Source/Runtime/SlateCore/Public/Layout/Geometry.h b/Engine/Source/Runtime/SlateCore/Public/Layout/Geometry.h
--- a/Engine/Source/Runtime/SlateCore/Public/Layout/Geometry.h
+++ b/Engine/Source/Runtime/SlateCore/Public/Layout/Geometry.h
@@ -427,13 +427,13 @@ public:
 	/**
 	 * Translates the local coordinates into local coordinates that after being transformed into absolute space will be rounded
 	 * to a whole number or approximately a whole number.  This is important for cases where you want to show a popup or a tooltip
-	 * and not have the window start on a half pixel, which can cause the contents to jitter in relation to eachother as the tooltip 
+	 * and not have the window start on a half pixel, which can cause the contents to jitter in relation to each other as the tooltip 
 	 * or popup moves around.
 	 */
 	FORCEINLINE_DEBUGGABLE FVector2D LocalToRoundedLocal(FVector2D LocalCoordinate) const
 	{
 		const FVector2D AbsoluteCoordinate = LocalToAbsolute(LocalCoordinate);
-		const FVector2D AbsoluteCoordinateRounded = FVector2D(FMath::RoundToInt(AbsoluteCoordinate.X), FMath::RoundToInt(AbsoluteCoordinate.Y));
+		const FVector2D AbsoluteCoordinateRounded = FVector2D(FMath::RoundToFloat(AbsoluteCoordinate.X), FMath::RoundToFloat(AbsoluteCoordinate.Y));
 
 		return AbsoluteToLocal(AbsoluteCoordinateRounded);
 	}
diff --git a/Engine/Source/Runtime/SlateCore/Public/Rendering/RenderingCommon.h b/Engine/Source/Runtime/SlateCore/Public/Rendering/RenderingCommon.h
--- a/Engine/Source/Runtime/SlateCore/Public/Rendering/RenderingCommon.h
+++ b/Engine/Source/Runtime/SlateCore/Public/Rendering/RenderingCommon.h
@@ -63,6 +63,8 @@ enum class ESlateShader : uint8
 	Custom = 5,
 	/** For post processing passes */
 	PostProcess = 6,
+	/** Rounded Box shader. **/
+	RoundedBox = 7,
 };
 
 /**
@@ -465,11 +467,6 @@ typedef TArray<FSlateDrawElement> FSlateDrawElementArray;
 
 #endif
 
-static FVector2D RoundToInt(const FVector2D& Vec)
-{
-	return FVector2D(FMath::RoundToInt(Vec.X), FMath::RoundToInt(Vec.Y));
-}
-
 /**
  * Viewport implementation interface that is used by SViewport when it needs to draw and processes input.                   
  */
diff --git a/Engine/Source/Runtime/SlateCore/Public/Textures/SlateShaderResource.h b/Engine/Source/Runtime/SlateCore/Public/Textures/SlateShaderResource.h
--- a/Engine/Source/Runtime/SlateCore/Public/Textures/SlateShaderResource.h
+++ b/Engine/Source/Runtime/SlateCore/Public/Textures/SlateShaderResource.h
@@ -122,7 +122,7 @@ public:
 		: StartUV(0.0f, 0.0f)
 		, SizeUV(1.0f, 1.0f)
 		, Resource(nullptr)
-		, ActualSize(0.0f, 0.0f)
+		, ActualSize(0, 0)
 	{ }
 
 	~FSlateShaderResourceProxy()
