Further flattening of the editor by replacing background colors

diff --git a/Engine/Source/Runtime/Slate/Private/Framework/Docking/SDockingArea.cpp b/Engine/Source/Runtime/Slate/Private/Framework/Docking/SDockingArea.cpp
--- a/Engine/Source/Runtime/Slate/Private/Framework/Docking/SDockingArea.cpp
+++ b/Engine/Source/Runtime/Slate/Private/Framework/Docking/SDockingArea.cpp
@@ -36,6 +36,7 @@ void SDockingArea::Construct( const FArguments& InArgs, const TSharedRef<FTabMan
 			[
 				SAssignNew(LeftSidebar, STabSidebar)
 				.Location(ESidebarLocation::Left)
+				.Visibility(EVisibility::Collapsed)
 			]
 			+ SHorizontalBox::Slot()
 			[
@@ -47,6 +48,7 @@ void SDockingArea::Construct( const FArguments& InArgs, const TSharedRef<FTabMan
 			[
 				SAssignNew(RightSidebar, STabSidebar)
 				.Location(ESidebarLocation::Right)
+				.Visibility(EVisibility::Collapsed)
 			]
 		]
 		
@@ -618,12 +620,10 @@ void SDockingArea::UpdateWindowChromeAndSidebar()
 
 			bCanHaveSidebar = !bContainsMajorTabs;
 
-			if (!bCanHaveSidebar)
-			{
-				LeftSidebar->SetVisibility(EVisibility::Collapsed);
-				RightSidebar->SetVisibility(EVisibility::Collapsed);
-			}
-			else
+			LeftSidebar->SetVisibility(EVisibility::Collapsed);
+			RightSidebar->SetVisibility(EVisibility::Collapsed);
+
+			if (bCanHaveSidebar)
 			{
 				if (bAccountForMenuBarPadding)
 				{
diff --git a/Engine/Source/Runtime/Slate/Private/Framework/Docking/SDockingTabStack.cpp b/Engine/Source/Runtime/Slate/Private/Framework/Docking/SDockingTabStack.cpp
--- a/Engine/Source/Runtime/Slate/Private/Framework/Docking/SDockingTabStack.cpp
+++ b/Engine/Source/Runtime/Slate/Private/Framework/Docking/SDockingTabStack.cpp
@@ -1129,8 +1129,9 @@ FSlateColor SDockingTabStack::GetUnhideTabWellButtonOpacity() const
 
 const FSlateBrush* SDockingTabStack::GetTabStackBorderImage() const
 {
-	static const FSlateColorBrush MajorTabBackgroundBrush(FAppStyle::Get().GetSlateColor("Colors.Title"));
-	return bShowingTitleBarArea ? &MajorTabBackgroundBrush : FStyleDefaults::GetNoBrush();
+	static const FSlateBrush* MajorTabBackgroundBrush = FAppStyle::Get().GetBrush("Brushes.Title");
+	static const FSlateBrush* MinorTabBackgroundBrush = FAppStyle::Get().GetBrush("Brushes.Recessed");
+	return bShowingTitleBarArea ? MajorTabBackgroundBrush : MinorTabBackgroundBrush;
 }
 
 int32 SDockingTabStack::OpenPersistentTab( const FTabId& TabId, int32 OpenLocationAmongActiveTabs )
diff --git a/Engine/Source/Runtime/SlateCore/Private/Styling/StarshipCoreStyle.cpp b/Engine/Source/Runtime/SlateCore/Private/Styling/StarshipCoreStyle.cpp
--- a/Engine/Source/Runtime/SlateCore/Private/Styling/StarshipCoreStyle.cpp
+++ b/Engine/Source/Runtime/SlateCore/Private/Styling/StarshipCoreStyle.cpp
@@ -1428,7 +1428,7 @@ void FStarshipCoreStyle::SetupDockingStyles(TSharedRef<FStyle>& Style)
 
 		.SetColorOverlayTabBrush(FSlateNoResource())
 		.SetColorOverlayIconBrush(FSlateNoResource())
-		.SetContentAreaBrush(FSlateNoResource())
+		.SetContentAreaBrush(FSlateColorBrush(FStyleColors::Background))
 		.SetTabWellBrush(FSlateNoResource())
 		.SetFlashColor(TabFlashColor)
 
@@ -1453,7 +1453,7 @@ void FStarshipCoreStyle::SetupDockingStyles(TSharedRef<FStyle>& Style)
 
 		.SetColorOverlayTabBrush(FSlateNoResource())
 		.SetColorOverlayIconBrush(FSlateNoResource())
-		.SetContentAreaBrush(FSlateNoResource())
+		.SetContentAreaBrush(FSlateColorBrush(FStyleColors::Background))
 		.SetTabWellBrush(FSlateNoResource())
 
 		.SetTabPadding(FMargin(10, 7, 10, 8))
