Added an optional separator to slate header rows

diff --git a/Engine/Source/Runtime/Slate/Private/Widgets/Views/SHeaderRow.cpp b/Engine/Source/Runtime/Slate/Private/Widgets/Views/SHeaderRow.cpp
--- a/Engine/Source/Runtime/Slate/Private/Widgets/Views/SHeaderRow.cpp
+++ b/Engine/Source/Runtime/Slate/Private/Widgets/Views/SHeaderRow.cpp
@@ -16,6 +16,7 @@
 #include "Widgets/Layout/SScrollBar.h"
 #include "Widgets/Input/SButton.h"
 #include "Widgets/Input/SComboButton.h"
+#include "Widgets/Layout/SSeparator.h"
 
 #define LOCTEXT_NAMESPACE "SHeaderRow"
 
@@ -660,7 +661,7 @@ void SHeaderRow::RegenerateWidgets()
 
 	TSharedPtr<SSplitter> Splitter;
 
-	TSharedRef< SHorizontalBox > Box = 
+	TSharedRef< SHorizontalBox > HeaderContent = 
 		SNew(SHorizontalBox)
 		+SHorizontalBox::Slot()
 		.FillWidth( 1.0f )
@@ -856,8 +857,27 @@ void SHeaderRow::RegenerateWidgets()
 		}
 	}
 
-	// Create a box to contain widgets for each column
-	SetContent( Box );
+	if(Style->HorizontalSeparatorBrush.GetDrawType() != ESlateBrushDrawType::NoDrawType && Style->HorizontalSeparatorThickness > 0)
+	{
+		// Create a box to contain widgets for each column
+		SetContent(
+			SNew(SVerticalBox)
+			+ SVerticalBox::Slot()
+			[
+				HeaderContent
+			]
+			+ SVerticalBox::Slot()
+			.AutoHeight()
+			[
+				SNew(SSeparator)
+				.Thickness(Style->HorizontalSeparatorThickness)
+				.SeparatorImage(&Style->HorizontalSeparatorBrush)
+			]);
+	}
+	else
+	{
+		SetContent(HeaderContent);
+	}
 }
 
 void SHeaderRow::OnGenerateSelectColumnsSubMenu(FMenuBuilder& InSubMenuBuilder)
diff --git a/Engine/Source/Runtime/SlateCore/Private/Styling/SlateTypes.cpp b/Engine/Source/Runtime/SlateCore/Private/Styling/SlateTypes.cpp
--- a/Engine/Source/Runtime/SlateCore/Private/Styling/SlateTypes.cpp
+++ b/Engine/Source/Runtime/SlateCore/Private/Styling/SlateTypes.cpp
@@ -577,6 +577,8 @@ const FTableColumnHeaderStyle& FTableColumnHeaderStyle::GetDefault()
 
 
 FHeaderRowStyle::FHeaderRowStyle()
+	: HorizontalSeparatorBrush(FSlateNoResource())
+	, HorizontalSeparatorThickness(0)
 {
 }
 
@@ -585,7 +587,8 @@ void FHeaderRowStyle::GetResources( TArray< const FSlateBrush* >& OutBrushes ) c
 	ColumnStyle.GetResources(OutBrushes);
 	LastColumnStyle.GetResources(OutBrushes);
 	ColumnSplitterStyle.GetResources(OutBrushes);
-	OutBrushes.Add( &BackgroundBrush );
+	OutBrushes.Add(&BackgroundBrush);
+	OutBrushes.Add(&HorizontalSeparatorBrush);
 }
 
 const FName FHeaderRowStyle::TypeName( TEXT("FHeaderRowStyle") );
diff --git a/Engine/Source/Runtime/SlateCore/Private/Styling/StarshipCoreStyle.cpp b/Engine/Source/Runtime/SlateCore/Private/Styling/StarshipCoreStyle.cpp
--- a/Engine/Source/Runtime/SlateCore/Private/Styling/StarshipCoreStyle.cpp
+++ b/Engine/Source/Runtime/SlateCore/Private/Styling/StarshipCoreStyle.cpp
@@ -1644,6 +1644,8 @@ void FStarshipCoreStyle::SetupTableViewStyles(TSharedRef<FStyle>& Style)
 		.SetColumnSplitterStyle(TableHeaderSplitterStyle)
 		.SetBackgroundBrush(FSlateColorBrush(FStyleColors::Background))
 		.SetForegroundColor(FStyleColors::Foreground)
+		.SetHorizontalSeparatorBrush(FSlateColorBrush(FStyleColors::Recessed))
+		.SetHorizontalSeparatorThickness(2.0f)
 	);
 }
 
diff --git a/Engine/Source/Runtime/SlateCore/Public/Styling/SlateTypes.h b/Engine/Source/Runtime/SlateCore/Public/Styling/SlateTypes.h
--- a/Engine/Source/Runtime/SlateCore/Public/Styling/SlateTypes.h
+++ b/Engine/Source/Runtime/SlateCore/Public/Styling/SlateTypes.h
@@ -1622,6 +1622,17 @@ struct SLATECORE_API FHeaderRowStyle : public FSlateWidgetStyle
 	UPROPERTY(EditAnywhere, Category=Appearance)
 	FSlateColor ForegroundColor;
 	FHeaderRowStyle& SetForegroundColor( const FSlateColor& InForegroundColor ){ ForegroundColor = InForegroundColor; return *this; }
+
+	/** Brush used to draw the splitter between the header and the contents below it*/
+	UPROPERTY(EditAnywhere, Category = Appearance)
+	FSlateBrush HorizontalSeparatorBrush;
+	FHeaderRowStyle& SetHorizontalSeparatorBrush(const FSlateBrush& InHorizontalSeparatorBrush) { HorizontalSeparatorBrush = InHorizontalSeparatorBrush; return *this; }
+
+	UPROPERTY(EditAnywhere, Category=Appearance)
+	float HorizontalSeparatorThickness;
+	FHeaderRowStyle& SetHorizontalSeparatorThickness(const float InHorizontalSeparatorThickness) { HorizontalSeparatorThickness = InHorizontalSeparatorThickness; return *this; }
+
+
 };
 
 
