Widget: Modify SpinBox to support double and int64

diff --git a/Engine/Source/Runtime/Slate/Public/Widgets/Input/SSpinBox.h b/Engine/Source/Runtime/Slate/Public/Widgets/Input/SSpinBox.h
--- a/Engine/Source/Runtime/Slate/Public/Widgets/Input/SSpinBox.h
+++ b/Engine/Source/Runtime/Slate/Public/Widgets/Input/SSpinBox.h
@@ -26,6 +26,8 @@
 #include "Widgets/SCompoundWidget.h"
 #include "Widgets/Text/STextBlock.h"
 
+#include <limits>
+
 /*
  * This function compute a slider position by simulating two log on both side of the neutral value
  * Example a slider going from 0.0 to 2.0 with a neutral value of 1.0, the user will have a lot of precision around the neutral value
@@ -76,16 +78,16 @@ public:
 		, _ShiftMouseMovePixelPerDelta(1)
 		, _SupportDynamicSliderMaxValue(false)
 		, _SupportDynamicSliderMinValue(false)
-		, _SliderExponent(1)
-		, _Font( FCoreStyle::Get().GetFontStyle( TEXT( "NormalFont" ) ) )
-		, _ContentPadding(  FMargin( 2.0f, 1.0f) )
+		, _SliderExponent(1.f)
+		, _Font(FCoreStyle::Get().GetFontStyle(TEXT("NormalFont")))
+		, _ContentPadding(FMargin(2.0f, 1.0f))
 		, _OnValueChanged()
 		, _OnValueCommitted()
-		, _ClearKeyboardFocusOnCommit( false )
-		, _SelectAllTextOnCommit( true )
+		, _ClearKeyboardFocusOnCommit(false)
+		, _SelectAllTextOnCommit(true)
 		, _MinDesiredWidth(0.0f)
 		{}
-	
+
 		/** The style used to draw this spinbox */
 		SLATE_STYLE_ARGUMENT( FSpinBoxStyle, Style )
 
@@ -116,9 +118,9 @@ public:
 		/** Tell us if we want to support dynamically changing of the min value using ctrl */
 		SLATE_ATTRIBUTE(bool, SupportDynamicSliderMinValue)
 		/** Called right after the max slider value is changed (only relevant if SupportDynamicSliderMaxValue is true) */
-		SLATE_EVENT(FOnDynamicSliderMinMaxValueChanged, OnDynamicSliderMaxValueChanged)
+		SLATE_EVENT( FOnDynamicSliderMinMaxValueChanged, OnDynamicSliderMaxValueChanged )
 		/** Called right after the min slider value is changed (only relevant if SupportDynamicSliderMinValue is true) */
-		SLATE_EVENT(FOnDynamicSliderMinMaxValueChanged, OnDynamicSliderMinValueChanged)
+		SLATE_EVENT( FOnDynamicSliderMinMaxValueChanged, OnDynamicSliderMinValueChanged )
 		/** Use exponential scale for the slider */
 		SLATE_ATTRIBUTE( float, SliderExponent )
 		/** When use exponential scale for the slider which is the neutral value */
@@ -155,10 +157,10 @@ public:
 	SSpinBox()
 	{
 	}
-		
+
 	/**
 	 * Construct the widget
-	 * 
+	 *
 	 * @param InArgs   A declaration from which to construct the widget
 	 */
 	void Construct( const FArguments& InArgs )
@@ -176,7 +178,7 @@ public:
 		OnBeginSliderMovement = InArgs._OnBeginSliderMovement;
 		OnEndSliderMovement = InArgs._OnEndSliderMovement;
 		MinDesiredWidth = InArgs._MinDesiredWidth;
-	
+
 		MinValue = InArgs._MinValue;
 		MaxValue = InArgs._MaxValue;
 		MinSliderValue = (InArgs._MinSliderValue.Get().IsSet()) ? InArgs._MinSliderValue : MinValue;
@@ -194,32 +196,31 @@ public:
 
 		bPreventThrottling = InArgs._PreventThrottling;
 
-		// Update the max slider value based on the current value if we're in dynamic mode
-		NumericType CurrentMaxValue = GetMaxValue();
-		NumericType CurrentMinValue = GetMinValue();
+		CachedExternalValue = ValueAttribute.Get();
+		InternalValue = (double)CachedExternalValue;
 
-		if (SupportDynamicSliderMaxValue.Get() && ValueAttribute.Get() > GetMaxSliderValue())
-		{			
-			ApplySliderMaxValueChanged(ValueAttribute.Get() - GetMaxSliderValue(), true);
+		if (SupportDynamicSliderMaxValue.Get() && CachedExternalValue > GetMaxSliderValue())
+		{
+			ApplySliderMaxValueChanged(CachedExternalValue - GetMaxSliderValue(), true);
 		}
-		else if (SupportDynamicSliderMinValue.Get() && ValueAttribute.Get() < GetMinSliderValue())
+		else if (SupportDynamicSliderMinValue.Get() && CachedExternalValue < GetMinSliderValue())
 		{
-			ApplySliderMinValueChanged(ValueAttribute.Get() - GetMinSliderValue(), true);
+			ApplySliderMinValueChanged(CachedExternalValue - GetMinSliderValue(), true);
 		}
 
 		UpdateIsSpinRangeUnlimited();
-	
+
 		SliderExponent = InArgs._SliderExponent;
 
 		SliderExponentNeutralValue = InArgs._SliderExponentNeutralValue;
 
 		DistanceDragged = 0.0f;
-		PreDragValue = 0;
+		PreDragValue = NumericType();
 
 		Delta = InArgs._Delta;
 		ShiftMouseMovePixelPerDelta = InArgs._ShiftMouseMovePixelPerDelta;
 		LinearDeltaSensitivity = InArgs._LinearDeltaSensitivity;
-	
+
 		BackgroundHoveredBrush = &InArgs._Style->HoveredBackgroundBrush;
 		BackgroundBrush = &InArgs._Style->BackgroundBrush;
 		ActiveFillBrush = &InArgs._Style->ActiveFillBrush;
@@ -229,9 +230,6 @@ public:
 		bDragging = false;
 		PointerDraggingSliderIndex = INDEX_NONE;
 
-		CachedExternalValue = ValueAttribute.Get();
-		IntermediateDragFractionalValue = 0.0;
-
 		bIsTextChanging = false;
 
 		this->ChildSlot
@@ -286,7 +284,7 @@ public:
 			]
 		];
 	}
-	
+
 	virtual int32 OnPaint( const FPaintArgs& Args, const FGeometry& AllottedGeometry, const FSlateRect& MyCullingRect, FSlateWindowElementList& OutDrawElements, int32 LayerId, const FWidgetStyle& InWidgetStyle, bool bParentEnabled ) const override
 	{
 		const bool bActiveFeedback = IsHovered() || bDragging;
@@ -301,7 +299,7 @@ public:
 
 		const int32 BackgroundLayer = LayerId;
 
-		const bool bEnabled = ShouldBeEnabled( bParentEnabled );
+		const bool bEnabled = ShouldBeEnabled(bParentEnabled);
 		const ESlateDrawEffect DrawEffects = bEnabled ? ESlateDrawEffect::None : ESlateDrawEffect::DisabledEffect;
 
 		FSlateDrawElement::MakeBox(
@@ -320,24 +318,24 @@ public:
 		{
 			NumericType Value = ValueAttribute.Get();
 			NumericType CurrentDelta = Delta.Get();
-			if( CurrentDelta != 0.0f )
+			if (CurrentDelta != NumericType())
 			{
-				Value = FMath::GridSnap(Value, CurrentDelta); // snap floating point value to nearest Delta
+				Value = FMath::GridSnap(Value, CurrentDelta); // snap value to nearest Delta
 			}
 
-			float FractionFilled = Fraction(Value, GetMinSliderValue(), GetMaxSliderValue());
+			float FractionFilled = Fraction((double)Value, (double)GetMinSliderValue(), (double)GetMaxSliderValue());
 			const float CachedSliderExponent = SliderExponent.Get();
-			if (CachedSliderExponent != 1)
+			if (!FMath::IsNearlyEqual(CachedSliderExponent, 1.f))
 			{
 				if (SliderExponentNeutralValue.IsSet() && SliderExponentNeutralValue.Get() > GetMinSliderValue() && SliderExponentNeutralValue.Get() < GetMaxSliderValue())
 				{
 					//Compute a log curve on both side of the neutral value
-					float StartFractionFilled = Fraction(SliderExponentNeutralValue.Get(), GetMinSliderValue(), GetMaxSliderValue());
+					float StartFractionFilled = Fraction((double)SliderExponentNeutralValue.Get(), (double)GetMinSliderValue(), (double)GetMaxSliderValue());
 					FractionFilled = SpinBoxComputeExponentSliderFraction(FractionFilled, StartFractionFilled, CachedSliderExponent);
 				}
 				else
 				{
-					FractionFilled = 1.0f - FMath::Pow( 1.0f - FractionFilled, CachedSliderExponent);
+					FractionFilled = 1.0f - FMath::Pow(1.0f - FractionFilled, CachedSliderExponent);
 				}
 			}
 			const FVector2D FillSize( AllottedGeometry.GetLocalSize().X * FractionFilled, AllottedGeometry.GetLocalSize().Y );
@@ -347,7 +345,7 @@ public:
 				FSlateDrawElement::MakeBox(
 					OutDrawElements,
 					FilledLayer,
-					AllottedGeometry.ToPaintGeometry(FVector2D(0,0), FillSize),
+					AllottedGeometry.ToPaintGeometry(FVector2D(0.0f, 0.0f), FillSize),
 					FillImage,
 					DrawEffects,
 					FillImage->GetTint(InWidgetStyle) * InWidgetStyle.GetColorAndOpacityTint()
@@ -355,7 +353,7 @@ public:
 			}
 		}
 
-		return FMath::Max( FilledLayer, SCompoundWidget::OnPaint(Args, AllottedGeometry, MyCullingRect, OutDrawElements, FilledLayer, InWidgetStyle, bEnabled ) );
+		return FMath::Max(FilledLayer, SCompoundWidget::OnPaint(Args, AllottedGeometry, MyCullingRect, OutDrawElements, FilledLayer, InWidgetStyle, bEnabled));
 	}
 
 	/**
@@ -369,14 +367,14 @@ public:
 	{
 		if ( MouseEvent.GetEffectingButton() == EKeys::LeftMouseButton && PointerDraggingSliderIndex == INDEX_NONE )
 		{
-			DistanceDragged = 0;
+			DistanceDragged = 0.f;
 			PreDragValue = ValueAttribute.Get();
-			IntermediateDragFractionalValue = 0.0;
+			InternalValue = (double)PreDragValue;
 			PointerDraggingSliderIndex = MouseEvent.GetPointerIndex();
 			CachedMousePosition = MouseEvent.GetScreenSpacePosition().IntPoint();
 
 			FReply ReturnReply = FReply::Handled().CaptureMouse(SharedThis(this)).UseHighPrecisionMouseMovement(SharedThis(this)).SetUserFocus(SharedThis(this), EFocusCause::Mouse);
-			if (bPreventThrottling) 
+			if (bPreventThrottling)
 			{
 				ReturnReply.PreventThrottling();
 			}
@@ -388,7 +386,7 @@ public:
 			return FReply::Unhandled();
 		}
 	}
-	
+
 	/**
 	 * The system calls this method to notify the widget that a mouse button was release within it. This event is bubbled.
 	 *
@@ -400,18 +398,25 @@ public:
 	{
 		if ( MouseEvent.GetEffectingButton() == EKeys::LeftMouseButton && PointerDraggingSliderIndex == MouseEvent.GetPointerIndex() )
 		{
-			if(!this->HasMouseCapture())
+			if ( !this->HasMouseCapture() )
 			{
 				// Lost Capture - ensure reset
 				bDragging = false;
 				PointerDraggingSliderIndex = INDEX_NONE;
-				
+
 				return FReply::Unhandled();
 			}
-		
-			if( bDragging )
+
+			if ( bDragging )
 			{
-				NotifyValueCommitted();
+				NumericType CurrentDelta = Delta.Get();
+				if (CurrentDelta != NumericType())
+				{
+					InternalValue = FMath::GridSnap(InternalValue, (double)CurrentDelta);
+				}
+
+				const NumericType CurrentValue = RoundIfIntegerValue(InternalValue);
+				NotifyValueCommitted(CurrentValue);
 			}
 
 			bDragging = false;
@@ -432,7 +437,7 @@ public:
 
 			return Reply;
 		}
-		
+
 		return FReply::Unhandled();
 	}
 
@@ -440,15 +445,15 @@ public:
 	{
 		check(SupportDynamicSliderMaxValue.Get());
 
-		NumericType NewMaxSliderValue = TNumericLimits<NumericType>::Min();
-		
+		NumericType NewMaxSliderValue = std::numeric_limits<NumericType>::min();
+
 		if (MaxSliderValue.IsSet() && MaxSliderValue.Get().IsSet())
 		{
 			NewMaxSliderValue = GetMaxSliderValue();
 
-			if ((NewMaxSliderValue + SliderDeltaToAdd > GetMaxSliderValue() && UpdateOnlyIfHigher) || !UpdateOnlyIfHigher)
+			if ((NewMaxSliderValue + (NumericType)SliderDeltaToAdd > GetMaxSliderValue() && UpdateOnlyIfHigher) || !UpdateOnlyIfHigher)
 			{
-				NewMaxSliderValue += SliderDeltaToAdd;
+				NewMaxSliderValue += (NumericType)SliderDeltaToAdd;
 
 				if (!MaxSliderValue.IsBound()) // simple value so we can update it without breaking the mechanic otherwise it must be handle by the callback implementer
 				{
@@ -467,29 +472,29 @@ public:
 	{
 		check(SupportDynamicSliderMaxValue.Get());
 
-		NumericType NewMinSliderValue = TNumericLimits<NumericType>::Min();
-		
+		NumericType NewMinSliderValue = std::numeric_limits<NumericType>::min();
+
 		if (MinSliderValue.IsSet() && MinSliderValue.Get().IsSet())
 		{
 			NewMinSliderValue = GetMinSliderValue();
 
-			if ((NewMinSliderValue + SliderDeltaToAdd < GetMinSliderValue() && UpdateOnlyIfLower) || !UpdateOnlyIfLower)
+			if ((NewMinSliderValue + (NumericType)SliderDeltaToAdd < GetMinSliderValue() && UpdateOnlyIfLower) || !UpdateOnlyIfLower)
 			{
-				NewMinSliderValue += SliderDeltaToAdd;
+				NewMinSliderValue += (NumericType)SliderDeltaToAdd;
 
 				if (!MinSliderValue.IsBound()) // simple value so we can update it without breaking the mechanic otherwise it must be handle by the callback implementer
 				{
 					SetMinSliderValue(NewMinSliderValue);
 				}
 			}
-		}		
+		}
 
 		if (OnDynamicSliderMinValueChanged.IsBound())
 		{
 			OnDynamicSliderMinValueChanged.Execute(NewMinSliderValue, TWeakPtr<SWidget>(AsShared()), true, UpdateOnlyIfLower);
 		}
 	}
-	
+
 	/**
 	 * The system calls this method to notify the widget that a mouse moved within it. This event is bubbled.
 	 *
@@ -506,14 +511,14 @@ public:
 				// Lost the mouse capture - ensure reset
 				bDragging = false;
 				PointerDraggingSliderIndex = INDEX_NONE;
-				
+
 				return FReply::Unhandled();
 			}
-			
+
 			if (!bDragging)
 			{
 				DistanceDragged += FMath::Abs(MouseEvent.GetCursorDelta().X);
-				if ( DistanceDragged > FSlateApplication::Get().GetDragTriggerDistance() )
+				if (DistanceDragged > FSlateApplication::Get().GetDragTriggerDistance())
 				{
 					ExitTextMode();
 					bDragging = true;
@@ -525,12 +530,14 @@ public:
 			}
 			else
 			{
+				double NewValue = 0.0;
+
 				// Increments the spin based on delta mouse movement.
 
 				// A minimum slider width to use for calculating deltas in the slider-range space
 				const float MinSliderWidth = 100.f;
 				float SliderWidthInSlateUnits = FMath::Max(MyGeometry.GetDrawSize().X, MinSliderWidth);
-				
+
 				const int32 CachedShiftMouseMovePixelPerDelta = ShiftMouseMovePixelPerDelta.Get();
 				if (CachedShiftMouseMovePixelPerDelta > 1 && MouseEvent.IsShiftDown())
 				{
@@ -541,31 +548,31 @@ public:
 				{
 					float DeltaToAdd = MouseEvent.GetCursorDelta().X / SliderWidthInSlateUnits;
 
-					if (SupportDynamicSliderMaxValue.Get() && CachedExternalValue == GetMaxSliderValue())
+					if (SupportDynamicSliderMaxValue.Get() && (NumericType)InternalValue == GetMaxSliderValue())
 					{
 						ApplySliderMaxValueChanged(DeltaToAdd, false);
 					}
-					else if (SupportDynamicSliderMinValue.Get() && CachedExternalValue == GetMinSliderValue())
+					else if (SupportDynamicSliderMinValue.Get() && (NumericType)InternalValue == GetMinSliderValue())
 					{
 						ApplySliderMinValueChanged(DeltaToAdd, false);
 					}
 				}
-				
+
 				//if we have a range to draw in
-				if ( !bUnlimitedSpinRange) 
+				if (!bUnlimitedSpinRange)
 				{
 					bool HasValidExponentNeutralValue = SliderExponentNeutralValue.IsSet() && SliderExponentNeutralValue.Get() > GetMinSliderValue() && SliderExponentNeutralValue.Get() < GetMaxSliderValue();
 
 					const float CachedSliderExponent = SliderExponent.Get();
 					// The amount currently filled in the spinbox, needs to be calculated to do deltas correctly.
-					float FractionFilled = Fraction(PreDragValue, GetMinSliderValue(), GetMaxSliderValue());
-						
-					if (CachedSliderExponent != 1)
+					float FractionFilled = Fraction(InternalValue, (double)GetMinSliderValue(), (double)GetMaxSliderValue());
+
+					if (!FMath::IsNearlyEqual(CachedSliderExponent, 1.0f))
 					{
 						if (HasValidExponentNeutralValue)
 						{
 							//Compute a log curve on both side of the neutral value
-							float StartFractionFilled = Fraction(SliderExponentNeutralValue.Get(), GetMinSliderValue(), GetMaxSliderValue());
+							float StartFractionFilled = Fraction((double)SliderExponentNeutralValue.Get(), (double)GetMinSliderValue(), (double)GetMaxSliderValue());
 							FractionFilled = SpinBoxComputeExponentSliderFraction(FractionFilled, StartFractionFilled, CachedSliderExponent);
 						}
 						else
@@ -576,36 +583,31 @@ public:
 					FractionFilled *= SliderWidthInSlateUnits;
 
 					// Now add the delta to the fraction filled, this causes the spin.
-					float MouseDelta = MouseEvent.GetScreenSpacePosition().IntPoint().X - CachedMousePosition.X;
-					FractionFilled += MouseDelta;
-						
+					FractionFilled += MouseEvent.GetCursorDelta().X;
+
 					// Clamp the fraction to be within the bounds of the geometry.
 					FractionFilled = FMath::Clamp(FractionFilled, 0.0f, SliderWidthInSlateUnits);
 
 					// Convert the fraction filled to a percent.
 					float Percent = FMath::Clamp(FractionFilled / SliderWidthInSlateUnits, 0.0f, 1.0f);
-					if (CachedSliderExponent != 1)
+					if (!FMath::IsNearlyEqual(CachedSliderExponent, 1.0f))
 					{
 						// Have to convert the percent to the proper value due to the exponent component to the spin.
 						if (HasValidExponentNeutralValue)
 						{
 							//Compute a log curve on both side of the neutral value
 							float StartFractionFilled = Fraction(SliderExponentNeutralValue.Get(), GetMinSliderValue(), GetMaxSliderValue());
-							Percent = SpinBoxComputeExponentSliderFraction(Percent, StartFractionFilled, 1.0/CachedSliderExponent);
+							Percent = SpinBoxComputeExponentSliderFraction(Percent, StartFractionFilled, 1.0f / CachedSliderExponent);
 						}
 						else
 						{
 							Percent = 1.0f - FMath::Pow(1.0f - Percent, 1.0 / CachedSliderExponent);
 						}
-						
-						
+
+
 					}
 
-					double ValueToRound = FMath::LerpStable<double>(GetMinSliderValue(), GetMaxSliderValue(), Percent);
-					NumericType NewValue = TIsIntegral<NumericType>::Value
-						? (NumericType)FMath::FloorToDouble(ValueToRound + 0.5)
-						: (NumericType)ValueToRound;
-					CommitValue(NewValue, CommittedViaSpin, ETextCommit::OnEnter);
+					NewValue = FMath::LerpStable<double>((double)GetMinSliderValue(), (double)GetMaxSliderValue(), Percent);
 				}
 				else
 				{
@@ -614,18 +616,19 @@ public:
 					if (LinearDeltaSensitivity.IsSet() && Delta.IsSet() && Delta.Get() > 0)
 					{
 						const float MouseDelta = FMath::Abs(MouseEvent.GetCursorDelta().X / LinearDeltaSensitivity.Get());
-						IntermediateDragFractionalValue += (Sign * MouseDelta * FMath::Pow(Delta.Get(), SliderExponent.Get()));
+						NewValue = InternalValue + (Sign * MouseDelta * FMath::Pow(Delta.Get(), SliderExponent.Get()));
 					}
 					else
 					{
 						const float MouseDelta = FMath::Abs(MouseEvent.GetCursorDelta().X / SliderWidthInSlateUnits);
-						const double CurrentValue = FMath::Clamp<double>(FMath::Abs(static_cast<double>(CachedExternalValue)), 1.0, TNumericLimits<NumericType>::Max());
-						IntermediateDragFractionalValue += (Sign * MouseDelta * FMath::Pow(CurrentValue, SliderExponent.Get()));
+						const double CurrentValue = FMath::Clamp<double>(FMath::Abs(InternalValue), 1.0, (double)std::numeric_limits<NumericType>::max());
+						NewValue = InternalValue + (Sign * MouseDelta * FMath::Pow(CurrentValue, SliderExponent.Get()));
 					}
-
-					NumericType NewValue = UpdateDraggingValues(CachedExternalValue);
-					CommitValue(NewValue, CommittedViaSpin, ETextCommit::OnEnter);
 				}
+
+				NumericType RoundedNewValue = RoundIfIntegerValue(NewValue);
+				//UE_LOG(LogSlate, Warning, TEXT("A - %llx B - %f  C - %llx"), RoundedNewValue, NewValue, static_cast<NumericType>(NewValue));
+				CommitValue(RoundedNewValue, NewValue, CommittedViaSpin, ETextCommit::OnEnter);
 			}
 
 			return FReply::Handled();
@@ -660,7 +663,7 @@ public:
 			return FReply::Unhandled();
 		}
 	}
-	
+
 	virtual FReply OnKeyDown( const FGeometry& MyGeometry, const FKeyEvent& InKeyEvent ) override
 	{
 		const FKey Key = InKeyEvent.GetKey();
@@ -669,25 +672,31 @@ public:
 			bDragging = false;
 			PointerDraggingSliderIndex = INDEX_NONE;
 
-			CachedExternalValue = PreDragValue;
-			NotifyValueCommitted();
+			InternalValue = (double)PreDragValue;
+			NotifyValueCommitted(PreDragValue);
 			return FReply::Handled().ReleaseMouseCapture().SetMousePos(CachedMousePosition);
 		}
 		else if ( Key == EKeys::Up || Key == EKeys::Right )
 		{
-			CommitValue( ValueAttribute.Get() + Delta.Get(), CommittedViaArrowKey, ETextCommit::OnEnter );
+			const NumericType LocalValueAttribute = ValueAttribute.Get();
+			const NumericType LocalDelta = Delta.Get();
+			InternalValue = (double)LocalValueAttribute;
+			CommitValue(LocalValueAttribute + LocalDelta, InternalValue + (double)LocalDelta, CommittedViaArrowKey, ETextCommit::OnEnter);
 			ExitTextMode();
-			return FReply::Handled();	
+			return FReply::Handled();
 		}
 		else if ( Key == EKeys::Down || Key == EKeys::Left )
 		{
-			CommitValue( ValueAttribute.Get() - Delta.Get(), CommittedViaArrowKey, ETextCommit::OnEnter );
+			const NumericType LocalValueAttribute = ValueAttribute.Get();
+			const NumericType LocalDelta = Delta.Get();
+			InternalValue = (double)LocalValueAttribute;
+			CommitValue(LocalValueAttribute - LocalDelta, InternalValue + (double)LocalDelta, CommittedViaArrowKey, ETextCommit::OnEnter);
 			ExitTextMode();
 			return FReply::Handled();
 		}
 		else if ( Key == EKeys::Enter )
 		{
-			CachedExternalValue = ValueAttribute.Get();
+			InternalValue = (double)ValueAttribute.Get();
 			EnterTextMode();
 			return FReply::Handled().SetUserFocus(EditableText.ToSharedRef(), EFocusCause::Navigation);
 		}
@@ -696,7 +705,7 @@ public:
 			return FReply::Unhandled();
 		}
 	}
-	
+
 	virtual bool HasKeyboardFocus() const override
 	{
 		// The spinbox is considered focused when we are typing it text.
@@ -711,21 +720,22 @@ public:
 	void SetValue(const TAttribute<NumericType>& InValueAttribute) 
 	{
 		ValueAttribute = InValueAttribute; 
-		CommitValue(InValueAttribute.Get(), ECommitMethod::CommittedViaCode, ETextCommit::Default);
+		const NumericType LocalValueAttribute = ValueAttribute.Get();
+		CommitValue(LocalValueAttribute, (double)LocalValueAttribute, ECommitMethod::CommittedViaCode, ETextCommit::Default);
 	}
 
 	/** See the MinValue attribute */
-	NumericType GetMinValue() const { return MinValue.Get().Get(TNumericLimits<NumericType>::Lowest()); }
+	NumericType GetMinValue() const { return MinValue.Get().Get(std::numeric_limits<NumericType>::lowest()); }
 	void SetMinValue(const TAttribute<TOptional<NumericType>>& InMinValue) 
-	{ 
+	{
 		MinValue = InMinValue;
 		UpdateIsSpinRangeUnlimited();
 	}
 
 	/** See the MaxValue attribute */
-	NumericType GetMaxValue() const { return MaxValue.Get().Get(TNumericLimits<NumericType>::Max()); }
+	NumericType GetMaxValue() const { return MaxValue.Get().Get(std::numeric_limits<NumericType>::max()); }
 	void SetMaxValue(const TAttribute<TOptional<NumericType>>& InMaxValue) 
-	{ 
+	{
 		MaxValue = InMaxValue; 
 		UpdateIsSpinRangeUnlimited();
 	}
@@ -733,9 +743,9 @@ public:
 	/** See the MinSliderValue attribute */
 	bool IsMinSliderValueBound() const { return MinSliderValue.IsBound(); }
 
-	NumericType GetMinSliderValue() const { return MinSliderValue.Get().Get(TNumericLimits<NumericType>::Lowest()); }
+	NumericType GetMinSliderValue() const { return MinSliderValue.Get().Get(std::numeric_limits<NumericType>::lowest()); }
 	void SetMinSliderValue(const TAttribute<TOptional<NumericType>>& InMinSliderValue) 
-	{ 
+	{
 		MinSliderValue = (InMinSliderValue.Get().IsSet()) ? InMinSliderValue : MinValue;
 		UpdateIsSpinRangeUnlimited();
 	}
@@ -743,9 +753,9 @@ public:
 	/** See the MaxSliderValue attribute */
 	bool IsMaxSliderValueBound() const { return MaxSliderValue.IsBound(); }
 
-	NumericType GetMaxSliderValue() const { return MaxSliderValue.Get().Get(TNumericLimits<NumericType>::Max()); }
+	NumericType GetMaxSliderValue() const { return MaxSliderValue.Get().Get(std::numeric_limits<NumericType>::max()); }
 	void SetMaxSliderValue(const TAttribute<TOptional<NumericType>>& InMaxSliderValue) 
-	{ 
+	{
 		MaxSliderValue = (InMaxSliderValue.Get().IsSet()) ? InMaxSliderValue : MaxValue;;
 		UpdateIsSpinRangeUnlimited();
 	}
@@ -775,7 +785,7 @@ public:
 	/** See the SliderExponent attribute */
 	float GetSliderExponent() const { return SliderExponent.Get(); }
 	void SetSliderExponent(const TAttribute<float>& InSliderExponent) { SliderExponent = InSliderExponent; }
-	
+
 	/** See the MinDesiredWidth attribute */
 	float GetMinDesiredWidth() const { return SliderExponent.Get(); }
 	void SetMinDesiredWidth(const TAttribute<float>& InMinDesiredWidth) { MinDesiredWidth = InMinDesiredWidth; }
@@ -787,7 +797,7 @@ protected:
 		TextBlock->SetVisibility( EVisibility::Collapsed );
 		EditableText->SetVisibility( EVisibility::Visible );
 	}
-	
+
 	/** Make the spinbox switch to mouse-based input mode. */
 	void ExitTextMode()
 	{
@@ -812,7 +822,7 @@ protected:
 	 *
 	 * @param NewText		The value of the text in the text field
 	 */
-	void TextField_OnTextChanged( const FText& NewText)
+	void TextField_OnTextChanged(const FText& NewText)
 	{
 		if (!bIsTextChanging)
 		{
@@ -838,7 +848,7 @@ protected:
 			}
 		}
 	}
-	
+
 	/**
 	 * Invoked when the text field commits its text.
 	 *
@@ -855,7 +865,7 @@ protected:
 		TOptional<NumericType> NewValue = Interface->FromString(NewText.ToString(), ValueAttribute.Get());
 		if (NewValue.IsSet())
 		{
-			CommitValue(NewValue.GetValue(), CommittedViaTypeIn, CommitInfo);
+			CommitValue(NewValue.GetValue(), (double)NewValue.GetValue(), CommittedViaTypeIn, CommitInfo);
 		}
 	}
 
@@ -876,100 +886,109 @@ protected:
 	 * @param CommitMethod           Did the user type in the value or spin to it.
 	 * @param OriginalCommitInfo	 If the user typed in the value, information about the source of the commit
 	 */
-	void CommitValue( NumericType NewValue, ECommitMethod CommitMethod, ETextCommit::Type OriginalCommitInfo )
+	void CommitValue(NumericType NewValue, double NewSpinValue, ECommitMethod CommitMethod, ETextCommit::Type OriginalCommitInfo)
 	{
-		NumericType ValueToCommit = NewValue;
-		if( CommitMethod == CommittedViaSpin || CommitMethod == CommittedViaArrowKey )
+		if (CommitMethod == CommittedViaSpin || CommitMethod == CommittedViaArrowKey)
 		{
-			ValueToCommit = FMath::Clamp<NumericType>(ValueToCommit, GetMinSliderValue(), GetMaxSliderValue());
+			const NumericType LocalMinSliderValue = GetMinSliderValue();
+			const NumericType LocalMaxSliderValue = GetMaxSliderValue();
+			NewValue = FMath::Clamp<NumericType>(NewValue, LocalMinSliderValue, LocalMaxSliderValue);
+			NewSpinValue = FMath::Clamp<double>(NewSpinValue, (double)LocalMinSliderValue, (double)LocalMaxSliderValue);
 		}
 
-		ValueToCommit = FMath::Clamp<NumericType>(ValueToCommit, GetMinValue(), GetMaxValue());
+		{
+			const NumericType LocalMinValue = GetMinValue();
+			const NumericType LocalMaxValue = GetMaxValue();
+			NewValue = FMath::Clamp<NumericType>(NewValue, LocalMinValue, LocalMaxValue);
+			NewSpinValue = FMath::Clamp<double>(NewSpinValue, (double)LocalMinValue, (double)LocalMaxValue);
+		}
+
+		if (!ValueAttribute.IsBound())
+		{
+			ValueAttribute.Set(NewValue);
+		}
 
 		// If not in spin mode, there is no need to jump to the value from the external source, continue to use the committed value.
-		if(CommitMethod == CommittedViaSpin)
+		if (CommitMethod == CommittedViaSpin)
 		{
-			// This will detect if an external force has changed the value. Internally it will abandon the delta calculated this tick and update the internal value instead.
 			const NumericType CurrentValue = ValueAttribute.Get();
-			if(CurrentValue != CachedExternalValue)
+			// This will detect if an external force has changed the value. Internally it will abandon the delta calculated this tick and update the internal value instead.
+			if (CurrentValue != CachedExternalValue)
 			{
-				ValueToCommit = CurrentValue;
+				NewValue = CurrentValue;
+				NewSpinValue = (double)CurrentValue;
 			}
 		}
-		else
-		{
-			// Reset intermediate spin value
-			IntermediateDragFractionalValue = 0.0;
-		}
+
+		// Update the internal value, this needs to be done before rounding.
+		InternalValue = NewSpinValue;
 
 		const bool bAlwaysUsesDeltaSnap = GetAlwaysUsesDeltaSnap();
 		// If needed, round this value to the delta. Internally the value is not held to the Delta but externally it appears to be.
-		if ( CommitMethod == CommittedViaSpin || CommitMethod == CommittedViaArrowKey || bAlwaysUsesDeltaSnap)
+		if (CommitMethod == CommittedViaSpin || CommitMethod == CommittedViaArrowKey || bAlwaysUsesDeltaSnap)
 		{
 			NumericType CurrentDelta = Delta.Get();
-			if( CurrentDelta != 0 )
+			if (CurrentDelta != NumericType())
 			{
-				ValueToCommit = FMath::GridSnap((double)ValueToCommit, (double)CurrentDelta); // snap numeric point value to nearest Delta
-				if (CachedExternalValue != ValueToCommit)
-				{
-					IntermediateDragFractionalValue = 0.0;
-				}
+				NewValue = FMath::GridSnap<NumericType>(NewValue, CurrentDelta); // snap numeric point value to nearest Delta
 			}
-		}		
+		}
 
 		// Update the max slider value based on the current value if we're in dynamic mode
-		if (SupportDynamicSliderMaxValue.Get() && ValueToCommit > GetMaxSliderValue())
+		if (SupportDynamicSliderMaxValue.Get() && ValueAttribute.Get() > GetMaxSliderValue())
 		{
-			ApplySliderMaxValueChanged(ValueToCommit - GetMaxSliderValue(), true);
+			ApplySliderMaxValueChanged(ValueAttribute.Get() - GetMaxSliderValue(), true);
 		}
-		else if (SupportDynamicSliderMinValue.Get() && ValueToCommit < GetMinSliderValue())
+		else if (SupportDynamicSliderMinValue.Get() && ValueAttribute.Get() < GetMinSliderValue())
 		{
-			ApplySliderMinValueChanged(ValueToCommit - GetMinSliderValue(), true);
+			ApplySliderMinValueChanged(ValueAttribute.Get() - GetMinSliderValue(), true);
 		}
 
-		if( CommitMethod == CommittedViaTypeIn || CommitMethod == CommittedViaArrowKey )
+		if (CommitMethod == CommittedViaTypeIn || CommitMethod == CommittedViaArrowKey)
 		{
-			OnValueCommitted.ExecuteIfBound(ValueToCommit, OriginalCommitInfo);
+			OnValueCommitted.ExecuteIfBound(NewValue, OriginalCommitInfo);
 		}
 
-		OnValueChanged.ExecuteIfBound(ValueToCommit);
+		OnValueChanged.ExecuteIfBound(NewValue);
 
 		if (!ValueAttribute.IsBound())
 		{
-			ValueAttribute.Set(ValueToCommit);
+			ValueAttribute.Set(NewValue);
 		}
 
 		// Update the cache of the external value to what the user believes the value is now.
 		CachedExternalValue = ValueAttribute.Get();
 
 		// This ensures that dragging is cleared if focus has been removed from this widget in one of the delegate calls, such as when spawning a modal dialog.
-		if ( !this->HasMouseCapture() )
+		if (!this->HasMouseCapture())
 		{
 			bDragging = false;
 			PointerDraggingSliderIndex = INDEX_NONE;
 		}
 	}
-	
-	void NotifyValueCommitted() const
+
+	void NotifyValueCommitted(NumericType CurrentValue) const
 	{
-		OnValueCommitted.ExecuteIfBound( CachedExternalValue, ETextCommit::OnEnter );
-		OnEndSliderMovement.ExecuteIfBound( CachedExternalValue );
+		// The internal value will have been clamped and rounded to the delta at this point, but integer values may still need to be rounded
+		// if the delta is 0.
+		OnValueCommitted.ExecuteIfBound(CurrentValue, ETextCommit::OnEnter);
+		OnEndSliderMovement.ExecuteIfBound(CurrentValue);
 	}
 
 	/** @return true when we are in keyboard-based input mode; false otherwise */
 	bool IsInTextMode() const
 	{
-		return ( EditableText->GetVisibility() == EVisibility::Visible );
+		return (EditableText->GetVisibility() == EVisibility::Visible);
 	}
 
 	/** Calculates range fraction. Possible to use on full numeric range  */
-	static float Fraction(NumericType InValue, NumericType InMinValue, NumericType InMaxValue)
+	static float Fraction(double InValue, double InMinValue, double InMaxValue)
 	{
-		const double HalfMax = InMaxValue*0.5;
-		const double HalfMin = InMinValue*0.5;
-		const double HalfVal = InValue*0.5;
+		const double HalfMax = InMaxValue * 0.5;
+		const double HalfMin = InMinValue * 0.5;
+		const double HalfVal = InValue * 0.5;
 
-		return (float)FMath::Clamp((HalfVal - HalfMin)/(HalfMax - HalfMin), 0.0, 1.0);
+		return (float)FMath::Clamp((HalfVal - HalfMin) / (HalfMax - HalfMin), 0.0, 1.0);
 	}
 
 private:
@@ -1036,20 +1055,26 @@ private:
 		return Interface->IsCharacterValid(InChar);
 	}
 
-	/** Update the IntermediateDragFractionalValue and return the updated CurrentValue */
-	NumericType UpdateDraggingValues(NumericType CurrentValue)
+	/**
+	 * Rounds the submitted value to the correct value if it's an integer.
+	 * For int64, not all values can be represented by a double. We can only round until we reach that limit.
+	 * This function should only be used when we drag the value. We accept that we can't drag huge numbers.
+	 */
+	NumericType RoundIfIntegerValue(double ValueToRound) const
 	{
-		if (TIsIntegral<NumericType>::Value)
-		{						
-			double IntegralPart = 0.f;
-			IntermediateDragFractionalValue = FMath::Modf(IntermediateDragFractionalValue, &IntegralPart);
-						
-			return CurrentValue + static_cast<NumericType>(IntegralPart);
+		constexpr bool bIsIntegral = TIsIntegral<NumericType>::Value;
+		constexpr bool bCanBeRepresentedInDouble = std::numeric_limits<double>::digits >= std::numeric_limits<NumericType>::digits;
+		if (bIsIntegral && !bCanBeRepresentedInDouble)
+		{
+			return (NumericType)FMath::Clamp<double>(FMath::FloorToDouble(ValueToRound + 0.5), -1.0*(double)(1ll << std::numeric_limits<double>::digits), (double)(1ll << std::numeric_limits<double>::digits));
+		}
+		else if (bIsIntegral)
+		{
+			return (NumericType)FMath::Clamp<double>(FMath::FloorToDouble(ValueToRound + 0.5), (double)std::numeric_limits<NumericType>::lowest(), (double)std::numeric_limits<NumericType>::max());
 		}
 		else
 		{
-			NumericType TmpValue = CurrentValue + IntermediateDragFractionalValue;
-			return TmpValue;
+			return (NumericType)FMath::Clamp<double>(ValueToRound, (double)std::numeric_limits<NumericType>::lowest(), (double)std::numeric_limits<NumericType>::max());
 		}
 	}
 
@@ -1058,25 +1083,25 @@ private:
 
 	/** Tracks which cursor is currently dragging the slider (e.g., the mouse cursor or a specific finger) */
 	int32 PointerDraggingSliderIndex;
-	
+
 	/** Cached mouse position to restore after scrolling. */
 	FIntPoint CachedMousePosition;
 
-	/*
-	 * This is the cached value the user believes it to be. Used for identifying external forces on the spinbox
-	 * and syncing the internal value to them. Synced when a value is committed to the spinbox. 
+	/**
+	 * This value represents what the spinbox believes the value to be, regardless of delta and the user binding to an int.
+	 * The spinbox will always count using floats between values, this is important to keep it flowing smoothly and feeling right,
+	 * and most importantly not conflicting with the user truncating the value to an int.
 	 */
-	NumericType CachedExternalValue;
-	
-	/** The state of CachedExternalValue before a drag operation was started */
+	double InternalValue;
+
+	/** The state of InternalValue before a drag operation was started */
 	NumericType PreDragValue;
 
 	/**
-	 * This value represents the fractional part of the spinbox when using integers. 
-	 * The spinbox will always count using floats between values, this is important to keep it flowing smoothly and feeling right, 
-	 * and most importantly not conflicting with the user truncating the value to an int.
+	 * This is the cached value the user believes it to be (usually different due to truncation to an int). Used for identifying
+	 * external forces on the spinbox and syncing the internal value to them. Synced when a value is committed to the spinbox.
 	 */
-	double IntermediateDragFractionalValue;
+	NumericType CachedExternalValue;
 
 	/** Re-entrant guard for the text changed handler */
 	bool bIsTextChanging;
