SSplitter now takes a HighlightedHandleIndex that will always appear to be hovered, and has a delegate OnHandleHovered when hovering a handle. These combined allow synchronizing the hovered handle in the details view.

diff --git a/Engine/Source/Runtime/Slate/Private/Widgets/Layout/SSplitter.cpp b/Engine/Source/Runtime/Slate/Private/Widgets/Layout/SSplitter.cpp
--- a/Engine/Source/Runtime/Slate/Private/Widgets/Layout/SSplitter.cpp
+++ b/Engine/Source/Runtime/Slate/Private/Widgets/Layout/SSplitter.cpp
@@ -46,12 +46,14 @@ void SSplitter::Construct( const SSplitter::FArguments& InArgs )
 {
 	check(InArgs._Style);
 
+	OnHandleHovered = InArgs._OnHandleHovered;
 	OnSplitterFinishedResizing = InArgs._OnSplitterFinishedResizing;
 	ResizeMode = InArgs._ResizeMode;
 	PhysicalSplitterHandleSize = InArgs._PhysicalSplitterHandleSize;
 	HitDetectionSplitterHandleSize = InArgs._HitDetectionSplitterHandleSize;
 	Orientation = InArgs._Orientation;
 	HoveredHandleIndex = INDEX_NONE;
+	HighlightedHandleIndex = InArgs._HighlightedHandleIndex;
 	bIsResizing = false;
 	Style = InArgs._Style;
 	OnGetMaxSlotSize = InArgs._OnGetMaxSlotSize;
@@ -147,7 +149,7 @@ void SSplitter::OnArrangeChildren( const FGeometry& AllottedGeometry, FArrangedC
 	TArray<FLayoutGeometry> LayoutChildren = ArrangeChildrenForLayout(AllottedGeometry);
 
 	// Arrange the children horizontally or vertically.
-	for (int32 ChildIndex=0; ChildIndex < Children.Num(); ++ChildIndex)
+	for (int32 ChildIndex = 0; ChildIndex < Children.Num(); ++ChildIndex)
 	{
 		ArrangedChildren.AddWidget( AllottedGeometry.MakeChild( Children[ChildIndex].GetWidget(), LayoutChildren[ChildIndex] ) );
 	}
@@ -166,18 +168,25 @@ int32 SSplitter::OnPaint( const FPaintArgs& Args, const FGeometry& AllottedGeome
 	// Draw the splitter above any children
 	MaxLayerId += 1;
 
-	// HoveredHandleIndex is an index into our Children array, but some of those may be hidden and missing from ArrangedChildren so we need to update the index for the loop below
-	int32 ArrangedHoveredHandleIndex = HoveredHandleIndex;
-	if (ArrangedHoveredHandleIndex != INDEX_NONE && ArrangedChildren.Num() != Children.Num())
+	// HandleIndex is an index into our Children array, but some of those may be hidden and missing from ArrangedChildren so we need to update the index for the loop below
+	auto FindArrangedHandleIndex = [this, &ArrangedChildren](int32 HandleIndex)
 	{
-		const TSharedRef<SWidget>& HoveredChildToFind = Children[HoveredHandleIndex].GetWidget();
-		ArrangedHoveredHandleIndex = ArrangedChildren.IndexOfByPredicate([&HoveredChildToFind](const FArrangedWidget& InArrangedWidget)
+		if (HandleIndex == INDEX_NONE || ArrangedChildren.Num() == Children.Num())
+		{
+			return HandleIndex;
+		}
+
+		const TSharedRef<SWidget>& HoveredChildToFind = Children[HandleIndex].GetWidget();
+		return ArrangedChildren.IndexOfByPredicate([&HoveredChildToFind](const FArrangedWidget& InArrangedWidget)
 		{
 			return InArrangedWidget.Widget == HoveredChildToFind;
 		});
-	}
+	};
+
+	int32 ArrangedHoveredHandleIndex = FindArrangedHandleIndex(HoveredHandleIndex);
+	int32 ArrangedHighlightedHandleIndex = FindArrangedHandleIndex(HighlightedHandleIndex.Get(INDEX_NONE));
 
-	for( int32 ChildIndex = 0; ChildIndex < ArrangedChildren.Num() - 1; ++ChildIndex )
+	for (int32 ChildIndex = 0; ChildIndex < ArrangedChildren.Num() - 1; ++ChildIndex)
 	{
 		const FGeometry& GeometryAfterSplitter = ArrangedChildren[ FMath::Clamp(ChildIndex + 1, 0, ArrangedChildren.Num()-1) ].Geometry;
 
@@ -186,7 +195,7 @@ int32 SSplitter::OnPaint( const FPaintArgs& Args, const FGeometry& AllottedGeome
 
 		FVector2D HandleSize;		
 		FVector2D HandlePosition;
-		if ( Orientation == Orient_Horizontal )
+		if (Orientation == Orient_Horizontal)
 		{
 			HandleSize.Set( PhysicalSplitterHandleSize, GeometryAfterSplitter.Size.Y );
 			HandlePosition.Set( -PhysicalSplitterHandleSize/*-(HalfHitDetectionSplitterHandleSize + HalfPhysicalSplitterHandleSize)*/, 0 );
@@ -197,15 +206,16 @@ int32 SSplitter::OnPaint( const FPaintArgs& Args, const FGeometry& AllottedGeome
 			HandlePosition.Set( 0, -PhysicalSplitterHandleSize/*-(HalfHitDetectionSplitterHandleSize + HalfPhysicalSplitterHandleSize) */);
 		}
 
-		if (ArrangedHoveredHandleIndex != ChildIndex)
+		if (ArrangedHoveredHandleIndex == ChildIndex || ArrangedHighlightedHandleIndex == ChildIndex)
 		{
+			// handle is hovered
 			FSlateDrawElement::MakeBox(
 				OutDrawElements,
 				MaxLayerId,
 				GeometryAfterSplitter.ToPaintGeometry( HandlePosition, HandleSize, 1.0f ),
 				NormalHandleBrush,
 				ShouldBeEnabled(bParentEnabled) ? ESlateDrawEffect::None : ESlateDrawEffect::DisabledEffect,
-				InWidgetStyle.GetColorAndOpacityTint() * NormalHandleBrush->TintColor.GetSpecifiedColor()
+				InWidgetStyle.GetColorAndOpacityTint() * Style->HandleHighlightBrush.TintColor.GetSpecifiedColor()
 			);
 		}
 		else
@@ -216,7 +226,7 @@ int32 SSplitter::OnPaint( const FPaintArgs& Args, const FGeometry& AllottedGeome
 				GeometryAfterSplitter.ToPaintGeometry( HandlePosition, HandleSize, 1.0f ),
 				&Style->HandleHighlightBrush,
 				ShouldBeEnabled(bParentEnabled) ? ESlateDrawEffect::None : ESlateDrawEffect::DisabledEffect,
-				InWidgetStyle.GetColorAndOpacityTint() * Style->HandleHighlightBrush.TintColor.GetSpecifiedColor()
+				InWidgetStyle.GetColorAndOpacityTint() * NormalHandleBrush->TintColor.GetSpecifiedColor()
 			);	
 		}
 	}
@@ -227,10 +237,10 @@ int32 SSplitter::OnPaint( const FPaintArgs& Args, const FGeometry& AllottedGeome
 template<EOrientation Orientation>
 static FVector2D ComputeDesiredSizeForSplitter( const float PhysicalSplitterHandleSize, const TPanelChildren<SSplitter::FSlot>& Children )
 {
-	FVector2D MyDesiredSize(0,0);
+	FVector2D MyDesiredSize(0, 0);
 
 	int32 NumNonCollapsed = 0;
-	for (int32 ChildIndex=0; ChildIndex < Children.Num(); ++ChildIndex)
+	for (int32 ChildIndex = 0; ChildIndex < Children.Num(); ++ChildIndex)
 	{
 		const SSplitter::FSlot& CurSlot = Children[ChildIndex];
 		const EVisibility ChildVisibility = CurSlot.GetWidget()->GetVisibility();
@@ -345,9 +355,9 @@ FReply SSplitter::OnMouseMove( const FGeometry& MyGeometry, const FPointerEvent&
 
 	TArray<FLayoutGeometry> LayoutChildren = ArrangeChildrenForLayout(MyGeometry);
 
-	if ( bIsResizing )
+	if (bIsResizing)
 	{
-		if ( !MouseEvent.GetCursorDelta().IsZero() )
+		if (!MouseEvent.GetCursorDelta().IsZero())
 		{
 			HandleResizingByMousePosition(Orientation, PhysicalSplitterHandleSize, ResizeMode, HoveredHandleIndex, LocalMousePosition, Children, LayoutChildren );
 		}
@@ -355,20 +365,27 @@ FReply SSplitter::OnMouseMove( const FGeometry& MyGeometry, const FPointerEvent&
 		return FReply::Handled();
 	}
 	else
-	{	
-		// Hit test which handle we are hovering over.		
+	{
+		// Hit test which handle we are hovering over.
+		int32 PreviousHoveredHandleIndex = HoveredHandleIndex;
+
 		HoveredHandleIndex = (Orientation == Orient_Horizontal)
 			? GetHandleBeingResizedFromMousePosition<Orient_Horizontal>( PhysicalSplitterHandleSize, HitDetectionSplitterHandleSize, LocalMousePosition, LayoutChildren )
 			: GetHandleBeingResizedFromMousePosition<Orient_Vertical>( PhysicalSplitterHandleSize, HitDetectionSplitterHandleSize, LocalMousePosition, LayoutChildren );
 
 		if (HoveredHandleIndex != INDEX_NONE)
 		{
-			if ( FindResizeableSlotBeforeHandle(HoveredHandleIndex, Children) <= INDEX_NONE || FindResizeableSlotAfterHandle(HoveredHandleIndex, Children) >= Children.Num() )
+			if (FindResizeableSlotBeforeHandle(HoveredHandleIndex, Children) <= INDEX_NONE || FindResizeableSlotAfterHandle(HoveredHandleIndex, Children) >= Children.Num())
 			{
 				HoveredHandleIndex = INDEX_NONE;
 			}
 		}
 
+		if (HoveredHandleIndex != PreviousHoveredHandleIndex)
+		{
+			OnHandleHovered.ExecuteIfBound(HoveredHandleIndex);
+		}
+
 		return FReply::Unhandled();
 	}
 
@@ -395,10 +412,17 @@ FReply SSplitter::OnMouseButtonDoubleClick(const FGeometry& InMyGeometry, const
 
 void SSplitter::OnMouseLeave( const FPointerEvent& MouseEvent )
 {
-	if ( !bIsResizing )
+	int32 PreviousHoveredHandleIndex = HoveredHandleIndex;
+
+	if (!bIsResizing)
 	{
 		HoveredHandleIndex = INDEX_NONE;
 	}
+
+	if (HoveredHandleIndex != PreviousHoveredHandleIndex)
+	{
+		OnHandleHovered.ExecuteIfBound(HoveredHandleIndex);
+	}
 }
 
 /**
@@ -408,6 +432,13 @@ void SSplitter::OnMouseLeave( const FPointerEvent& MouseEvent )
 */
 FCursorReply SSplitter::OnCursorQuery( const FGeometry& MyGeometry, const FPointerEvent& CursorEvent ) const
 {
+	if (bIsResizing)
+	{
+		return Orientation == Orient_Horizontal
+			? FCursorReply::Cursor( EMouseCursor::ResizeLeftRight )
+			: FCursorReply::Cursor( EMouseCursor::ResizeUpDown );
+	}
+
 	const FVector2D LocalMousePosition = MyGeometry.AbsoluteToLocal( CursorEvent.GetScreenSpacePosition() );
 
 	TArray<FLayoutGeometry> LayoutChildren = ArrangeChildrenForLayout(MyGeometry);
@@ -417,9 +448,10 @@ FCursorReply SSplitter::OnCursorQuery( const FGeometry& MyGeometry, const FPoint
 		? GetHandleBeingResizedFromMousePosition<Orient_Horizontal>( PhysicalSplitterHandleSize, HitDetectionSplitterHandleSize, LocalMousePosition, LayoutChildren )
 		: GetHandleBeingResizedFromMousePosition<Orient_Vertical>( PhysicalSplitterHandleSize, HitDetectionSplitterHandleSize, LocalMousePosition, LayoutChildren );
 
-	if (CurrentHoveredHandleIndex != INDEX_NONE)
+	if (CurrentHoveredHandleIndex != INDEX_NONE &&
+		Children[CurrentHoveredHandleIndex].SizingRule.Get() != ESizeRule::SizeToContent)
 	{
-		return ( Orientation == Orient_Horizontal )
+		return Orientation == Orient_Horizontal
 			? FCursorReply::Cursor( EMouseCursor::ResizeLeftRight )
 			: FCursorReply::Cursor( EMouseCursor::ResizeUpDown );
 	}
@@ -463,7 +495,7 @@ int32 SSplitter::FindResizeableSlotBeforeHandle( int32 DraggedHandle, const TPan
 	// Resizing collapsed or autosizing slots does not make sense (their size is predetermined).
 	// Search out from the DraggedHandle to find the first non-collapsed, non-autosizing slot we can resize.
 	int32 SlotBeforeDragHandle = DraggedHandle;
-	while ( SlotBeforeDragHandle >= 0 && (Children[SlotBeforeDragHandle].GetWidget()->GetVisibility() == EVisibility::Collapsed || Children[SlotBeforeDragHandle].SizingRule.Get() == SSplitter::SizeToContent ) )
+	while ( SlotBeforeDragHandle >= 0 && (Children[SlotBeforeDragHandle].GetWidget()->GetVisibility() == EVisibility::Collapsed || Children[SlotBeforeDragHandle].SizingRule.Get() == SSplitter::SizeToContent) )
 	{
 		--SlotBeforeDragHandle;
 	}
@@ -514,22 +546,22 @@ void SSplitter::HandleResizingDelta(EOrientation SplitterOrientation, const floa
 	//  - Prev vs. Next refers to the widgets in the order they are laid out (left->right, top->bottom).
 	//  - New vs. Old refers to the Old values for width/height vs. the post-resize values.
 
-	TArray< int32 > SlotsAfterDragHandleIndicies;
+	TArray< int32 > SlotsAfterDragHandleIndices;
 	if (InResizeMode == ESplitterResizeMode::FixedPosition)
 	{
 		const int32 SlotAfterDragHandle = FindResizeableSlotAfterHandle( DraggedHandle, InChildren);
 
 		if ( SlotAfterDragHandle < NumChildren )
 		{
-			SlotsAfterDragHandleIndicies.Add( SlotAfterDragHandle );
+			SlotsAfterDragHandleIndices.Add( SlotAfterDragHandle );
 		}
 	}
 	else if (InResizeMode == ESplitterResizeMode::Fill || InResizeMode == ESplitterResizeMode::FixedSize)
 	{
-		FindAllResizeableSlotsAfterHandle( DraggedHandle, InChildren, /*OUT*/ SlotsAfterDragHandleIndicies );
+		FindAllResizeableSlotsAfterHandle( DraggedHandle, InChildren, /*OUT*/ SlotsAfterDragHandleIndices );
 	}
 
-	const int32 NumSlotsAfterDragHandle = SlotsAfterDragHandleIndicies.Num();
+	const int32 NumSlotsAfterDragHandle = SlotsAfterDragHandleIndices.Num();
 	if ( NumSlotsAfterDragHandle > 0 )
 	{
 		struct FSlotInfo 
@@ -544,8 +576,8 @@ void SSplitter::HandleResizingDelta(EOrientation SplitterOrientation, const floa
 		{
 			FSlotInfo SlotInfo;
 
-			SlotInfo.Slot = &InChildren[ SlotsAfterDragHandleIndicies[ SlotIndex ] ];
-			SlotInfo.Geometry = &ChildGeometries[ SlotsAfterDragHandleIndicies[ SlotIndex ] ];
+			SlotInfo.Slot = &InChildren[ SlotsAfterDragHandleIndices[ SlotIndex ] ];
+			SlotInfo.Geometry = &ChildGeometries[ SlotsAfterDragHandleIndices[ SlotIndex ] ];
 			SlotInfo.NewSize = ClampChild(SlotInfo.Geometry->GetSizeInParentSpace().Component( AxisIndex ));
 
 			SlotsAfterDragHandle.Add( SlotInfo );
@@ -558,7 +590,7 @@ void SSplitter::HandleResizingDelta(EOrientation SplitterOrientation, const floa
 
 		// Compute the new sizes of the children
 		const float PrevChildLength = PrevChildGeom.GetSizeInParentSpace().Component(AxisIndex);
-		float NewPrevChildLength = ClampChild( PrevChildLength + Delta );
+		float NewPrevChildLength = ClampChild(PrevChildLength + Delta);
 		Delta = NewPrevChildLength - PrevChildLength;
 
 		// Distribute the Delta across the affected slots after the drag handle
diff --git a/Engine/Source/Runtime/Slate/Public/Widgets/Layout/SSplitter.h b/Engine/Source/Runtime/Slate/Public/Widgets/Layout/SSplitter.h
--- a/Engine/Source/Runtime/Slate/Public/Widgets/Layout/SSplitter.h
+++ b/Engine/Source/Runtime/Slate/Public/Widgets/Layout/SSplitter.h
@@ -105,6 +105,7 @@ public:
 	 */
 	FSlot& AddSlot( int32 AtIndex = INDEX_NONE );
 
+	DECLARE_DELEGATE_OneParam(FOnHandleHovered, int32);
 
 	SLATE_BEGIN_ARGS(SSplitter)
 		: _Style( &FCoreStyle::Get().GetWidgetStyle<FSplitterStyle>("Splitter") )
@@ -132,9 +133,13 @@ public:
 
 		SLATE_ARGUMENT( float, MinimumSlotHeight )
 
+		SLATE_ATTRIBUTE( int32, HighlightedHandleIndex )
+
+		SLATE_EVENT( FOnHandleHovered, OnHandleHovered )
+
 		SLATE_EVENT( FSimpleDelegate, OnSplitterFinishedResizing )
 		
-		SLATE_EVENT(FOnGetMaxSlotSize, OnGetMaxSlotSize)
+		SLATE_EVENT( FOnGetMaxSlotSize, OnGetMaxSlotSize )
 
 	SLATE_END_ARGS()
 
@@ -232,8 +237,7 @@ public:
 	/**
 	 * @return the current orientation of the splitter.
 	 */
-	EOrientation GetOrientation() const;	
-
+	EOrientation GetOrientation() const;
 
 private:
 	TArray<FLayoutGeometry> ArrangeChildrenForLayout( const FGeometry& AllottedGeometry ) const;
@@ -290,12 +294,15 @@ protected:
 	TPanelChildren< FSlot > Children;
 
 	int32 HoveredHandleIndex;
+	TAttribute<int32> HighlightedHandleIndex;
 	bool bIsResizing;
 	EOrientation Orientation;
 	ESplitterResizeMode::Type ResizeMode;
 
 	FSimpleDelegate OnSplitterFinishedResizing;
 	FOnGetMaxSlotSize OnGetMaxSlotSize;
+	FOnHandleHovered OnHandleHovered;
+
 	/** The user is not allowed to make any of the splitter's children smaller than this. */
 	float MinSplitterChildLength;
 
