Merge keyboard fixes for HL2 (CL 14386403)

diff --git a/Engine/Source/Runtime/Slate/Private/Framework/Text/HoloLens/HoloLensPlatformTextField.cpp b/Engine/Source/Runtime/Slate/Private/Framework/Text/HoloLens/HoloLensPlatformTextField.cpp
--- a/Engine/Source/Runtime/Slate/Private/Framework/Text/HoloLens/HoloLensPlatformTextField.cpp
+++ b/Engine/Source/Runtime/Slate/Private/Framework/Text/HoloLens/HoloLensPlatformTextField.cpp
@@ -18,6 +18,11 @@ internal:
 
 	void CreateContext()
 	{
+		if (coreTextContext != nullptr)
+		{
+			return;
+		}
+
 		coreTextContext = CoreTextServicesManager::GetForCurrentView()->CreateEditContext();
 		coreTextContext->InputPaneDisplayPolicy = CoreTextInputPaneDisplayPolicy::Manual;
 
@@ -25,11 +30,11 @@ internal:
 		SelectionUpdatingToken = coreTextContext->SelectionUpdating += ref new TypedEventHandler<CoreTextEditContext^, CoreTextSelectionUpdatingEventArgs^>(this, &VirtualKeyboardInputContext::OnSelectionUpdating);
 		TextRequestedToken = coreTextContext->TextRequested += ref new TypedEventHandler<CoreTextEditContext^, CoreTextTextRequestedEventArgs^>(this, &VirtualKeyboardInputContext::OnTextRequested);
 		TextUpdatingToken = coreTextContext->TextUpdating += ref new TypedEventHandler<CoreTextEditContext^, CoreTextTextUpdatingEventArgs^>(this, &VirtualKeyboardInputContext::OnTextUpdating);
+		FormatUpdatingToken = coreTextContext->FormatUpdating += ref new TypedEventHandler<CoreTextEditContext^, CoreTextFormatUpdatingEventArgs^>(this, &VirtualKeyboardInputContext::OnFormatUpdating);
 
 		KeyDownCoreWindowToken = CoreWindow::GetForCurrentThread()->KeyDown += ref new TypedEventHandler<CoreWindow^, KeyEventArgs^>(this, &VirtualKeyboardInputContext::OnKeyDownCoreWindow);
 
 		TextEntryWidget->OnSelectionChanged.BindLambda([this]() { OnSelectionChanged(); });
-
 	}
 
 	void DeleteContext()
@@ -42,6 +47,8 @@ internal:
 		coreTextContext->SelectionUpdating -= SelectionUpdatingToken;
 		coreTextContext->TextRequested -= TextRequestedToken;
 		coreTextContext->TextUpdating -= TextUpdatingToken;
+		coreTextContext->FormatUpdating -= FormatUpdatingToken;
+
 		coreTextContext = nullptr;
 	}
 
@@ -84,6 +91,11 @@ internal:
 		TextEntryWidget->SetSelectionFromVirtualKeyboard(args->NewSelection.StartCaretPosition, args->NewSelection.EndCaretPosition);
 	}
 
+	// Needed for autocomplete
+	void OnFormatUpdating(CoreTextEditContext^, CoreTextFormatUpdatingEventArgs^ args)
+	{
+	}
+
 	void OnSelectionChanged()
 	{
 		int start, end;
@@ -166,39 +178,16 @@ internal:
 			start = end = text.Len();
 		}
 
-		auto shift = CoreWindow::GetForCurrentThread()->GetKeyState(VirtualKey::Shift);
-
-		if (shift == CoreVirtualKeyStates::Down || shift == CoreVirtualKeyStates::Locked)
-		{
-			if (start == end)
-			{
-				extendingSelectionLeft = true;
-			}
-
-			if (extendingSelectionLeft)
-			{
-				start = FMath::Max(0, start - 1);
-			}
-			else
-			{
-				end = FMath::Min(oldLen, end - 1);
-			}
-		}
-		else
+		if (start == end)
 		{
-			//no shift block
-			if (start == end)
+			if (start == 0)
 			{
-				if (start == 0)
-				{
-					return;
-				}
-
-				--start;
+				return;
 			}
 
-			end = start;
+			--start;
 		}
+		end = start;
 
 		CoreTextRange selRange;
 		selRange.StartCaretPosition = start;
@@ -222,38 +211,16 @@ internal:
 			start = end = text.Len();
 		}
 
-		auto shift = CoreWindow::GetForCurrentThread()->GetKeyState(VirtualKey::Shift);
-
-		if (shift == CoreVirtualKeyStates::Down || shift == CoreVirtualKeyStates::Locked)
+		if (start == end)
 		{
-			if (start == end)
+			if (end == oldLen)
 			{
-				extendingSelectionLeft = false;
-			}
-
-			if (extendingSelectionLeft)
-			{
-				start = FMath::Max(0, start + 1);
-			}
-			else
-			{
-				end = FMath::Min(oldLen, end + 1);
+				return;
 			}
-		}
-		else
-		{
-			//no shift block
-			if (start == end)
-			{
-				if (end == oldLen)
-				{
-					return;
-				}
 
-				++end;
-			}
-			start = end;
+			++end;
 		}
+		start = end;
 
 		CoreTextRange selRange;
 		selRange.StartCaretPosition = start;
@@ -285,11 +252,11 @@ private:
 	Windows::Foundation::EventRegistrationToken SelectionUpdatingToken;
 	Windows::Foundation::EventRegistrationToken TextRequestedToken;
 	Windows::Foundation::EventRegistrationToken TextUpdatingToken;
+	Windows::Foundation::EventRegistrationToken FormatUpdatingToken;
 	Windows::Foundation::EventRegistrationToken KeyDownCoreWindowToken;
 
 	Windows::UI::Text::Core::CoreTextEditContext^ coreTextContext = nullptr;
 	TSharedPtr<IVirtualKeyboardEntry> TextEntryWidget;
-	bool extendingSelectionLeft = false;
 };
 
 FHoloLensPlatformTextField::FHoloLensPlatformTextField()
