Slowtask dialog and progress bar update

diff --git a/Engine/Source/Editor/UnrealEd/Private/FeedbackContextEditor.cpp b/Engine/Source/Editor/UnrealEd/Private/FeedbackContextEditor.cpp
--- a/Engine/Source/Editor/UnrealEd/Private/FeedbackContextEditor.cpp
+++ b/Engine/Source/Editor/UnrealEd/Private/FeedbackContextEditor.cpp
@@ -20,6 +20,7 @@
 #include "Engine/Engine.h"
 #include "StudioAnalytics.h"
 #include "AnalyticsEventAttribute.h"
+#include "Styling/StarshipCoreStyle.h"
 
 /** Called to cancel the slow task activity */
 DECLARE_DELEGATE( FOnCancelClickedDelegate );
@@ -33,7 +34,7 @@ DECLARE_DELEGATE_RetVal(bool, FReceiveUserCancelDelegate);
 class SSlowTaskWidget : public SBorder
 {
 	/** The maximum number of secondary bars to show on the widget */
-	static const int32 MaxNumSecondaryBars = 3;
+	static const int32 MaxNumSecondaryBars = 1;
 
 	/** The width of the dialog, and horizontal padding */
 	static const int32 FixedWidth = 600, FixedPaddingH = 24;
@@ -84,21 +85,21 @@ public:
 
 						+ SHorizontalBox::Slot()
 						[
-							SNew( STextBlock )
+							SNew(STextBlock)
 							.AutoWrapText(true)
-							.Text( this, &SSlowTaskWidget::GetProgressText, 0 )
+							.Text(this, &SSlowTaskWidget::GetProgressText, 0)
 							// The main font size dynamically changes depending on the content
-							.Font( this, &SSlowTaskWidget::GetMainTextFont )
+							.Font(this, &SSlowTaskWidget::GetMainTextFont)
 						]
 
 						+ SHorizontalBox::Slot()
 						.Padding(FMargin(5.f, 0, 0, 0))
 						.AutoWidth()
 						[
-							SNew( STextBlock )
-							.Text( this, &SSlowTaskWidget::GetPercentageText )
+							SNew(STextBlock )
+							.Text(this, &SSlowTaskWidget::GetPercentageText)
 							// The main font size dynamically changes depending on the content
-							.Font( FCoreStyle::GetDefaultFontStyle("Light", 14) )
+							.Font(FStarshipCoreStyle::GetDefaultFontStyle("NormalFont", 14))
 						]
 					]
 				]
@@ -111,9 +112,7 @@ public:
 					[
 						SNew(SProgressBar)
 						.BorderPadding(FVector2D::ZeroVector)
-						.Percent( this, &SSlowTaskWidget::GetProgressFraction, 0 )
-						.BackgroundImage( FEditorStyle::GetBrush("ProgressBar.ThinBackground") )
-						.FillImage( FEditorStyle::GetBrush("ProgressBar.ThinFill") )
+						.Percent(this, &SSlowTaskWidget::GetProgressFraction, 0)
 					]
 				]
 			]
@@ -143,7 +142,7 @@ public:
 		}
 
 		SBorder::Construct( SBorder::FArguments()
-			.BorderImage(FEditorStyle::GetBrush("Menu.Background"))
+			.BorderImage(FEditorStyle::GetBrush("Brushes.Header"))
 			.VAlign(VAlign_Center)
 			.Padding(FMargin(FixedPaddingH))
 			[
@@ -155,6 +154,11 @@ public:
 			]
 		);
 
+		for (int SecondaryBarIndex = 0; SecondaryBarIndex < MaxNumSecondaryBars; ++SecondaryBarIndex)
+		{
+			CreateSecondaryBar(SecondaryBarIndex);
+		}
+
 		// Make sure all our bars are set up
 		UpdateDynamicProgressBars();
 	}
@@ -205,12 +209,6 @@ private:
 			}
 		}
 
-		// Create progress bars for anything that we haven't cached yet
-		// We don't destroy old widgets, they just remain ghosted until shown again
-		for (int32 Index = SecondaryBars->GetChildren()->Num() + 1; Index < DynamicProgressIndices.Num(); ++Index)
-		{
-			CreateSecondaryBar(Index);
-		}
 	}
 
 	/** Create a progress bar for the specified index */
@@ -227,7 +225,6 @@ private:
 			[
 				SNew( STextBlock )
 				.Text( this, &SSlowTaskWidget::GetProgressText, Index )
-				.Font( FCoreStyle::GetDefaultFontStyle("Regular", 9) )
 				.ColorAndOpacity( FSlateColor::UseSubduedForeground() )
 			]
 
@@ -237,17 +234,9 @@ private:
 				SNew(SBox)
 				.HeightOverride(SecondaryBarHeight)
 				[
-					SNew(SBorder)
-					.Padding(0)
-					.BorderImage(FEditorStyle::GetBrush("NoBorder"))
-					.ColorAndOpacity( this, &SSlowTaskWidget::GetSecondaryProgressBarTint, Index )
-					[
-						SNew(SProgressBar)
-						.BorderPadding(FVector2D::ZeroVector)
-						.Percent( this, &SSlowTaskWidget::GetProgressFraction, Index )
-						.BackgroundImage( FEditorStyle::GetBrush("ProgressBar.ThinBackground") )
-						.FillImage( FEditorStyle::GetBrush("ProgressBar.ThinFill") )
-					]
+					SNew(SProgressBar)
+					.BorderPadding(FVector2D::ZeroVector)
+					.Percent( this, &SSlowTaskWidget::GetProgressFraction, Index )
 				]
 			]
 		];
@@ -273,7 +262,7 @@ private:
 		TSharedRef<FSlateFontMeasure> MeasureService = FSlateApplication::Get().GetRenderer()->GetFontMeasureService();
 
 		const int32 MaxFontSize = 14;
-		FSlateFontInfo FontInfo = FCoreStyle::GetDefaultFontStyle("Light", MaxFontSize);
+		FSlateFontInfo FontInfo = FStarshipCoreStyle::GetDefaultFontStyle("NormalFont", MaxFontSize);
 
 		const FText MainText = GetProgressText(0);
 		const int32 MaxTextWidth = FixedWidth - FixedPaddingH*2;
@@ -330,7 +319,7 @@ private:
 
 	EVisibility GetSecondaryBarVisibility(int32 Index) const
 	{
-		return DynamicProgressIndices.IsValidIndex(Index) ? EVisibility::HitTestInvisible : EVisibility::Collapsed;
+		return DynamicProgressIndices.IsValidIndex(Index) ? EVisibility::HitTestInvisible : EVisibility::Hidden;
 	}
 
 	/** Called when the cancel button is clicked */
diff --git a/Engine/Source/Runtime/Slate/Private/Widgets/Notifications/SProgressBar.cpp b/Engine/Source/Runtime/Slate/Private/Widgets/Notifications/SProgressBar.cpp
--- a/Engine/Source/Runtime/Slate/Private/Widgets/Notifications/SProgressBar.cpp
+++ b/Engine/Source/Runtime/Slate/Private/Widgets/Notifications/SProgressBar.cpp
@@ -52,6 +52,7 @@ void SProgressBar::SetStyle(const FProgressBarStyle* InStyle)
 
 	check(Style);
 
+	UpdateMarqueeActiveTimer();
 	Invalidate(EInvalidateWidget::Layout);
 }
 
@@ -170,7 +171,7 @@ int32 SProgressBar::OnPaint( const FPaintArgs& Args, const FGeometry& AllottedGe
 		CurrentBackgroundImage,
 		DrawEffects,
 		InWidgetStyle.GetColorAndOpacityTint() * CurrentBackgroundImage->GetTint( InWidgetStyle )
-	);	
+	);
 	
 	if( ProgressFraction.IsSet() )
 	{
@@ -195,13 +196,22 @@ int32 SProgressBar::OnPaint( const FPaintArgs& Args, const FGeometry& AllottedGe
 			{
 				if (PushTransformedClip(OutDrawElements, AllottedGeometry, BorderPaddingRef, FVector2D(1, 0), FSlateRect(ClampedFraction, 0, 0, 1)))
 				{
+					float MarqueeAnimOffset = 0;
+					float MarqueeImageSize = 0;
+
+					if (Style->EnableFillAnimation)
+					{
+						MarqueeAnimOffset = CurrentFillImage->ImageSize.X * MarqueeOffset;
+						MarqueeImageSize = CurrentFillImage->ImageSize.X;
+					}
+
 					// Draw Fill
 					FSlateDrawElement::MakeBox(
 						OutDrawElements,
 						RetLayerId++,
 						AllottedGeometry.ToPaintGeometry(
-							FVector2D::ZeroVector,
-							FVector2D( AllottedGeometry.GetLocalSize().X, AllottedGeometry.GetLocalSize().Y )),
+							FVector2D(MarqueeAnimOffset - MarqueeImageSize, 0.0f),
+							FVector2D(AllottedGeometry.GetLocalSize().X+MarqueeImageSize, AllottedGeometry.GetLocalSize().Y)),
 						CurrentFillImage,
 						DrawEffects,
 						FillColorAndOpacitySRGB
@@ -221,8 +231,8 @@ int32 SProgressBar::OnPaint( const FPaintArgs& Args, const FGeometry& AllottedGe
 						OutDrawElements,
 						RetLayerId++,
 						AllottedGeometry.ToPaintGeometry(
-							FVector2D( (AllottedGeometry.GetLocalSize().X * 0.5f) - ((AllottedGeometry.GetLocalSize().X * ( ClampedFraction ))*0.5f), 0.0f),
-							FVector2D( AllottedGeometry.GetLocalSize().X * ( ClampedFraction ) , AllottedGeometry.GetLocalSize().Y )),
+							FVector2D((AllottedGeometry.GetLocalSize().X * 0.5f) - ((AllottedGeometry.GetLocalSize().X * ( ClampedFraction ))*0.5f), 0.0f),
+							FVector2D(AllottedGeometry.GetLocalSize().X * ( ClampedFraction ) , AllottedGeometry.GetLocalSize().Y )),
 						CurrentFillImage,
 						DrawEffects,
 						FillColorAndOpacitySRGB
@@ -236,13 +246,22 @@ int32 SProgressBar::OnPaint( const FPaintArgs& Args, const FGeometry& AllottedGe
 			{
 				if (PushTransformedClip(OutDrawElements, AllottedGeometry, BorderPaddingRef, FVector2D(0, 0), FSlateRect(0, 0, 1, ClampedFraction)))
 				{
+					float MarqueeAnimOffset = 0;
+					float MarqueeImageSize = 0;
+
+					if (Style->EnableFillAnimation)
+					{
+						MarqueeAnimOffset = CurrentFillImage->ImageSize.Y * MarqueeOffset;
+						MarqueeImageSize = CurrentFillImage->ImageSize.Y;
+					}
+
 					// Draw Fill
 					FSlateDrawElement::MakeBox(
 						OutDrawElements,
 						RetLayerId++,
 						AllottedGeometry.ToPaintGeometry(
-							FVector2D::ZeroVector,
-							FVector2D( AllottedGeometry.GetLocalSize().X, AllottedGeometry.GetLocalSize().Y )),
+							FVector2D(0.0f, MarqueeAnimOffset - MarqueeImageSize),
+							FVector2D(AllottedGeometry.GetLocalSize().X, AllottedGeometry.GetLocalSize().Y + MarqueeImageSize)),
 						CurrentFillImage,
 						DrawEffects,
 						FillColorAndOpacitySRGB
@@ -256,16 +275,22 @@ int32 SProgressBar::OnPaint( const FPaintArgs& Args, const FGeometry& AllottedGe
 			{
 				if (PushTransformedClip(OutDrawElements, AllottedGeometry, BorderPaddingRef, FVector2D(0, 1), FSlateRect(0, ClampedFraction, 1, 0)))
 				{
-					FSlateRect ClippedAllotedGeometry = FSlateRect(AllottedGeometry.AbsolutePosition, AllottedGeometry.AbsolutePosition + AllottedGeometry.GetLocalSize() * AllottedGeometry.Scale);
-					ClippedAllotedGeometry.Top = ClippedAllotedGeometry.Bottom - ClippedAllotedGeometry.GetSize().Y * ClampedFraction;
+					float MarqueeAnimOffset = 0;
+					float MarqueeImageSize = 0;
+
+					if (Style->EnableFillAnimation)
+					{
+						MarqueeAnimOffset = CurrentFillImage->ImageSize.Y * MarqueeOffset;
+						MarqueeImageSize = CurrentFillImage->ImageSize.Y;
+					}
 
 					// Draw Fill
 					FSlateDrawElement::MakeBox(
 						OutDrawElements,
 						RetLayerId++,
 						AllottedGeometry.ToPaintGeometry(
-							FVector2D::ZeroVector,
-							FVector2D( AllottedGeometry.GetLocalSize().X, AllottedGeometry.GetLocalSize().Y )),
+							FVector2D(0.0f, MarqueeAnimOffset - MarqueeImageSize),
+							FVector2D( AllottedGeometry.GetLocalSize().X, AllottedGeometry.GetLocalSize().Y + MarqueeImageSize)),
 						CurrentFillImage,
 						DrawEffects,
 						FillColorAndOpacitySRGB
@@ -278,22 +303,33 @@ int32 SProgressBar::OnPaint( const FPaintArgs& Args, const FGeometry& AllottedGe
 			case EProgressBarFillType::LeftToRight:
 			default:
 			{
+				float MarqueeAnimOffset = 0;
+				float MarqueeImageSize = 0;
+
+				if(Style->EnableFillAnimation)
+				{
+					MarqueeAnimOffset = CurrentFillImage->ImageSize.X * MarqueeOffset;
+					MarqueeImageSize = CurrentFillImage->ImageSize.X;
+				}
+
 				if (PushTransformedClip(OutDrawElements, AllottedGeometry, BorderPaddingRef, FVector2D(0, 0), FSlateRect(0, 0, ClampedFraction, 1)))
 				{
-					// Draw Fill
+
+					// Draw fill
 					FSlateDrawElement::MakeBox(
 						OutDrawElements,
 						RetLayerId++,
 						AllottedGeometry.ToPaintGeometry(
-							FVector2D::ZeroVector,
-							FVector2D( AllottedGeometry.GetLocalSize().X, AllottedGeometry.GetLocalSize().Y )),
+							FVector2D(MarqueeAnimOffset - MarqueeImageSize, 0.0f),
+							FVector2D(AllottedGeometry.GetLocalSize().X+ MarqueeImageSize, AllottedGeometry.GetLocalSize().Y)),
 						CurrentFillImage,
 						DrawEffects,
 						FillColorAndOpacitySRGB
-						);
-
+					);
 					OutDrawElements.PopClip();
 				}
+
+
 				break;
 			}
 		}
@@ -316,7 +352,7 @@ int32 SProgressBar::OnPaint( const FPaintArgs& Args, const FGeometry& AllottedGe
 					FVector2D( AllottedGeometry.GetLocalSize().X + MarqueeImageSize, AllottedGeometry.GetLocalSize().Y )),
 				CurrentMarqueeImage,
 				DrawEffects,
-				CurrentMarqueeImage->TintColor.GetSpecifiedColor() * ColorAndOpacitySRGB
+				ColorAndOpacitySRGB
 				);
 
 			OutDrawElements.PopClip();
@@ -359,7 +395,7 @@ void SProgressBar::UpdateMarqueeActiveTimer()
 		UnRegisterActiveTimer(ActiveTimerHandle.Pin().ToSharedRef());
 	}
 
-	if (!Percent.IsBound() && !Percent.Get().IsSet())
+	if ((!Percent.IsBound() && !Percent.Get().IsSet()) || Style->EnableFillAnimation)
 	{
 		// If percent is not bound or set then its marquee. Set the timer
 		ActiveTimerHandle = RegisterActiveTimer(CurrentTickRate, FWidgetActiveTimerDelegate::CreateSP(this, &SProgressBar::ActiveTick));
@@ -371,7 +407,7 @@ EActiveTimerReturnType SProgressBar::ActiveTick(double InCurrentTime, float InDe
 	MarqueeOffset = (float)(InCurrentTime - FMath::FloorToDouble(InCurrentTime));
 	
 	TOptional<float> PercentFraction = Percent.Get();
-	if (PercentFraction.IsSet())
+	if (PercentFraction.IsSet() && !Style->EnableFillAnimation)
 	{
 		SetActiveTimerTickRate(MinimumTickRate);
 	}
diff --git a/Engine/Source/Runtime/Slate/Public/Widgets/Notifications/SProgressBar.h b/Engine/Source/Runtime/Slate/Public/Widgets/Notifications/SProgressBar.h
--- a/Engine/Source/Runtime/Slate/Public/Widgets/Notifications/SProgressBar.h
+++ b/Engine/Source/Runtime/Slate/Public/Widgets/Notifications/SProgressBar.h
@@ -44,7 +44,7 @@ class SLATE_API SProgressBar : public SLeafWidget
 
 public:
 	SLATE_BEGIN_ARGS(SProgressBar)
-		: _Style( &FCoreStyle::Get().GetWidgetStyle<FProgressBarStyle>("ProgressBar") )
+		: _Style( &FAppStyle::Get().GetWidgetStyle<FProgressBarStyle>("ProgressBar") )
 		, _BarFillType(EProgressBarFillType::LeftToRight)
 		, _Percent( TOptional<float>() )
 		, _FillColorAndOpacity( FLinearColor::White )
diff --git a/Engine/Source/Runtime/SlateCore/Private/Styling/SlateTypes.cpp b/Engine/Source/Runtime/SlateCore/Private/Styling/SlateTypes.cpp
--- a/Engine/Source/Runtime/SlateCore/Private/Styling/SlateTypes.cpp
+++ b/Engine/Source/Runtime/SlateCore/Private/Styling/SlateTypes.cpp
@@ -305,6 +305,7 @@ const FInlineEditableTextBlockStyle& FInlineEditableTextBlockStyle::GetDefault()
 
 
 FProgressBarStyle::FProgressBarStyle()
+	: EnableFillAnimation(false)
 {
 }
 
diff --git a/Engine/Source/Runtime/SlateCore/Private/Styling/StarshipCoreStyle.cpp b/Engine/Source/Runtime/SlateCore/Private/Styling/StarshipCoreStyle.cpp
--- a/Engine/Source/Runtime/SlateCore/Private/Styling/StarshipCoreStyle.cpp
+++ b/Engine/Source/Runtime/SlateCore/Private/Styling/StarshipCoreStyle.cpp
@@ -464,12 +464,14 @@ TSharedRef<ISlateStyle> FStarshipCoreStyle::Create()
 	// SProgressBar defaults...
 	{
 		Style->Set("ProgressBar", FProgressBarStyle()
-			.SetBackgroundImage(BOX_BRUSH("Common/ProgressBar_Background", FMargin(5.f / 12.f)))
-			.SetFillImage(BOX_BRUSH("Common/ProgressBar_Fill", FMargin(5.f / 12.f), FLinearColor(1.0f, 0.22f, 0.0f)))
-			.SetMarqueeImage(IMAGE_BRUSH("Common/ProgressBar_Marquee", FVector2D(20, 12), FLinearColor::White, ESlateBrushTileType::Horizontal))
+			.SetBackgroundImage(FSlateColorBrush(FStyleColors::Foldout))
+			.SetFillImage(IMAGE_BRUSH("Starship/CoreWidgets/ProgressBar/ProgressMarquee", FVector2D(20, 12), FStyleColors::Primary, ESlateBrushTileType::Horizontal))
+			.SetMarqueeImage(IMAGE_BRUSH("Starship/CoreWidgets/ProgressBar/ProgressMarquee", FVector2D(20, 12), FLinearColor::White, ESlateBrushTileType::Horizontal))
+			.SetEnableFillAnimation(true)
 		);
 	}
 
+
 	// SThrobber, SCircularThrobber defaults...
 	{
 		Style->Set("Throbber.Chunk", new IMAGE_BRUSH("Common/Throbber_Piece", FVector2D(16, 16)));
diff --git a/Engine/Source/Runtime/SlateCore/Public/Styling/SlateTypes.h b/Engine/Source/Runtime/SlateCore/Public/Styling/SlateTypes.h
--- a/Engine/Source/Runtime/SlateCore/Public/Styling/SlateTypes.h
+++ b/Engine/Source/Runtime/SlateCore/Public/Styling/SlateTypes.h
@@ -1027,6 +1027,12 @@ struct SLATECORE_API FProgressBarStyle : public FSlateWidgetStyle
 	FSlateBrush MarqueeImage;
 	FProgressBarStyle& SetMarqueeImage( const FSlateBrush& InMarqueeImage ){ MarqueeImage = InMarqueeImage; return *this; }
 
+	/** Enables a simple animation on the fill image to give the appearance that progress has not stalled. Disable this if you have a custom material which animates itself. 
+	 * This requires a pattern in your material or texture to give the appearance of movement.  A solid color will show no movement. */
+	UPROPERTY(EditAnywhere, BlueprintReadWrite, Category = Appearance)
+	bool EnableFillAnimation;
+	FProgressBarStyle& SetEnableFillAnimation(bool InEnableFillAnimation) { EnableFillAnimation = InEnableFillAnimation; return *this; }
+
 	/**
 	* Unlinks all colors in this style.
 	* @see FSlateColor::Unlink
