Slate: Call SWidget's event when the input is synthetic. Look like the issue was introduced by VREditor (merge issue). It may cause the OnMouseLeave event to be fired without the corresponding OnMouseEnter when we are capturing.

diff --git a/Engine/Source/Runtime/Slate/Private/Framework/Application/SlateApplication.cpp b/Engine/Source/Runtime/Slate/Private/Framework/Application/SlateApplication.cpp
--- a/Engine/Source/Runtime/Slate/Private/Framework/Application/SlateApplication.cpp
+++ b/Engine/Source/Runtime/Slate/Private/Framework/Application/SlateApplication.cpp
@@ -5008,9 +5008,9 @@ bool FSlateApplication::RoutePointerMoveEvent(const FWidgetPath& WidgetsUnderPoi
 		for (int32 WidgetIndex = LastWidgetsUnderPointer.Widgets.Num()-1; WidgetIndex >=0; --WidgetIndex)
 		{
 			// Guards for cases where WidgetIndex can become invalid due to MouseMove being re-entrant.
-			while (WidgetIndex >= LastWidgetsUnderPointer.Widgets.Num())
+			if (WidgetIndex >= LastWidgetsUnderPointer.Widgets.Num())
 			{
-				WidgetIndex--;
+				WidgetIndex = LastWidgetsUnderPointer.Widgets.Num() - 1;
 			}
 
 			if (WidgetIndex >= 0)
@@ -5055,68 +5055,65 @@ bool FSlateApplication::RoutePointerMoveEvent(const FWidgetPath& WidgetsUnderPoi
 
 	if (MouseCaptorPath.IsValid())
 	{
-		if (!bIsSynthetic)
-		{
-			// Switch worlds widgets in the current path
-			FScopedSwitchWorldHack SwitchWorld(MouseCaptorPath);
+		// Switch worlds widgets in the current path
+		FScopedSwitchWorldHack SwitchWorld(MouseCaptorPath);
 
-			FEventRouter::Route<FNoReply>(this, FEventRouter::FBubblePolicy(WidgetsUnderPointer), PointerEvent, [&MouseCaptorPath, &LastWidgetsUnderPointer](const FArrangedWidget& WidgetUnderCursor, const FPointerEvent& Event)
+		FEventRouter::Route<FNoReply>(this, FEventRouter::FBubblePolicy(WidgetsUnderPointer), PointerEvent, [&MouseCaptorPath, &LastWidgetsUnderPointer](const FArrangedWidget& WidgetUnderCursor, const FPointerEvent& Event)
+			{
+				if (!LastWidgetsUnderPointer.ContainsWidget(WidgetUnderCursor.Widget))
 				{
-					if (!LastWidgetsUnderPointer.ContainsWidget(WidgetUnderCursor.Widget))
+					if (MouseCaptorPath.ContainsWidget(WidgetUnderCursor.Widget))
 					{
-						if (MouseCaptorPath.ContainsWidget(WidgetUnderCursor.Widget))
-						{
-							WidgetUnderCursor.Widget->OnMouseEnter(WidgetUnderCursor.Geometry, Event);
+						WidgetUnderCursor.Widget->OnMouseEnter(WidgetUnderCursor.Geometry, Event);
 #if WITH_SLATE_DEBUGGING
-							FSlateDebugging::BroadcastNoReplyInputEvent(ESlateDebuggingInputEvent::MouseEnter, WidgetUnderCursor.Widget);
+						FSlateDebugging::BroadcastNoReplyInputEvent(ESlateDebuggingInputEvent::MouseEnter, WidgetUnderCursor.Widget);
 #endif
-						}
 					}
-					return FNoReply();
-				}, ESlateDebuggingInputEvent::MouseEnter);
+				}
+				return FNoReply();
+			}, ESlateDebuggingInputEvent::MouseEnter);
 
-			FReply Reply = FEventRouter::Route<FReply>(this, FEventRouter::FToLeafmostPolicy(MouseCaptorPath), PointerEvent, [this](const FArrangedWidget& MouseCaptorWidget, const FPointerEvent& Event)
-				{
-					FReply TempReply = FReply::Unhandled();
+		FReply Reply = FEventRouter::Route<FReply>(this, FEventRouter::FToLeafmostPolicy(MouseCaptorPath), PointerEvent, [this, bIsSynthetic](const FArrangedWidget& MouseCaptorWidget, const FPointerEvent& Event)
+			{
+				FReply TempReply = FReply::Unhandled();
 
-					bool bAllowMouseFallback = true;
-					if (Event.IsTouchEvent())
+				bool bAllowMouseFallback = true;
+				if (Event.IsTouchEvent())
+				{
+					if (Event.IsTouchForceChangedEvent())
 					{
-						if (Event.IsTouchForceChangedEvent())
-						{
-							TempReply = MouseCaptorWidget.Widget->OnTouchForceChanged(MouseCaptorWidget.Geometry, Event);
-#if WITH_SLATE_DEBUGGING
-							FSlateDebugging::BroadcastInputEvent(ESlateDebuggingInputEvent::TouchForceChanged, TempReply, MouseCaptorWidget.Widget);
-#endif
-							bAllowMouseFallback = false;
-						}
-						else if (Event.IsTouchFirstMoveEvent())
-						{
-							TempReply = MouseCaptorWidget.Widget->OnTouchFirstMove(MouseCaptorWidget.Geometry, Event);
+						TempReply = MouseCaptorWidget.Widget->OnTouchForceChanged(MouseCaptorWidget.Geometry, Event);
 #if WITH_SLATE_DEBUGGING
-							FSlateDebugging::BroadcastInputEvent(ESlateDebuggingInputEvent::TouchFirstMove, TempReply, MouseCaptorWidget.Widget);
+						FSlateDebugging::BroadcastInputEvent(ESlateDebuggingInputEvent::TouchForceChanged, TempReply, MouseCaptorWidget.Widget);
 #endif
-							bAllowMouseFallback = false;
-						}
-						else
-						{
-							TempReply = MouseCaptorWidget.Widget->OnTouchMoved(MouseCaptorWidget.Geometry, Event);
+						bAllowMouseFallback = false;
+					}
+					else if (Event.IsTouchFirstMoveEvent())
+					{
+						TempReply = MouseCaptorWidget.Widget->OnTouchFirstMove(MouseCaptorWidget.Geometry, Event);
 #if WITH_SLATE_DEBUGGING
-							FSlateDebugging::BroadcastInputEvent(ESlateDebuggingInputEvent::TouchMoved, TempReply, MouseCaptorWidget.Widget);
+						FSlateDebugging::BroadcastInputEvent(ESlateDebuggingInputEvent::TouchFirstMove, TempReply, MouseCaptorWidget.Widget);
 #endif
-						}
+						bAllowMouseFallback = false;
 					}
-					if ((!Event.IsTouchEvent() && bAllowMouseFallback) || (!TempReply.IsEventHandled() && this->bTouchFallbackToMouse))
+					else
 					{
-						TempReply = MouseCaptorWidget.Widget->OnMouseMove(MouseCaptorWidget.Geometry, Event);
+						TempReply = MouseCaptorWidget.Widget->OnTouchMoved(MouseCaptorWidget.Geometry, Event);
 #if WITH_SLATE_DEBUGGING
-						FSlateDebugging::BroadcastInputEvent(ESlateDebuggingInputEvent::MouseMove, TempReply, MouseCaptorWidget.Widget);
+						FSlateDebugging::BroadcastInputEvent(ESlateDebuggingInputEvent::TouchMoved, TempReply, MouseCaptorWidget.Widget);
 #endif
 					}
-					return TempReply;
-				}, ESlateDebuggingInputEvent::MouseEnter);
-			bHandled = Reply.IsEventHandled();
-		}
+				}
+				if ((!Event.IsTouchEvent() && bAllowMouseFallback) || (!TempReply.IsEventHandled() && this->bTouchFallbackToMouse))
+				{
+					TempReply = MouseCaptorWidget.Widget->OnMouseMove(MouseCaptorWidget.Geometry, Event);
+#if WITH_SLATE_DEBUGGING
+					FSlateDebugging::BroadcastInputEvent(ESlateDebuggingInputEvent::MouseMove, TempReply, MouseCaptorWidget.Widget);
+#endif
+				}
+				return TempReply;
+			}, ESlateDebuggingInputEvent::MouseEnter);
+		bHandled = Reply.IsEventHandled();
 	}
 	else
 	{
