- Removed remaining vestiges of PLATFORM_TCHAR_IS_1_BYTE

diff --git a/Engine/Source/Runtime/Core/Public/HAL/Platform.h b/Engine/Source/Runtime/Core/Public/HAL/Platform.h
--- a/Engine/Source/Runtime/Core/Public/HAL/Platform.h
+++ b/Engine/Source/Runtime/Core/Public/HAL/Platform.h
@@ -247,9 +247,6 @@
 #ifndef PLATFORM_COMPILER_HAS_FOLD_EXPRESSIONS
 	#define PLATFORM_COMPILER_HAS_FOLD_EXPRESSIONS 0
 #endif
-#ifndef PLATFORM_TCHAR_IS_1_BYTE
-	#define PLATFORM_TCHAR_IS_1_BYTE			0
-#endif
 #ifndef PLATFORM_TCHAR_IS_4_BYTES
 	#define PLATFORM_TCHAR_IS_4_BYTES			0
 #endif
diff --git a/Engine/Source/Runtime/Slate/Private/Framework/Text/SlatePasswordRun.cpp b/Engine/Source/Runtime/Slate/Private/Framework/Text/SlatePasswordRun.cpp
--- a/Engine/Source/Runtime/Slate/Private/Framework/Text/SlatePasswordRun.cpp
+++ b/Engine/Source/Runtime/Slate/Private/Framework/Text/SlatePasswordRun.cpp
@@ -173,11 +173,7 @@ FSlatePasswordRun::FSlatePasswordRun(const FSlatePasswordRun& Run)
 
 TCHAR FSlatePasswordRun::GetPasswordChar()
 {
-#if PLATFORM_TCHAR_IS_1_BYTE
-	return TEXT('*');
-#else
 	return TEXT('\u2022');
-#endif
 }
 
 FString FSlatePasswordRun::BuildPasswordString(const int32 InLength)
