Fix crash when clicking pie when a sidebar drawer is open

diff --git a/Engine/Source/Runtime/Slate/Private/Framework/Docking/STabDrawer.cpp b/Engine/Source/Runtime/Slate/Private/Framework/Docking/STabDrawer.cpp
--- a/Engine/Source/Runtime/Slate/Private/Framework/Docking/STabDrawer.cpp
+++ b/Engine/Source/Runtime/Slate/Private/Framework/Docking/STabDrawer.cpp
@@ -387,26 +387,33 @@ void STabDrawer::OnGlobalFocusChanging(const FFocusEvent& FocusEvent, const FWea
 		// Do not close due to slow tasks as those opening send window activation events
 		if (!GIsSlowTask && !FSlateApplication::Get().GetActiveModalWindow().IsValid() && !IsLegalWidgetFocused(NewFocusedWidgetPath, MakeArrayView(LegalFocusWidgets)))
 		{
-			TSharedRef<SWindow> NewWindow = NewFocusedWidgetPath.GetWindow();
-			TSharedPtr<SWindow> MyWindow = FSlateApplication::Get().FindWidgetWindow(ThisWidget);
-
-			if (!NewWindow->IsDescendantOf(MyWindow))
+			if(NewFocusedWidgetPath.IsValid())
 			{
-				if (TSharedPtr<SWidget> MenuHost = FSlateApplication::Get().GetMenuHostWidget())
-				{
-					FWidgetPath MenuHostPath;
+				TSharedRef<SWindow> NewWindow = NewFocusedWidgetPath.GetWindow();
+				TSharedPtr<SWindow> MyWindow = FSlateApplication::Get().FindWidgetWindow(ThisWidget);
 
-					// See if the menu being opened is owned by the drawer contents and if so the menu should not be dismissed
-					FSlateApplication::Get().GeneratePathToWidgetUnchecked(MenuHost.ToSharedRef(), MenuHostPath);
-					if (!MenuHostPath.ContainsWidget(ChildSlot.GetWidget()))
+				if (!NewWindow->IsDescendantOf(MyWindow))
+				{
+					if (TSharedPtr<SWidget> MenuHost = FSlateApplication::Get().GetMenuHostWidget())
+					{
+						FWidgetPath MenuHostPath;
+
+						// See if the menu being opened is owned by the drawer contents and if so the menu should not be dismissed
+						FSlateApplication::Get().GeneratePathToWidgetUnchecked(MenuHost.ToSharedRef(), MenuHostPath);
+						if (!MenuHostPath.ContainsWidget(ChildSlot.GetWidget()))
+						{
+							bShouldLoseFocus = true;
+						}
+					}
+					else
 					{
 						bShouldLoseFocus = true;
 					}
 				}
-				else
-				{
-					bShouldLoseFocus = true;
-				}
+			}
+			else
+			{
+				bShouldLoseFocus = true;
 			}
 		}
 
