Slate: Bad merge. Missing STableRow from CL 12164680 introduced by CL 13258546.
#rnx

diff --git a/Engine/Source/Runtime/Slate/Public/Widgets/Views/STableRow.h b/Engine/Source/Runtime/Slate/Public/Widgets/Views/STableRow.h
--- a/Engine/Source/Runtime/Slate/Public/Widgets/Views/STableRow.h
+++ b/Engine/Source/Runtime/Slate/Public/Widgets/Views/STableRow.h
@@ -1335,7 +1335,7 @@ protected:
 		for( int32 ColumnIndex = 0; ColumnIndex < NumColumns; ++ColumnIndex )
 		{
 			const SHeaderRow::FColumn& Column = Columns[ColumnIndex];
-			if (Column.ShouldGenerateWidget.Get(true))
+			if (InColumnHeaders->ShouldGeneratedColumn(Column.ColumnId))
 			{
 				TSharedRef< SWidget >* ExistingWidget = ColumnIdToSlotContents.Find(Column.ColumnId);
 				TSharedRef< SWidget > CellContents = SNullWidget::NullWidget;
