Add a Failure return value in ScrollIntoView to better handle the case where a widget is invalid when executing

diff --git a/Engine/Source/Runtime/Slate/Private/Widgets/Views/STableViewBase.cpp b/Engine/Source/Runtime/Slate/Private/Widgets/Views/STableViewBase.cpp
--- a/Engine/Source/Runtime/Slate/Private/Widgets/Views/STableViewBase.cpp
+++ b/Engine/Source/Runtime/Slate/Private/Widgets/Views/STableViewBase.cpp
@@ -334,15 +334,14 @@ void STableViewBase::Tick( const FGeometry& AllottedGeometry, const double InCur
 			ItemsPanel->SetRefreshPending(false);
 
 			Invalidate(EInvalidateWidget::ChildOrder);
-
-			const bool bScrollIntoViewSuccess = ScrollIntoViewResult == EScrollIntoViewResult::Success;
-			if (bScrollIntoViewSuccess)
+			
+			if (ScrollIntoViewResult == EScrollIntoViewResult::Success)
 			{
 				// Notify as soon as we've made a widget for the item, even if we still have scrolling to do
 				NotifyItemScrolledIntoView();
 			}
 
-			if (!bScrollIntoViewSuccess || CurrentScrollOffset != TargetScrollOffset)
+			if (ScrollIntoViewResult == EScrollIntoViewResult::Deferred || CurrentScrollOffset != TargetScrollOffset)
 			{
 				// Either we haven't made the item yet or we still have scrolling to do, so we'll need another refresh next frame
 				// We call this rather than just leave bItemsNeedRefresh as true to ensure that EnsureTickToRefresh is registered
diff --git a/Engine/Source/Runtime/Slate/Public/Widgets/Views/SListView.h b/Engine/Source/Runtime/Slate/Public/Widgets/Views/SListView.h
--- a/Engine/Source/Runtime/Slate/Public/Widgets/Views/SListView.h
+++ b/Engine/Source/Runtime/Slate/Public/Widgets/Views/SListView.h
@@ -1821,7 +1821,7 @@ protected:
 		}
 		else
 		{
-			return EScrollIntoViewResult::Deferred;
+			return EScrollIntoViewResult::Failure;
 		}
 		return EScrollIntoViewResult::Success;
 	}
diff --git a/Engine/Source/Runtime/Slate/Public/Widgets/Views/STableViewBase.h b/Engine/Source/Runtime/Slate/Public/Widgets/Views/STableViewBase.h
--- a/Engine/Source/Runtime/Slate/Public/Widgets/Views/STableViewBase.h
+++ b/Engine/Source/Runtime/Slate/Public/Widgets/Views/STableViewBase.h
@@ -316,6 +316,8 @@ protected:
 		Success,
 		/** The function did not have enough data to scroll the given item into view, so it should be deferred until the next Tick */
 		Deferred,
+		/** The function failed to scroll to the specified item.*/
+		Failure
 	};
 
 	/**
