Add custom horizontal alignment override to widgets in toolbars

diff --git a/Engine/Source/Runtime/Slate/Private/Framework/MultiBox/MultiBox.cpp b/Engine/Source/Runtime/Slate/Private/Framework/MultiBox/MultiBox.cpp
--- a/Engine/Source/Runtime/Slate/Private/Framework/MultiBox/MultiBox.cpp
+++ b/Engine/Source/Runtime/Slate/Private/Framework/MultiBox/MultiBox.cpp
@@ -660,9 +660,31 @@ void SMultiBoxWidget::AddBlockWidget( const FMultiBlock& Block, TSharedPtr<SHori
 	case EMultiBoxType::ToolBar:
 	case EMultiBoxType::SlimHorizontalToolBar:
 		{
-			HorizontalBox->AddSlot()
-			.AutoWidth()
-			.Padding(0)
+			EHorizontalAlignment HAlign;
+			EVerticalAlignment VAlign;
+			bool bAutoWidth;
+
+			bool bOverride = Block.GetAlignmentOverrides(HAlign, VAlign, bAutoWidth);
+
+			SHorizontalBox::FSlot& Slot = HorizontalBox->AddSlot();
+		
+			if(bOverride)
+			{
+				if (bAutoWidth)
+				{
+					Slot.AutoWidth();
+				}
+
+				Slot.HAlign(HAlign);
+				Slot.VAlign(VAlign);
+			}
+			else
+			{
+				Slot.AutoWidth();
+			}
+
+			Slot.Padding(0);
+			Slot
 			[
 				FinalWidgetWithHook
 			];
diff --git a/Engine/Source/Runtime/Slate/Private/Framework/MultiBox/MultiBoxBuilder.cpp b/Engine/Source/Runtime/Slate/Private/Framework/MultiBox/MultiBoxBuilder.cpp
--- a/Engine/Source/Runtime/Slate/Private/Framework/MultiBox/MultiBoxBuilder.cpp
+++ b/Engine/Source/Runtime/Slate/Private/Framework/MultiBox/MultiBoxBuilder.cpp
@@ -485,7 +485,7 @@ void FToolBarBuilder::AddToolBarWidget( TSharedRef<SWidget> InWidget, const TAtt
 	NewWidgetBlock->SetSearchable(bSearchable);
 }
 
-void FToolBarBuilder::AddWidget( TSharedRef<SWidget> InWidget, FName InTutorialHighlightName, bool bSearchable )
+void FToolBarBuilder::AddWidget( TSharedRef<SWidget> InWidget, FName InTutorialHighlightName, bool bSearchable, EHorizontalAlignment Alignment )
 {
 	ApplySectionBeginning();
 
@@ -500,7 +500,7 @@ void FToolBarBuilder::AddWidget( TSharedRef<SWidget> InWidget, FName InTutorialH
 			ChildWidget
 		];
 	
-	TSharedRef< FWidgetBlock > NewWidgetBlock( new FWidgetBlock( InWidget, FText::GetEmpty(), true ) );
+	TSharedRef< FWidgetBlock > NewWidgetBlock( new FWidgetBlock( InWidget, FText::GetEmpty(), true, Alignment) );
 	MultiBox->AddMultiBlock( NewWidgetBlock );
 	NewWidgetBlock->SetSearchable(bSearchable);
 }
diff --git a/Engine/Source/Runtime/Slate/Private/Framework/MultiBox/SWidgetBlock.cpp b/Engine/Source/Runtime/Slate/Private/Framework/MultiBox/SWidgetBlock.cpp
--- a/Engine/Source/Runtime/Slate/Private/Framework/MultiBox/SWidgetBlock.cpp
+++ b/Engine/Source/Runtime/Slate/Private/Framework/MultiBox/SWidgetBlock.cpp
@@ -11,11 +11,12 @@
  *
  * @param	InHeadingText	Heading text
  */
-FWidgetBlock::FWidgetBlock( TSharedRef<SWidget> InContent, const FText& InLabel, bool bInNoIndent )
+FWidgetBlock::FWidgetBlock( TSharedRef<SWidget> InContent, const FText& InLabel, bool bInNoIndent, EHorizontalAlignment InHorizontalAlignment)
 	: FMultiBlock( nullptr, nullptr, NAME_None, EMultiBlockType::Widget )
 	, ContentWidget( InContent )
 	, Label( InLabel )
 	, bNoIndent( bInNoIndent )
+	, HorizontalAlignment(InHorizontalAlignment)
 {
 }
 
@@ -39,6 +40,14 @@ TSharedRef< class IMultiBlockBaseWidget > FWidgetBlock::ConstructWidget() const
 }
 
 
+bool FWidgetBlock::GetAlignmentOverrides(EHorizontalAlignment& OutHorizontalAlignment, EVerticalAlignment& OutVerticalAlignment, bool& bOutAutoWidth) const
+{
+	OutHorizontalAlignment = HorizontalAlignment;
+	OutVerticalAlignment = VAlign_Fill;
+	bOutAutoWidth = HorizontalAlignment != HAlign_Fill ? false : true;
+	return true;
+}
+
 /**
  * Construct this widget
  *
diff --git a/Engine/Source/Runtime/Slate/Private/Framework/MultiBox/SWidgetBlock.h b/Engine/Source/Runtime/Slate/Private/Framework/MultiBox/SWidgetBlock.h
--- a/Engine/Source/Runtime/Slate/Private/Framework/MultiBox/SWidgetBlock.h
+++ b/Engine/Source/Runtime/Slate/Private/Framework/MultiBox/SWidgetBlock.h
@@ -23,7 +23,7 @@ public:
 	 * @param	InLabel		Optional label text to be added to the left of the content
 	 * @param	bInNoIndent	If true, removes the padding from the left of the widget that lines it up with other menu items
 	 */
-	FWidgetBlock( TSharedRef<SWidget> InContent, const FText& InLabel, bool bInNoIndent );
+	FWidgetBlock(TSharedRef<SWidget> InContent, const FText& InLabel, bool bInNoIndent, EHorizontalAlignment InHorizontalAlignment = HAlign_Fill);
 
 	/** FMultiBlock interface */
 	virtual void CreateMenuEntry(class FMenuBuilder& MenuBuilder) const override;
@@ -33,8 +33,7 @@ private:
 
 	/** FMultiBlock private interface */
 	virtual TSharedRef< class IMultiBlockBaseWidget > ConstructWidget() const override;
-
-
+	virtual bool GetAlignmentOverrides(EHorizontalAlignment& OutHorizontalAlignment, EVerticalAlignment& OutVerticalAlignment, bool& bOutAutoWidth) const;
 private:
 
 	// Friend our corresponding widget class
@@ -48,6 +47,9 @@ private:
 
 	/** Remove the padding from the left of the widget that lines it up with other menu items? */
 	bool bNoIndent;
+
+	/** Hortizontal aligment for this widget in its parent container. Note: only applies to toolbars */
+	EHorizontalAlignment HorizontalAlignment;
 };
 
 
diff --git a/Engine/Source/Runtime/Slate/Public/Framework/MultiBox/MultiBox.h b/Engine/Source/Runtime/Slate/Public/Framework/MultiBox/MultiBox.h
--- a/Engine/Source/Runtime/Slate/Public/Framework/MultiBox/MultiBox.h
+++ b/Engine/Source/Runtime/Slate/Public/Framework/MultiBox/MultiBox.h
@@ -208,6 +208,15 @@ private:
 	 */
 	virtual TSharedRef< class IMultiBlockBaseWidget > ConstructWidget() const = 0;
 
+	/**
+ 	 * Gets any aligment overrides for this block
+	 *
+	 * @param OutHorizontalAligment	Horizontal alignment override
+	 * @param OutVerticalAlignment	Vertical Alignment override 
+	 * @param bOutAutoWidth		Fill or Auto width override
+	 * @return true if overrides should be applied, false to use defaults 
+ 	 */ 
+	virtual bool GetAlignmentOverrides(EHorizontalAlignment& OutHorizontalAlignment, EVerticalAlignment& OutVerticalAlignment, bool& bOutAutoWidth) const { return false; }
 private:
 
 	// We're friends with SMultiBoxWidget so that it can call MakeWidget() directly
diff --git a/Engine/Source/Runtime/Slate/Public/Framework/MultiBox/MultiBoxBuilder.h b/Engine/Source/Runtime/Slate/Public/Framework/MultiBox/MultiBoxBuilder.h
--- a/Engine/Source/Runtime/Slate/Public/Framework/MultiBox/MultiBoxBuilder.h
+++ b/Engine/Source/Runtime/Slate/Public/Framework/MultiBox/MultiBoxBuilder.h
@@ -470,8 +470,9 @@ public:
 	 * @param	InWidget				The widget that should be shown in the toolbar
 	 * @param	InTutorialHighlightName	Name to identify this widget and highlight during tutorials
 	 * @param	bSearchable			If true, widget will be searchable (default == true)
+	 * @param	Alignment			Horizontal alignment for the widget inside the toolbar
 	 */
-	void AddWidget(TSharedRef<SWidget> InWidget, FName InTutorialHighlightName = NAME_None, bool bSearchable = true);
+	void AddWidget(TSharedRef<SWidget> InWidget, FName InTutorialHighlightName = NAME_None, bool bSearchable = true, EHorizontalAlignment Alignment = HAlign_Fill);
 	
 	/**
 	 * Adds a toolbar separator
@@ -505,7 +506,6 @@ public:
 	 * The override will be used for any added blocks until EndStyleOverride is called
 	 */
 	void BeginStyleOverride(FName StyleOverrideName);
-
 	void EndStyleOverride();
 protected:
 
