Make multiline editable text blocks respect FSlateColor::UseStyle when returned from an attribute override

diff --git a/Engine/Source/Runtime/Slate/Private/Widgets/Input/SMultiLineEditableTextBox.cpp b/Engine/Source/Runtime/Slate/Private/Widgets/Input/SMultiLineEditableTextBox.cpp
--- a/Engine/Source/Runtime/Slate/Private/Widgets/Input/SMultiLineEditableTextBox.cpp
+++ b/Engine/Source/Runtime/Slate/Private/Widgets/Input/SMultiLineEditableTextBox.cpp
@@ -214,22 +214,44 @@ FSlateColor SMultiLineEditableTextBox::DetermineForegroundColor() const
 {
 	check(Style);
 
+	FSlateColor Result = FSlateColor::UseStyle();
+
 	if ( EditableText->IsTextReadOnly() )
 	{
 		if (ReadOnlyForegroundColorOverride.IsSet())
 		{
-			return ReadOnlyForegroundColorOverride.Get();
+			Result = ReadOnlyForegroundColorOverride.Get();
 		}
 		if (ForegroundColorOverride.IsSet())
 		{
-			return ForegroundColorOverride.Get();
+			Result = ForegroundColorOverride.Get();
 		}
 
-		return Style->ReadOnlyForegroundColor;
+		if (Result == FSlateColor::UseStyle())
+		{
+			return Style->ReadOnlyForegroundColor;
+		}
+		else
+		{
+			return Result;
+		}
 	}
 	else
 	{
-		return ForegroundColorOverride.IsSet() ? ForegroundColorOverride.Get() : Style->ForegroundColor;
+		if (ForegroundColorOverride.IsSet())
+		{
+			Result = ForegroundColorOverride.Get();
+		}
+
+		if (Result == FSlateColor::UseStyle())
+		{
+			return Style->ForegroundColor;
+		}
+		else
+		{
+			return Result;
+		}
+
 	}
 }
 
