Fix mouse tracking on startup.

diff --git a/Engine/Source/Runtime/Slate/Private/Framework/Application/SlateApplication.cpp b/Engine/Source/Runtime/Slate/Private/Framework/Application/SlateApplication.cpp
--- a/Engine/Source/Runtime/Slate/Private/Framework/Application/SlateApplication.cpp
+++ b/Engine/Source/Runtime/Slate/Private/Framework/Application/SlateApplication.cpp
@@ -5065,7 +5065,7 @@ bool FSlateApplication::RoutePointerMoveEvent(const FWidgetPath& WidgetsUnderPoi
 				return FNoReply();
 			}, ESlateDebuggingInputEvent::MouseEnter);
 
-		FReply Reply = FEventRouter::Route<FReply>(this, FEventRouter::FToLeafmostPolicy(MouseCaptorPath), PointerEvent, [this, bIsSynthetic](const FArrangedWidget& MouseCaptorWidget, const FPointerEvent& Event)
+		FReply Reply = FEventRouter::Route<FReply>(this, FEventRouter::FToLeafmostPolicy(MouseCaptorPath), PointerEvent, [this](const FArrangedWidget& MouseCaptorWidget, const FPointerEvent& Event)
 			{
 				FReply TempReply = FReply::Unhandled();
 
@@ -5098,7 +5098,11 @@ bool FSlateApplication::RoutePointerMoveEvent(const FWidgetPath& WidgetsUnderPoi
 				}
 				if ((!Event.IsTouchEvent() && bAllowMouseFallback) || (!TempReply.IsEventHandled() && this->bTouchFallbackToMouse))
 				{
-					TempReply = MouseCaptorWidget.Widget->OnMouseMove(MouseCaptorWidget.Geometry, Event);
+					// Only handle if hovered, else widgets with mouse capture can cause the mouse to move at app start.
+					if (MouseCaptorWidget.Widget->IsHovered())
+					{
+						TempReply = MouseCaptorWidget.Widget->OnMouseMove(MouseCaptorWidget.Geometry, Event);
+					}
 #if WITH_SLATE_DEBUGGING
 					FSlateDebugging::BroadcastInputEvent(ESlateDebuggingInputEvent::MouseMove, TempReply, MouseCaptorWidget.Widget);
 #endif
