Updates TreeView Hierarchy Wire Styling

diff --git a/Engine/Source/Runtime/Slate/Private/Widgets/Views/SExpanderArrow.cpp b/Engine/Source/Runtime/Slate/Private/Widgets/Views/SExpanderArrow.cpp
--- a/Engine/Source/Runtime/Slate/Private/Widgets/Views/SExpanderArrow.cpp
+++ b/Engine/Source/Runtime/Slate/Private/Widgets/Views/SExpanderArrow.cpp
@@ -40,7 +40,7 @@ void SExpanderArrow::Construct( const FArguments& InArgs, const TSharedPtr<class
 int32 SExpanderArrow::OnPaint(const FPaintArgs& Args, const FGeometry& AllottedGeometry, const FSlateRect& MyCullingRect, FSlateWindowElementList& OutDrawElements, int32 LayerId, const FWidgetStyle& InWidgetStyle, bool bParentEnabled) const
 {
 
-	static const float WireThickness = 1.0f;
+	static const float WireThickness = 2.0f;
 	static const float HalfWireThickness = WireThickness / 2.0f;
 
 	// We want to support drawing wires for the tree
@@ -64,7 +64,7 @@ int32 SExpanderArrow::OnPaint(const FPaintArgs& Args, const FGeometry& AllottedG
 	{
 		const TSharedPtr<ITableRow> OwnerRow = OwnerRowPtr.Pin();
 		FLinearColor WireTint = InWidgetStyle.GetForegroundColor();
-		WireTint.A = 0.275f;
+		WireTint.A = 0.15f;
 
 		// Draw vertical wires to indicate paths to parent nodes.
 		const TBitArray<>& NeedsWireByLevel = OwnerRow->GetWiresNeededByDepth();
@@ -78,7 +78,7 @@ int32 SExpanderArrow::OnPaint(const FPaintArgs& Args, const FGeometry& AllottedG
 				FSlateDrawElement::MakeBox(
 					OutDrawElements,
 					LayerId,
-					AllottedGeometry.ToPaintGeometry(FVector2D(WireThickness, AllottedGeometry.Size.Y), FSlateLayoutTransform(FVector2D(CurrentIndent, 0))),
+					AllottedGeometry.ToPaintGeometry(FVector2D(WireThickness, AllottedGeometry.Size.Y), FSlateLayoutTransform(FVector2D(CurrentIndent - 3.f, 0))),
 					VerticalBarBrush,
 					ESlateDrawEffect::None,
 					WireTint
@@ -95,7 +95,7 @@ int32 SExpanderArrow::OnPaint(const FPaintArgs& Args, const FGeometry& AllottedG
 			FSlateDrawElement::MakeBox(
 				OutDrawElements,
 				LayerId,
-				AllottedGeometry.ToPaintGeometry(FVector2D(WireThickness, HalfCellHeight + HalfWireThickness), FSlateLayoutTransform(FVector2D(CurrentIndent, 0))),
+				AllottedGeometry.ToPaintGeometry(FVector2D(WireThickness, HalfCellHeight + HalfWireThickness), FSlateLayoutTransform(FVector2D(CurrentIndent - 3.f, 0))),
 				VerticalBarBrush,
 				ESlateDrawEffect::None,
 				WireTint
@@ -109,7 +109,7 @@ int32 SExpanderArrow::OnPaint(const FPaintArgs& Args, const FGeometry& AllottedG
 			FSlateDrawElement::MakeBox(
 				OutDrawElements,
 				LayerId,
-				AllottedGeometry.ToPaintGeometry(FVector2D(WireThickness, HalfCellHeight+ HalfWireThickness), FSlateLayoutTransform(FVector2D(CurrentIndent, HalfCellHeight- HalfWireThickness))),
+				AllottedGeometry.ToPaintGeometry(FVector2D(WireThickness, HalfCellHeight+ HalfWireThickness), FSlateLayoutTransform(FVector2D(CurrentIndent - 3.f, HalfCellHeight- HalfWireThickness))),
 				VerticalBarBrush,
 				ESlateDrawEffect::None,
 				WireTint
@@ -118,13 +118,14 @@ int32 SExpanderArrow::OnPaint(const FPaintArgs& Args, const FGeometry& AllottedG
 
 		// Draw horizontal connector from parent wire to child.
 		{
+			float LeafDepth = OwnerRow->DoesItemHaveChildren() ? 10.f : 0.0f;
 			const float HorizontalWireStart = (NumLevels - 1)*Indent;
 			FSlateDrawElement::MakeBox(
 				OutDrawElements,
 				LayerId,
 				AllottedGeometry.ToPaintGeometry(
-					FVector2D(AllottedGeometry.Size.X - HorizontalWireStart - WireThickness, WireThickness),
-					FSlateLayoutTransform(FVector2D(HorizontalWireStart+WireThickness, 0.5f*(AllottedGeometry.Size.Y - WireThickness)))
+					FVector2D(AllottedGeometry.Size.X - HorizontalWireStart - WireThickness - LeafDepth, WireThickness),
+					FSlateLayoutTransform(FVector2D(HorizontalWireStart + WireThickness - 3.f, 0.5f*(AllottedGeometry.Size.Y - WireThickness)))
 				),
 				VerticalBarBrush,
 				ESlateDrawEffect::None,
