SHeaderRow - Allows bound MenuContent, Provides ability to programmatically show/hide columns

diff --git a/Engine/Source/Runtime/Slate/Private/Widgets/Views/SHeaderRow.cpp b/Engine/Source/Runtime/Slate/Private/Widgets/Views/SHeaderRow.cpp
--- a/Engine/Source/Runtime/Slate/Private/Widgets/Views/SHeaderRow.cpp
+++ b/Engine/Source/Runtime/Slate/Private/Widgets/Views/SHeaderRow.cpp
@@ -146,7 +146,8 @@ public:
 			PrimaryContent
 		];
 
-		if( Column.HeaderMenuContent.Widget != SNullWidget::NullWidget )
+		if( Column.HeaderMenuContent.Widget != SNullWidget::NullWidget || 
+			Column.OnGetMenuContent.IsBound())
 		{
 			// Add Drop down menu button (only if menu content has been specified)
 			Box->AddSlot()
@@ -175,10 +176,6 @@ public:
 							SNew( SSpacer )
 							.Size( FVector2D( 14.0f, 0 ) )
 						]
-						.MenuContent()
-						[
-							ContextMenuContent
-						]
 					]
 				]
 
@@ -193,6 +190,16 @@ public:
 				]
 			];		
 
+
+			if (Column.HeaderMenuContent.Widget != SNullWidget::NullWidget)
+			{
+				ComboButton->SetMenuContent(ContextMenuContent);
+			}
+			else if (Column.OnGetMenuContent.IsBound())
+			{
+				ComboButton->SetOnGetMenuContent(Column.OnGetMenuContent);
+			}
+
 			AdjustedDefaultHeaderContentPadding.Right = 0;
 		}
 
@@ -926,4 +933,27 @@ ECheckBoxState SHeaderRow::GetGeneratedColumnCheckedState(FName ColumnId) const
 	return IsColumnGenerated(ColumnId) ? ECheckBoxState::Checked : ECheckBoxState::Unchecked;
 }
 
+void SHeaderRow::SetShowGeneratedColumn(const FName& ColumnId, bool InShow)
+{
+	// Only column that doesn't have a ShouldGenerateWidget, can be toggled
+	for (FColumn& SomeColumn : Columns)
+	{
+		if (SomeColumn.ColumnId == ColumnId)
+		{
+			if (!SomeColumn.ShouldGenerateWidget.IsSet())
+			{
+				if (SomeColumn.bIsVisible != InShow)
+				{
+					SomeColumn.bIsVisible = !SomeColumn.bIsVisible;
+
+					RefreshColumns();
+					ColumnsChanged.Broadcast(SharedThis(this));
+					OnHiddenColumnsListChanged.ExecuteIfBound();
+				}
+			}
+			break;
+		}
+	}
+}
+
 #undef LOCTEXT_NAMESPACE
diff --git a/Engine/Source/Runtime/Slate/Public/Widgets/Views/SHeaderRow.h b/Engine/Source/Runtime/Slate/Public/Widgets/Views/SHeaderRow.h
--- a/Engine/Source/Runtime/Slate/Public/Widgets/Views/SHeaderRow.h
+++ b/Engine/Source/Runtime/Slate/Public/Widgets/Views/SHeaderRow.h
@@ -13,6 +13,7 @@
 #include "Widgets/Layout/SBorder.h"
 #include "Styling/SlateWidgetStyleAsset.h"
 #include "Widgets/Layout/SSplitter.h"
+#include "Framework/SlateDelegates.h"
 
 class FMenuBuilder;
 class SScrollBar;
@@ -127,7 +128,9 @@ public:
 			SLATE_ARGUMENT( EVerticalAlignment, VAlignHeader )
 			SLATE_ARGUMENT( TOptional< FMargin >, HeaderContentPadding )
 			SLATE_ARGUMENT( EHeaderComboVisibility, HeaderComboVisibility )
+
 			SLATE_NAMED_SLOT( FArguments, MenuContent )
+			SLATE_EVENT( FOnGetContent, OnGetMenuContent )
 
 			SLATE_ARGUMENT( EHorizontalAlignment, HAlignCell )
 			SLATE_ARGUMENT( EVerticalAlignment, VAlignCell )
@@ -149,6 +152,7 @@ public:
 			, SizeRule( EColumnSizeMode::Fill )
 			, HeaderContent( InArgs._HeaderContent )
 			, HeaderMenuContent( InArgs._MenuContent )
+			, OnGetMenuContent( InArgs._OnGetMenuContent )
 			, HeaderHAlignment( InArgs._HAlignHeader )
 			, HeaderVAlignment( InArgs._VAlignHeader )
 			, HeaderContentPadding( InArgs._HeaderContentPadding )
@@ -221,6 +225,8 @@ public:
 		TAlwaysValidWidget HeaderContent;	
 		TAlwaysValidWidget HeaderMenuContent;
 
+		FOnGetContent OnGetMenuContent;
+
 		EHorizontalAlignment HeaderHAlignment;
 		EVerticalAlignment HeaderVAlignment;
 		TOptional< FMargin > HeaderContentPadding;
@@ -320,6 +326,10 @@ public:
 	virtual FReply OnMouseButtonUp( const FGeometry& MyGeometry, const FPointerEvent& MouseEvent ) override;
 	//~ End SWidget interface
 
+	/** Show or Hide a generated column */
+	void SetShowGeneratedColumn(const FName& InColumnIdm, bool InShow = true);
+
+
 private:
 
 	/** Regenerates all widgets in the header */
