Fix mouse tracking on game startup immediately moving camera, without breaking sequencer.

diff --git a/Engine/Source/Runtime/Slate/Private/Framework/Application/SlateApplication.cpp b/Engine/Source/Runtime/Slate/Private/Framework/Application/SlateApplication.cpp
--- a/Engine/Source/Runtime/Slate/Private/Framework/Application/SlateApplication.cpp
+++ b/Engine/Source/Runtime/Slate/Private/Framework/Application/SlateApplication.cpp
@@ -5098,7 +5098,11 @@ bool FSlateApplication::RoutePointerMoveEvent(const FWidgetPath& WidgetsUnderPoi
 				}
 				if ((!Event.IsTouchEvent() && bAllowMouseFallback) || (!TempReply.IsEventHandled() && this->bTouchFallbackToMouse))
 				{
-					TempReply = MouseCaptorWidget.Widget->OnMouseMove(MouseCaptorWidget.Geometry, Event);
+					// Only handle if not synthetic, else widgets with mouse capture can cause the mouse to move at app start.
+					if (!bIsSynthetic)
+					{
+						TempReply = MouseCaptorWidget.Widget->OnMouseMove(MouseCaptorWidget.Geometry, Event);
+					}
 #if WITH_SLATE_DEBUGGING
 					FSlateDebugging::BroadcastInputEvent(ESlateDebuggingInputEvent::MouseMove, TempReply, MouseCaptorWidget.Widget);
 #endif
