Add support for splitter width customization in column headers.

diff --git a/Engine/Source/Runtime/Slate/Private/Widgets/Views/SHeaderRow.cpp b/Engine/Source/Runtime/Slate/Private/Widgets/Views/SHeaderRow.cpp
--- a/Engine/Source/Runtime/Slate/Private/Widgets/Views/SHeaderRow.cpp
+++ b/Engine/Source/Runtime/Slate/Private/Widgets/Views/SHeaderRow.cpp
@@ -451,6 +451,7 @@ void SHeaderRow::Construct( const FArguments& InArgs )
 	Style = InArgs._Style;
 	OnGetMaxRowSizeForColumn = InArgs._OnGetMaxRowSizeForColumn;
 	ResizeMode = InArgs._ResizeMode;
+	SplitterHandleSize = InArgs._SplitterHandleSize;
 	bCanSelectGeneratedColumn = InArgs._CanSelectGeneratedColumn;
 	OnHiddenColumnsListChanged = InArgs._OnHiddenColumnsListChanged;
 
@@ -663,7 +664,7 @@ FReply SHeaderRow::OnMouseButtonUp(const FGeometry& MyGeometry, const FPointerEv
 
 void SHeaderRow::RegenerateWidgets()
 {
-	const float SplitterHandleDetectionSize = 5.0f;
+	const float SplitterHandleDetectionSize = SplitterHandleSize > 0.0f ? SplitterHandleSize + 4.0f : 5.0f;
 	HeaderWidgets.Empty();
 
 	TSharedPtr<SSplitter> Splitter;
@@ -676,7 +677,7 @@ void SHeaderRow::RegenerateWidgets()
 			SAssignNew(Splitter, SSplitter)
 			.Style( &Style->ColumnSplitterStyle )
 			.ResizeMode(ResizeMode)
-			.PhysicalSplitterHandleSize( 0.0f )
+			.PhysicalSplitterHandleSize( SplitterHandleSize )
 			.HitDetectionSplitterHandleSize( SplitterHandleDetectionSize )
 			.OnGetMaxSlotSize(this, &SHeaderRow::GetRowSizeForSlotIndex)
 		]
@@ -776,10 +777,10 @@ void SHeaderRow::RegenerateWidgets()
 				case EColumnSizeMode::Manual:
 				{
 					// Sizing grip to put at the end of the column - we can't use a SSplitter here as it doesn't have the resizing behavior we need
-					const float GripSize = 5.0f;
+					const float GripSize = SplitterHandleSize > 0.0f ? SplitterHandleSize + 4.0f : 5.0f;
 					TSharedRef<SBorder> SizingGrip = SNew(SBorder)
 						.Padding(0.0f)
-						.BorderImage(FAppStyle::Get().GetBrush("NoBorder"))
+						.BorderImage( &Style->ColumnSplitterStyle )
 						.Cursor(EMouseCursor::ResizeLeftRight)
 						.Content()
 						[
diff --git a/Engine/Source/Runtime/Slate/Public/Widgets/Views/SHeaderRow.h b/Engine/Source/Runtime/Slate/Public/Widgets/Views/SHeaderRow.h
--- a/Engine/Source/Runtime/Slate/Public/Widgets/Views/SHeaderRow.h
+++ b/Engine/Source/Runtime/Slate/Public/Widgets/Views/SHeaderRow.h
@@ -257,6 +257,7 @@ public:
 	SLATE_BEGIN_ARGS(SHeaderRow)
 		: _Style( &FCoreStyle::Get().GetWidgetStyle<FHeaderRowStyle>("TableView.Header") )
 		, _ResizeMode(ESplitterResizeMode::Fill)
+		, _SplitterHandleSize(0.0f)
 		, _CanSelectGeneratedColumn(false)
 		{}
 
@@ -265,6 +266,7 @@ public:
 		SLATE_EVENT( FColumnsChanged::FDelegate, OnColumnsChanged )
 		SLATE_EVENT( FOnGetMaxRowSizeForColumn, OnGetMaxRowSizeForColumn )
 		SLATE_ARGUMENT( ESplitterResizeMode::Type, ResizeMode )
+		SLATE_ARGUMENT( float, SplitterHandleSize )
 		/**
 		 * Can select the columns generated by right clicking on the header menu.
 		 * FColumn with ShouldGenerateWidget set can not be selected.
@@ -348,6 +350,7 @@ private:
 	const FHeaderRowStyle* Style;
 	FColumnsChanged ColumnsChanged;
 	ESplitterResizeMode::Type ResizeMode;
+	float SplitterHandleSize;
 	FOnGetMaxRowSizeForColumn OnGetMaxRowSizeForColumn;
 
 	FSimpleDelegate OnHiddenColumnsListChanged;
