SHeaderRow adds SplitterHandleSize to FHeaderRowStyle

diff --git a/Engine/Source/Runtime/Slate/Private/Widgets/Views/SHeaderRow.cpp b/Engine/Source/Runtime/Slate/Private/Widgets/Views/SHeaderRow.cpp
--- a/Engine/Source/Runtime/Slate/Private/Widgets/Views/SHeaderRow.cpp
+++ b/Engine/Source/Runtime/Slate/Private/Widgets/Views/SHeaderRow.cpp
@@ -451,7 +451,12 @@ void SHeaderRow::Construct( const FArguments& InArgs )
 	Style = InArgs._Style;
 	OnGetMaxRowSizeForColumn = InArgs._OnGetMaxRowSizeForColumn;
 	ResizeMode = InArgs._ResizeMode;
-	SplitterHandleSize = InArgs._SplitterHandleSize;
+
+	SplitterHandleSize =   Style->SplitterHandleSize;
+	if (InArgs._SplitterHandleSize.IsSet())
+	{
+		SplitterHandleSize = InArgs._SplitterHandleSize.GetValue();
+	}
 	bCanSelectGeneratedColumn = InArgs._CanSelectGeneratedColumn;
 	OnHiddenColumnsListChanged = InArgs._OnHiddenColumnsListChanged;
 
diff --git a/Engine/Source/Runtime/Slate/Public/Widgets/Views/SHeaderRow.h b/Engine/Source/Runtime/Slate/Public/Widgets/Views/SHeaderRow.h
--- a/Engine/Source/Runtime/Slate/Public/Widgets/Views/SHeaderRow.h
+++ b/Engine/Source/Runtime/Slate/Public/Widgets/Views/SHeaderRow.h
@@ -255,9 +255,8 @@ public:
 	FColumnsChanged* OnColumnsChanged() { return &ColumnsChanged; }
 
 	SLATE_BEGIN_ARGS(SHeaderRow)
-		: _Style( &FCoreStyle::Get().GetWidgetStyle<FHeaderRowStyle>("TableView.Header") )
+		: _Style( &FAppStyle::Get().GetWidgetStyle<FHeaderRowStyle>("TableView.Header") )
 		, _ResizeMode(ESplitterResizeMode::Fill)
-		, _SplitterHandleSize(0.0f)
 		, _CanSelectGeneratedColumn(false)
 		{}
 
@@ -266,7 +265,7 @@ public:
 		SLATE_EVENT( FColumnsChanged::FDelegate, OnColumnsChanged )
 		SLATE_EVENT( FOnGetMaxRowSizeForColumn, OnGetMaxRowSizeForColumn )
 		SLATE_ARGUMENT( ESplitterResizeMode::Type, ResizeMode )
-		SLATE_ARGUMENT( float, SplitterHandleSize )
+		SLATE_ARGUMENT( TOptional<float>, SplitterHandleSize )
 		/**
 		 * Can select the columns generated by right clicking on the header menu.
 		 * FColumn with ShouldGenerateWidget set can not be selected.
diff --git a/Engine/Source/Runtime/SlateCore/Private/Styling/SlateTypes.cpp b/Engine/Source/Runtime/SlateCore/Private/Styling/SlateTypes.cpp
--- a/Engine/Source/Runtime/SlateCore/Private/Styling/SlateTypes.cpp
+++ b/Engine/Source/Runtime/SlateCore/Private/Styling/SlateTypes.cpp
@@ -579,7 +579,8 @@ const FTableColumnHeaderStyle& FTableColumnHeaderStyle::GetDefault()
 
 
 FHeaderRowStyle::FHeaderRowStyle()
-	: HorizontalSeparatorBrush(FSlateNoResource())
+	: SplitterHandleSize(0.0)
+	, HorizontalSeparatorBrush(FSlateNoResource())
 	, HorizontalSeparatorThickness(0)
 {
 }
diff --git a/Engine/Source/Runtime/SlateCore/Public/Styling/SlateTypes.h b/Engine/Source/Runtime/SlateCore/Public/Styling/SlateTypes.h
--- a/Engine/Source/Runtime/SlateCore/Public/Styling/SlateTypes.h
+++ b/Engine/Source/Runtime/SlateCore/Public/Styling/SlateTypes.h
@@ -1626,6 +1626,11 @@ struct SLATECORE_API FHeaderRowStyle : public FSlateWidgetStyle
 	FSplitterStyle ColumnSplitterStyle;
 	FHeaderRowStyle& SetColumnSplitterStyle( const FSplitterStyle& InColumnSplitterStyle ){ ColumnSplitterStyle = InColumnSplitterStyle; return *this; }
 
+	/** Size of the splitter used between the columns */
+	UPROPERTY(EditAnywhere, Category=Appearance)
+	float SplitterHandleSize;
+	FHeaderRowStyle& SetSplitterHandleSize( const float  InSplitterHandleSize){ SplitterHandleSize = InSplitterHandleSize; return *this; }
+
 	/** Brush used to draw the header row background */
 	UPROPERTY(EditAnywhere, Category=Appearance)
 	FSlateBrush BackgroundBrush;
