Slate: Fix some implcit conversion in Slate module.

diff --git a/Engine/Source/Runtime/Slate/Private/Widgets/Input/SSuggestionTextBox.cpp b/Engine/Source/Runtime/Slate/Private/Widgets/Input/SSuggestionTextBox.cpp
--- a/Engine/Source/Runtime/Slate/Private/Widgets/Input/SSuggestionTextBox.cpp
+++ b/Engine/Source/Runtime/Slate/Private/Widgets/Input/SSuggestionTextBox.cpp
@@ -65,11 +65,11 @@ void SSuggestionTextBox::Construct( const FArguments& InArgs )
 							.MaxHeight(InArgs._SuggestionListMaxHeight)
 							[
 								SAssignNew(SuggestionListView, SListView< TSharedPtr<FString> >)
-									.ItemHeight(18)
-									.ListItemsSource(&Suggestions)
-									.SelectionMode(ESelectionMode::Single)
-									.OnGenerateRow(this, &SSuggestionTextBox::HandleSuggestionListViewGenerateRow)
-									.OnSelectionChanged(this, &SSuggestionTextBox::HandleSuggestionListViewSelectionChanged)
+								.ItemHeight(18.f)
+								.ListItemsSource(&Suggestions)
+								.SelectionMode(ESelectionMode::Single)
+								.OnGenerateRow(this, &SSuggestionTextBox::HandleSuggestionListViewGenerateRow)
+								.OnSelectionChanged(this, &SSuggestionTextBox::HandleSuggestionListViewSelectionChanged)
 							]
 					]
 			)
diff --git a/Engine/Source/Runtime/Slate/Public/Widgets/Docking/SDockTab.h b/Engine/Source/Runtime/Slate/Public/Widgets/Docking/SDockTab.h
--- a/Engine/Source/Runtime/Slate/Public/Widgets/Docking/SDockTab.h
+++ b/Engine/Source/Runtime/Slate/Public/Widgets/Docking/SDockTab.h
@@ -69,7 +69,7 @@ public:
 		: _Content()
 		, _TabWellContentLeft()
 		, _TabWellContentRight()
-		, _ContentPadding(0)
+		, _ContentPadding(0.f)
 		, _TabRole(ETabRole::PanelTab)
 		, _Label()
 		, _LabelSuffix()
diff --git a/Engine/Source/Runtime/Slate/Public/Widgets/Input/SSpinBox.h b/Engine/Source/Runtime/Slate/Public/Widgets/Input/SSpinBox.h
--- a/Engine/Source/Runtime/Slate/Public/Widgets/Input/SSpinBox.h
+++ b/Engine/Source/Runtime/Slate/Public/Widgets/Input/SSpinBox.h
@@ -601,7 +601,7 @@ public:
 						}
 						else
 						{
-							Percent = 1.0f - FMath::Pow(1.0f - Percent, 1.0 / CachedSliderExponent);
+							Percent = 1.0f - FMath::Pow(1.0f - Percent, 1.0f / CachedSliderExponent);
 						}
 
 
@@ -616,13 +616,13 @@ public:
 					if (LinearDeltaSensitivity.IsSet() && Delta.IsSet() && Delta.Get() > 0)
 					{
 						const float MouseDelta = FMath::Abs(MouseEvent.GetCursorDelta().X / LinearDeltaSensitivity.Get());
-						NewValue = InternalValue + (Sign * MouseDelta * FMath::Pow(Delta.Get(), SliderExponent.Get()));
+						NewValue = InternalValue + (Sign * MouseDelta * FMath::Pow((float)Delta.Get(), SliderExponent.Get()));
 					}
 					else
 					{
 						const float MouseDelta = FMath::Abs(MouseEvent.GetCursorDelta().X / SliderWidthInSlateUnits);
 						const double CurrentValue = FMath::Clamp<double>(FMath::Abs(InternalValue), 1.0, (double)std::numeric_limits<NumericType>::max());
-						NewValue = InternalValue + (Sign * MouseDelta * FMath::Pow(CurrentValue, SliderExponent.Get()));
+						NewValue = InternalValue + (Sign * MouseDelta * FMath::Pow((float)CurrentValue, SliderExponent.Get()));
 					}
 				}
 
diff --git a/Engine/Source/Runtime/Slate/Public/Widgets/Views/SListView.h b/Engine/Source/Runtime/Slate/Public/Widgets/Views/SListView.h
--- a/Engine/Source/Runtime/Slate/Public/Widgets/Views/SListView.h
+++ b/Engine/Source/Runtime/Slate/Public/Widgets/Views/SListView.h
@@ -326,7 +326,7 @@ public:
 					SelectionIndex = ItemsSourceRef.Find( TListTypeTraits<ItemType>::NullableItemTypeConvertToItemType( SelectorItem ) );
 				}
 
-				int32 NumItemsInAPage = GetNumLiveWidgets();
+				int32 NumItemsInAPage = FMath::TruncToInt(GetNumLiveWidgets());
 				int32 Remainder = NumItemsInAPage % GetNumItemsPerLine();
 				NumItemsInAPage -= Remainder;
 
@@ -350,7 +350,7 @@ public:
 					SelectionIndex = ItemsSourceRef.Find( TListTypeTraits<ItemType>::NullableItemTypeConvertToItemType( SelectorItem ) );
 				}
 
-				int32 NumItemsInAPage = GetNumLiveWidgets();
+				int32 NumItemsInAPage = FMath::TruncToInt(GetNumLiveWidgets());
 				int32 Remainder = NumItemsInAPage % GetNumItemsPerLine();
 				NumItemsInAPage -= Remainder;
 
@@ -1148,6 +1148,11 @@ public:
 	 */
 	virtual FReGenerateResults ReGenerateItems( const FGeometry& MyGeometry ) override
 	{
+		auto DoubleFractional = [](double Value) -> double
+		{
+			return Value - FMath::TruncToDouble(Value);
+		};
+
 		// Clear all the items from our panel. We will re-add them in the correct order momentarily.
 		this->ClearWidgets();
 
@@ -1167,7 +1172,7 @@ public:
 
 			// Index of the item at which we start generating based on how far scrolled down we are
 			// Note that we must generate at LEAST one item.
-			int32 StartIndex = FMath::Clamp( FMath::FloorToInt(CurrentScrollOffset), 0, ItemsSource->Num() - 1 );
+			int32 StartIndex = FMath::Clamp( (int32)(FMath::FloorToDouble(CurrentScrollOffset)), 0, ItemsSource->Num() - 1 );
 
 			// Length of the first item that is generated. This item is at the location where the user requested we scroll
 			float FirstItemLength = 0.0f;
@@ -1203,7 +1208,7 @@ public:
 				{
 					// The first item may not be fully visible (but cannot exceed 1)
 					// FirstItemFractionScrolledIntoView is the fraction of the item that is visible after taking into account anything that may be scrolled off the top/left of the list view
-					const float FirstItemFractionScrolledIntoView = 1.0f - FMath::Max(FMath::Fractional(CurrentScrollOffset), 0.0f);
+					const float FirstItemFractionScrolledIntoView = 1.0f - (float)FMath::Max(DoubleFractional(CurrentScrollOffset), 0.0);
 					
 					// FirstItemLengthScrolledIntoView is the length of the item, ignoring anything that is scrolled off the top/left of the list view
 					const float FirstItemLengthScrolledIntoView = ItemLength * FirstItemFractionScrolledIntoView;
@@ -1730,7 +1735,7 @@ protected:
 
 				EndInertialScrolling();
 				
-				const int32 NumFullEntriesInView = FMath::FloorToInt(CurrentScrollOffset + NumLiveWidgets) - FMath::CeilToInt(CurrentScrollOffset);
+				const int32 NumFullEntriesInView = (int32)(FMath::FloorToDouble(CurrentScrollOffset + NumLiveWidgets) - FMath::CeilToDouble(CurrentScrollOffset));
 
 				// Only scroll the item into view if it's not already in the visible range
 				// When navigating, we don't want to scroll partially visible existing rows all the way to the center, so we count partially displayed indices in the displayed range
@@ -1742,13 +1747,13 @@ protected:
 					double NewScrollOffset = IndexOfItem;
 
 					// Center the list view on the item in question.
-					NewScrollOffset -= (NumLiveWidgets / 2);
+					NewScrollOffset -= (NumLiveWidgets / 2.0);
 
 					// Limit offset to top and bottom of the list.
 					const double MaxScrollOffset = FMath::Max(0.0, static_cast<double>(ItemsSource->Num()) - NumLiveWidgets);
 					NewScrollOffset = FMath::Clamp<double>(NewScrollOffset, 0.0, MaxScrollOffset);
 
-					SetScrollOffset(NewScrollOffset);
+					SetScrollOffset((float)NewScrollOffset);
 				}
 				else if (bNavigateOnScrollIntoView)
 				{
@@ -1798,7 +1803,7 @@ protected:
 						}
 
 						const double MaxScrollOffset = FMath::Max(0.0, static_cast<double>(ItemsSource->Num()) - NumLiveWidgets);
-						SetScrollOffset(FMath::Min(NewScrollOffset, MaxScrollOffset));
+						SetScrollOffset((float)FMath::Min(NewScrollOffset, MaxScrollOffset));
 					}
 				}
 
@@ -1853,6 +1858,11 @@ protected:
 
 	virtual float ScrollBy( const FGeometry& MyGeometry, float ScrollByAmountInSlateUnits, EAllowOverscroll InAllowOverscroll ) override
 	{
+		auto DoubleFractional = [](double Value) -> double
+		{
+			return Value - FMath::TruncToDouble(Value);
+		};
+
 		if (InAllowOverscroll == EAllowOverscroll::No)
 		{
 			//check if we are on the top of the list and want to scroll up
@@ -1872,8 +1882,8 @@ protected:
 		int32 StartingItemIndex = (int32)CurrentScrollOffset;
 		double NewScrollOffset = DesiredScrollOffset;
 
-		const bool bWholeListVisible = DesiredScrollOffset == 0 && bWasAtEndOfList;
-		if ( InAllowOverscroll == EAllowOverscroll::Yes && Overscroll.ShouldApplyOverscroll(DesiredScrollOffset == 0, bWasAtEndOfList, ScrollByAmountInSlateUnits ) )
+		const bool bWholeListVisible = DesiredScrollOffset == 0.0 && bWasAtEndOfList;
+		if ( InAllowOverscroll == EAllowOverscroll::Yes && Overscroll.ShouldApplyOverscroll(DesiredScrollOffset == 0.0, bWasAtEndOfList, ScrollByAmountInSlateUnits ) )
 		{
 			const float UnclampedScrollDelta = FMath::Sign(ScrollByAmountInSlateUnits) * AbsScrollByAmount;				
 			const float ActuallyScrolledBy = Overscroll.ScrollBy(MyGeometry, UnclampedScrollDelta);
@@ -1925,14 +1935,14 @@ protected:
 					const FTableViewDimensions WidgetDimensions(this->Orientation, RowWidget->AsWidget()->GetDesiredSize());
 					if (ScrollByAmountInSlateUnits > 0)
 					{
-						const float RemainingDistance = WidgetDimensions.ScrollAxis * (1.0 - FMath::Fractional(NewScrollOffset));
+						const float RemainingDistance = WidgetDimensions.ScrollAxis * (float)(1.0 - DoubleFractional(NewScrollOffset));
 
 						if (AbsScrollByAmount > RemainingDistance)
 						{
 							if (ItemIndex != ItemsSource->Num())
 							{
 								AbsScrollByAmount -= RemainingDistance;
-								NewScrollOffset = 1.0f + (int32)NewScrollOffset;
+								NewScrollOffset = 1.0 + (int32)NewScrollOffset;
 								++ItemIndex;
 							}
 							else
@@ -1943,19 +1953,19 @@ protected:
 						} 
 						else if ( AbsScrollByAmount == RemainingDistance)
 						{
-							NewScrollOffset = 1.0f + (int32)NewScrollOffset;
+							NewScrollOffset = 1.0 + (int32)NewScrollOffset;
 							break;
 						}
 						else
 						{
-							NewScrollOffset = (int32)NewScrollOffset + (1.0f - ((RemainingDistance - AbsScrollByAmount) / WidgetDimensions.ScrollAxis));
+							NewScrollOffset = (int32)NewScrollOffset + (1.0 - ((RemainingDistance - AbsScrollByAmount) / WidgetDimensions.ScrollAxis));
 							break;
 						}
 					}
 					else
 					{
-						float Fractional = FMath::Fractional( NewScrollOffset );
-						if ( Fractional == 0 )
+						float Fractional = FMath::Fractional( (float)NewScrollOffset );
+						if ( FMath::IsNearlyEqual(Fractional, 0.f) )
 						{
 							Fractional = 1.0f;
 							--NewScrollOffset;
@@ -1968,18 +1978,18 @@ protected:
 							if ( ItemIndex != 0 )
 							{
 								AbsScrollByAmount -= PrecedingDistance;
-								NewScrollOffset -= FMath::Fractional( NewScrollOffset );
+								NewScrollOffset -= DoubleFractional( NewScrollOffset );
 								--ItemIndex;
 							}
 							else
 							{
-								NewScrollOffset = 0;
+								NewScrollOffset = 0.0;
 								break;
 							}
 						} 
 						else if ( AbsScrollByAmount == PrecedingDistance)
 						{
-							NewScrollOffset -= FMath::Fractional( NewScrollOffset );
+							NewScrollOffset -= DoubleFractional( NewScrollOffset );
 							break;
 						}
 						else
@@ -1992,7 +2002,7 @@ protected:
 			}
 
 
-			return ScrollTo( NewScrollOffset );
+			return ScrollTo( (float)NewScrollOffset );
 		}
 
 		return 0;
