Color block restyling for details panels
- Simplified color block drawing
- Added rounded corners for gradient elements
- Deprecated bIgnoreAlpha in favor of AlphaDisplayMode::Ignore

diff --git a/Engine/Source/Editor/DetailCustomizations/Private/ColorStructCustomization.cpp b/Engine/Source/Editor/DetailCustomizations/Private/ColorStructCustomization.cpp
--- a/Engine/Source/Editor/DetailCustomizations/Private/ColorStructCustomization.cpp
+++ b/Engine/Source/Editor/DetailCustomizations/Private/ColorStructCustomization.cpp
@@ -12,6 +12,7 @@
 #include "IPropertyUtilities.h"
 #include "Widgets/Colors/SColorPicker.h"
 #include "Framework/Application/SlateApplication.h"
+#include "Widgets/Images/SImage.h"
 
 
 #define LOCTEXT_NAMESPACE "FColorStructCustomization"
@@ -50,7 +51,7 @@ void FColorStructCustomization::MakeHeaderRow(TSharedRef<class IPropertyHandle>&
 	const FText DisplayToolTipOverride = FText::GetEmpty();
 	
 	TSharedPtr<SWidget> ColorWidget;
-	float ContentWidth = 250.0f;
+	float ContentWidth = 125.0f;
 
 	TWeakPtr<IPropertyHandle> StructWeakHandlePtr = StructPropertyHandle;
 
@@ -86,49 +87,46 @@ FColorStructCustomization::~FColorStructCustomization()
 
 TSharedRef<SWidget> FColorStructCustomization::CreateColorWidget(TWeakPtr<IPropertyHandle> StructWeakHandlePtr)
 {
-	FSlateFontInfo NormalText = IDetailLayoutBuilder::GetDetailFont();
-
-	return SNew(SHorizontalBox)
-		+ SHorizontalBox::Slot()
-		.VAlign(VAlign_Center)
-		.Padding(0.0f, 2.0f)
+	return
+		SNew(SBox)
+		.Padding(FMargin(0,0,4.0f,0.0f))
 		[
-			SNew(SOverlay)
-			+SOverlay::Slot()
-			[
-				// Displays the color with alpha unless it is ignored
-				SAssignNew(ColorPickerParentWidget, SColorBlock)
-				.Color(this, &FColorStructCustomization::OnGetColorForColorBlock)
-				.ShowBackgroundForAlpha(true)
-				.IgnoreAlpha(bIgnoreAlpha)
-				.OnMouseButtonDown(this, &FColorStructCustomization::OnMouseButtonDownColorBlock)
-				.Size(FVector2D(35.0f, 12.0f))
-				.IsEnabled(this, &FColorStructCustomization::IsValueEnabled, StructWeakHandlePtr)
-			]
-			+SOverlay::Slot()
-			.HAlign(HAlign_Center)
-			.VAlign(VAlign_Center)
+			SAssignNew(ColorWidgetBackgroundBorder, SBorder)
+			.Padding(1)
+			.BorderImage(FAppStyle::Get().GetBrush("ColorPicker.RoundedSolidBackground"))
+			.BorderBackgroundColor(this, &FColorStructCustomization::GetColorWidgetBorderColor)
 			[
-				SNew(STextBlock)
-				.Text(NSLOCTEXT("PropertyEditor", "MultipleValues", "Multiple Values"))
-				.Font(NormalText)
-				.ColorAndOpacity(FSlateColor(FLinearColor::Black)) // we know the background is always white, so can safely set this to black
-				.Visibility(this, &FColorStructCustomization::GetMultipleValuesTextVisibility)
+				SNew(SOverlay)
+				+ SOverlay::Slot()
+				[
+					SAssignNew(ColorPickerParentWidget, SColorBlock)
+					.AlphaBackgroundBrush(FAppStyle::Get().GetBrush("ColorPicker.RoundedAlphaBackground"))
+					.SolidBackgroundBrush(FAppStyle::Get().GetBrush("ColorPicker.RoundedSolidBackground"))
+					.Color(this, &FColorStructCustomization::OnGetColorForColorBlock)
+					.ShowBackgroundForAlpha(true)
+					.AlphaDisplayMode(bIgnoreAlpha ? EColorBlockAlphaDisplayMode::Ignore : EColorBlockAlphaDisplayMode::Separate)
+					.OnMouseButtonDown(this, &FColorStructCustomization::OnMouseButtonDownColorBlock)
+					.Size(FVector2D(70.0f, 22.0f))
+					.CornerRadius(4.0f)
+					.IsEnabled(this, &FColorStructCustomization::IsValueEnabled, StructWeakHandlePtr)
+				]
+				+ SOverlay::Slot()
+				[
+					SNew(SBorder)
+					.Visibility(this, &FColorStructCustomization::GetMultipleValuesTextVisibility)
+					.BorderImage(FAppStyle::Get().GetBrush("ColorPicker.MultipleValuesBackground"))
+					.VAlign(VAlign_Center)
+					.ForegroundColor(FAppStyle::Get().GetWidgetStyle<FEditableTextBoxStyle>("NormalEditableTextBox").ForegroundColor)
+					.Padding(FMargin(12.0f, 2.0f))
+					[
+						SNew(STextBlock)
+						.Text(NSLOCTEXT("PropertyEditor", "MultipleValues", "Multiple Values"))
+						.Font(IDetailLayoutBuilder::GetDetailFont())
+					]
+				]
 			]
-		]
-		+ SHorizontalBox::Slot()
-		.VAlign(VAlign_Center)
-		.Padding(0.0f, 2.0f)
-		[
-			// Displays the color without alpha
-			SNew(SColorBlock)
-			.Color(this, &FColorStructCustomization::OnGetColorForColorBlock)
-			.ShowBackgroundForAlpha(false)
-			.IgnoreAlpha(true)
-			.OnMouseButtonDown(this, &FColorStructCustomization::OnMouseButtonDownColorBlock)
-			.Size(FVector2D(35.0f, 12.0f))
 		];
-}
+	}
 
 
 void FColorStructCustomization::GetSortedChildren(TSharedRef<IPropertyHandle> InStructPropertyHandle, TArray< TSharedRef<IPropertyHandle> >& OutChildren)
@@ -368,6 +366,20 @@ FLinearColor FColorStructCustomization::OnGetColorForColorBlock() const
 }
 
 
+FSlateColor FColorStructCustomization::OnGetSlateColorForBlock() const
+{
+	FLinearColor Color = OnGetColorForColorBlock();
+	Color.A = 1;
+	return FSlateColor(Color);
+}
+
+FSlateColor FColorStructCustomization::GetColorWidgetBorderColor() const
+{
+	static const FSlateColor HoveredColor = FAppStyle::Get().GetSlateColor("Colors.Hover");
+	static const FSlateColor DefaultColor = FAppStyle::Get().GetSlateColor("Colors.InputOutline");
+	return ColorWidgetBackgroundBorder->IsHovered() ? HoveredColor : DefaultColor;
+}
+
 FPropertyAccess::Result FColorStructCustomization::GetColorAsLinear(FLinearColor& OutColor) const
 {
 	// Default to full alpha in case the alpha component is disabled.
@@ -402,7 +414,7 @@ EVisibility FColorStructCustomization::GetMultipleValuesTextVisibility() const
 {
 	FLinearColor Color;
 	const FPropertyAccess::Result ValueResult = GetColorAsLinear(Color);
-	return (ValueResult == FPropertyAccess::MultipleValues) ? EVisibility::Visible : EVisibility::Collapsed;
+	return (ValueResult == FPropertyAccess::MultipleValues) ? EVisibility::HitTestInvisible : EVisibility::Collapsed;
 }
 
 
diff --git a/Engine/Source/Editor/EditorStyle/Private/StarshipStyle.cpp b/Engine/Source/Editor/EditorStyle/Private/StarshipStyle.cpp
--- a/Engine/Source/Editor/EditorStyle/Private/StarshipStyle.cpp
+++ b/Engine/Source/Editor/EditorStyle/Private/StarshipStyle.cpp
@@ -759,7 +759,7 @@ void FStarshipEditorStyle::FStyle::SetupGeneralStyles()
 #endif // WITH_EDITOR || (IS_PROGRAM && WITH_UNREAL_DEVELOPER_TOOLS)
 
 	Set( "DashedBorder", new BORDER_BRUSH( "Old/DashedBorder", FMargin(6.0f/32.0f) ) );
-	Set( "Checker", new IMAGE_BRUSH( "Old/Checker", Icon16x16, FLinearColor::White, ESlateBrushTileType::Both ) );
+
 	Set( "UniformShadow", new BORDER_BRUSH( "Common/UniformShadow", FMargin( 16.0f / 64.0f ) ) );
 	Set( "UniformShadow_Tint", new BORDER_BRUSH( "Common/UniformShadow_Tint", FMargin( 16.0f / 64.0f ) ) );
 
diff --git a/Engine/Source/Runtime/AppFramework/Private/Widgets/Colors/SColorPicker.cpp b/Engine/Source/Runtime/AppFramework/Private/Widgets/Colors/SColorPicker.cpp
--- a/Engine/Source/Runtime/AppFramework/Private/Widgets/Colors/SColorPicker.cpp
+++ b/Engine/Source/Runtime/AppFramework/Private/Widgets/Colors/SColorPicker.cpp
@@ -1062,7 +1062,7 @@ TSharedRef<SWidget> SColorPicker::MakeColorPreviewBox() const
 							// old color
 							SNew(SColorBlock) 
 								.ColorIsHSV(true) 
-								.IgnoreAlpha(true)
+								.AlphaDisplayMode(EColorBlockAlphaDisplayMode::Ignore)
 								.ToolTipText(LOCTEXT("OldColorToolTip", "Old color without alpha (drag to theme bar to save)"))
 								.Color(OldColor) 
 								.OnMouseButtonDown(const_cast<SColorPicker*>(this), &SColorPicker::HandleOldColorBlockMouseButtonDown, false)
@@ -1094,7 +1094,7 @@ TSharedRef<SWidget> SColorPicker::MakeColorPreviewBox() const
 							// new color
 							SNew(SColorBlock) 
 								.ColorIsHSV(true) 
-								.IgnoreAlpha(true)
+								.AlphaDisplayMode(EColorBlockAlphaDisplayMode::Ignore)
 								.ToolTipText(LOCTEXT("NewColorToolTip", "New color without alpha (drag to theme bar to save)"))
 								.Color(this, &SColorPicker::GetCurrentColor)
 								.OnMouseButtonDown(const_cast<SColorPicker*>(this), &SColorPicker::HandleNewColorBlockMouseButtonDown, false)
diff --git a/Engine/Source/Runtime/AppFramework/Private/Widgets/Colors/SColorThemes.cpp b/Engine/Source/Runtime/AppFramework/Private/Widgets/Colors/SColorThemes.cpp
--- a/Engine/Source/Runtime/AppFramework/Private/Widgets/Colors/SColorThemes.cpp
+++ b/Engine/Source/Runtime/AppFramework/Private/Widgets/Colors/SColorThemes.cpp
@@ -71,7 +71,7 @@ TSharedPtr<SWidget> FColorDragDrop::GetDefaultDecorator() const
 			SNew(SColorBlock) 
 				.Color(Color)
 				.ColorIsHSV(true) 
-				.IgnoreAlpha(bIgnoreAlpha) 
+				.AlphaDisplayMode(bIgnoreAlpha ? EColorBlockAlphaDisplayMode::Ignore : EColorBlockAlphaDisplayMode::Combined)
 				.ShowBackgroundForAlpha(bShowBackgroundForAlpha) 
 				.UseSRGB(bUseSRGB)
 		];
@@ -257,8 +257,8 @@ void SThemeColorBlock::Construct(const FArguments& InArgs )
 						SNew(SColorBlock)
 							.Color(this, &SThemeColorBlock::GetColor)
 							.ColorIsHSV(true)
-							.IgnoreAlpha(TAttribute<bool>::Create(TAttribute<bool>::FGetter::CreateSP(this, &SThemeColorBlock::OnReadIgnoreAlpha)))
-							.ShowBackgroundForAlpha(TAttribute<bool>::Create(TAttribute<bool>::FGetter::CreateSP(this, &SThemeColorBlock::OnReadShowBackgroundForAlpha)))
+							.AlphaDisplayMode(this, &SThemeColorBlock::OnGetAlphaDisplayMode)
+							.ShowBackgroundForAlpha(this, &SThemeColorBlock::OnReadShowBackgroundForAlpha)
 							.UseSRGB(bUseSRGB)
 					]
 				]
@@ -318,26 +318,15 @@ void SThemeColorBlock::Construct(const FArguments& InArgs )
 			.BorderBackgroundColor(this, &SThemeColorBlock::HandleBorderColor)
 			.Padding(FMargin(1.0f))
 			.ToolTip(ColorTooltip)
+			.VAlign(VAlign_Center)
 			[
-				SNew(SHorizontalBox)
-				+ SHorizontalBox::Slot()
-				[
-					SNew(SColorBlock)
-					.Color(this, &SThemeColorBlock::GetColor)
-					.ColorIsHSV(true)
-					.IgnoreAlpha(true)
-					.ShowBackgroundForAlpha(false)
-					.UseSRGB(bUseSRGB)
-				]
-				+ SHorizontalBox::Slot()
-				[
-					SNew(SColorBlock)
-					.Color(this, &SThemeColorBlock::GetColor)
-					.ColorIsHSV(true)
-					.IgnoreAlpha(TAttribute<bool>::Create(TAttribute<bool>::FGetter::CreateSP(this, &SThemeColorBlock::OnReadIgnoreAlpha)))
-					.ShowBackgroundForAlpha(TAttribute<bool>::Create(TAttribute<bool>::FGetter::CreateSP(this, &SThemeColorBlock::OnReadShowBackgroundForAlpha)))
-					.UseSRGB(bUseSRGB)
-				]
+				SNew(SColorBlock)
+				.Color(this, &SThemeColorBlock::GetColor)
+				.AlphaDisplayMode(EColorBlockAlphaDisplayMode::Separate)
+				.ColorIsHSV(true)
+				.ShowBackgroundForAlpha(true)
+				.UseSRGB(bUseSRGB)
+				.Size(FVector2D(32,16))
 			]
 	];
 }
@@ -486,9 +475,9 @@ FText SThemeColorBlock::FormatToolTipText(const FText& ColorIdentifier, float Va
 }
 #undef LOCTEXT_NAMESPACE
 
-bool SThemeColorBlock::OnReadIgnoreAlpha() const
+EColorBlockAlphaDisplayMode SThemeColorBlock::OnGetAlphaDisplayMode() const
 {
-	return !bUseAlpha.Get();
+	return !bUseAlpha.Get() ? EColorBlockAlphaDisplayMode::Ignore : EColorBlockAlphaDisplayMode::Combined;
 }
 
 bool SThemeColorBlock::OnReadShowBackgroundForAlpha() const
diff --git a/Engine/Source/Runtime/AppFramework/Public/Widgets/Colors/SColorThemes.h b/Engine/Source/Runtime/AppFramework/Public/Widgets/Colors/SColorThemes.h
--- a/Engine/Source/Runtime/AppFramework/Public/Widgets/Colors/SColorThemes.h
+++ b/Engine/Source/Runtime/AppFramework/Public/Widgets/Colors/SColorThemes.h
@@ -14,6 +14,7 @@
 #include "Widgets/SPanel.h"
 #include "Widgets/Views/STableViewBase.h"
 #include "Widgets/Views/STableRow.h"
+#include "Widgets/Colors/SColorBlock.h"
 
 class FArrangedChildren;
 class SEditableTextBox;
@@ -221,7 +222,7 @@ private:
 	 */
 	FText FormatToolTipText(const FText& ColorIdentifier, float Value) const;
 	
-	bool OnReadIgnoreAlpha() const;
+	EColorBlockAlphaDisplayMode OnGetAlphaDisplayMode() const;
 	bool OnReadShowBackgroundForAlpha() const;
 	EVisibility OnGetAlphaVisibility() const;
 
diff --git a/Engine/Source/Runtime/Slate/Private/Widgets/Colors/SColorBlock.cpp b/Engine/Source/Runtime/Slate/Private/Widgets/Colors/SColorBlock.cpp
--- a/Engine/Source/Runtime/Slate/Private/Widgets/Colors/SColorBlock.cpp
+++ b/Engine/Source/Runtime/Slate/Private/Widgets/Colors/SColorBlock.cpp
@@ -10,79 +10,71 @@
  *
  * @param	InArgs	The declaration data for this widget
  */
-void SColorBlock::Construct( const FArguments& InArgs )
+void SColorBlock::Construct(const FArguments& InArgs)
 {
 	Color = InArgs._Color;
+	AlphaBackgroundBrush = InArgs._AlphaBackgroundBrush;
+	SolidBackgroundBrush = InArgs._SolidBackgroundBrush;
+	GradientCornerRadius = InArgs._CornerRadius;
 	ColorIsHSV = InArgs._ColorIsHSV;
-	IgnoreAlpha = InArgs._IgnoreAlpha;
+	AlphaDisplayMode = InArgs._AlphaDisplayMode;
 	ShowBackgroundForAlpha = InArgs._ShowBackgroundForAlpha;
 	MouseButtonDownHandler = InArgs._OnMouseButtonDown;
 	bUseSRGB = InArgs._UseSRGB;
 	ColorBlockSize = InArgs._Size;
 }
 
-int32 SColorBlock::OnPaint( const FPaintArgs& Args, const FGeometry& AllottedGeometry, const FSlateRect& MyCullingRect, FSlateWindowElementList& OutDrawElements, int32 LayerId, const FWidgetStyle& InWidgetStyle, bool bParentEnabled ) const
+int32 SColorBlock::OnPaint(const FPaintArgs& Args, const FGeometry& AllottedGeometry, const FSlateRect& MyCullingRect, FSlateWindowElementList& OutDrawElements, int32 LayerId, const FWidgetStyle& InWidgetStyle, bool bParentEnabled) const
 {
-	const FSlateBrush* GenericBrush = FCoreStyle::Get().GetBrush( "GenericWhiteBox" );
-
 	const ESlateDrawEffect DrawEffects = ESlateDrawEffect::None;
-	
+
+	EColorBlockAlphaDisplayMode DisplayMode = AlphaDisplayMode.Get();
+	bool bIgnoreAlpha = DisplayMode == EColorBlockAlphaDisplayMode::Ignore;
 	FLinearColor InColor = Color.Get();
 	if (ColorIsHSV.Get())
 	{
 		InColor = InColor.HSVToLinearRGB();
 	}
-	if (IgnoreAlpha.Get())
-	{
-		InColor.A = 1.f;
-	}
-	
-	const FColor DrawColor = InColor.ToFColor(bUseSRGB.Get());
-	if( ShowBackgroundForAlpha.Get() && DrawColor.A < 255 )
+
+	if (ShowBackgroundForAlpha.Get() && InColor.A < 1.0f && !bIgnoreAlpha)
 	{
-		// If we are showing a background pattern and the colors is transparent, draw a checker pattern
-		const FSlateBrush* CheckerBrush = FCoreStyle::Get().GetBrush("ColorPicker.AlphaBackground");
-		FSlateDrawElement::MakeBox( OutDrawElements, LayerId, AllottedGeometry.ToPaintGeometry(), CheckerBrush, DrawEffects );
+		FSlateDrawElement::MakeBox(OutDrawElements, LayerId, AllottedGeometry.ToPaintGeometry(), AlphaBackgroundBrush.Get(), DrawEffects);
 	}
-	
-	// determine if it is HDR
-	const float MaxRGB = FMath::Max3(InColor.R, InColor.G, InColor.B);
-	if (MaxRGB > 1.f)
-	{
-		FLinearColor NormalizedLinearColor = InColor / MaxRGB;
-		NormalizedLinearColor.A = InColor.A;
-		const FLinearColor DrawNormalizedColor = InWidgetStyle.GetColorAndOpacityTint() * NormalizedLinearColor.ToFColor(bUseSRGB.Get());
 
-		FLinearColor ClampedLinearColor = InColor;
-		ClampedLinearColor.A = InColor.A * MaxRGB;
-		const FLinearColor DrawClampedColor = InWidgetStyle.GetColorAndOpacityTint() * ClampedLinearColor.ToFColor(bUseSRGB.Get());
+	TArray<FSlateGradientStop> GradientStops;
 
-		TArray<FSlateGradientStop> GradientStops;
-		
-		GradientStops.Add( FSlateGradientStop( FVector2D::ZeroVector, DrawNormalizedColor ) );
-		GradientStops.Add( FSlateGradientStop( AllottedGeometry.GetLocalSize() * 0.5f, DrawClampedColor ) );
-		GradientStops.Add( FSlateGradientStop( AllottedGeometry.GetLocalSize(), DrawNormalizedColor ) );
-
-		FSlateDrawElement::MakeGradient(
-			OutDrawElements,
-			LayerId + 1,
-			AllottedGeometry.ToPaintGeometry(),
-			GradientStops,
-			(AllottedGeometry.GetLocalSize().X > AllottedGeometry.GetLocalSize().Y) ? Orient_Vertical : Orient_Horizontal,
-			DrawEffects
-		);
-	}
-	else
+	switch (DisplayMode)
 	{
-		FSlateDrawElement::MakeBox( OutDrawElements, LayerId + 1, AllottedGeometry.ToPaintGeometry(), GenericBrush, DrawEffects, InWidgetStyle.GetColorAndOpacityTint() * DrawColor );
+	default:
+	case EColorBlockAlphaDisplayMode::Combined:
+		MakeSection(GradientStops, FVector2D::ZeroVector, AllottedGeometry.GetLocalSize(), InColor, InWidgetStyle, false);
+		break;
+	case EColorBlockAlphaDisplayMode::Separate:
+		MakeSection(GradientStops, FVector2D::ZeroVector, AllottedGeometry.GetLocalSize() * 0.5f, InColor, InWidgetStyle, false);
+		MakeSection(GradientStops, AllottedGeometry.GetLocalSize() * 0.5f, AllottedGeometry.GetLocalSize(), InColor, InWidgetStyle, true);
+		break;
+	case EColorBlockAlphaDisplayMode::Ignore:
+		MakeSection(GradientStops, FVector2D::ZeroVector, AllottedGeometry.GetLocalSize(), InColor, InWidgetStyle, true);
+		break;
 	}
 
+	FSlateDrawElement::MakeGradient(
+		OutDrawElements,
+		LayerId+1,
+		AllottedGeometry.ToPaintGeometry(),
+		MoveTemp(GradientStops),
+		(AllottedGeometry.GetLocalSize().X > AllottedGeometry.GetLocalSize().Y) ? Orient_Vertical : Orient_Horizontal,
+		DrawEffects,
+		GradientCornerRadius.Get(0.0f)
+	);
+
+
 	return LayerId + 1;
 }
 
-FReply SColorBlock::OnMouseButtonDown( const FGeometry& MyGeometry, const FPointerEvent& MouseEvent )
+FReply SColorBlock::OnMouseButtonDown(const FGeometry& MyGeometry, const FPointerEvent& MouseEvent)
 {
-	if ( MouseButtonDownHandler.IsBound() )
+	if (MouseButtonDownHandler.IsBound())
 	{
 		// If a handler is assigned, call it.
 		return MouseButtonDownHandler.Execute(MyGeometry, MouseEvent);
@@ -94,7 +86,38 @@ FReply SColorBlock::OnMouseButtonDown( const FGeometry& MyGeometry, const FPoint
 	}
 }
 
-FVector2D SColorBlock::ComputeDesiredSize( float ) const
+FVector2D SColorBlock::ComputeDesiredSize(float) const
 {
 	return ColorBlockSize.Get();
 }
+
+void SColorBlock::MakeSection(TArray<FSlateGradientStop>& OutGradientStops, FVector2D StartPt, FVector2D EndPt, FLinearColor InColor, const FWidgetStyle& InWidgetStyle, bool bIgnoreAlpha) const
+{
+	// determine if it is HDR
+	const float MaxRGB = FMath::Max3(InColor.R, InColor.G, InColor.B);
+	if (MaxRGB > 1.f)
+	{
+		const float Alpha = bIgnoreAlpha ? 1.0f : InColor.A;
+		FLinearColor NormalizedLinearColor = InColor / MaxRGB;
+		NormalizedLinearColor.A = Alpha;
+		const FLinearColor DrawNormalizedColor = InWidgetStyle.GetColorAndOpacityTint() * NormalizedLinearColor.ToFColor(bUseSRGB.Get());
+
+		FLinearColor ClampedLinearColor = InColor;
+		ClampedLinearColor.A = Alpha * MaxRGB;
+		const FLinearColor DrawClampedColor = InWidgetStyle.GetColorAndOpacityTint() * ClampedLinearColor.ToFColor(bUseSRGB.Get());
+
+		OutGradientStops.Add(FSlateGradientStop(StartPt, DrawNormalizedColor));
+		OutGradientStops.Add(FSlateGradientStop((StartPt + EndPt) * 0.5f, DrawClampedColor));
+		OutGradientStops.Add(FSlateGradientStop(EndPt, DrawNormalizedColor));
+	}
+	else
+	{
+		FColor DrawColor = InColor.ToFColor(bUseSRGB.Get());
+		if (bIgnoreAlpha)
+		{
+			DrawColor.A = 255;
+		}
+		OutGradientStops.Add(FSlateGradientStop(StartPt, InWidgetStyle.GetColorAndOpacityTint() * DrawColor));
+		OutGradientStops.Add(FSlateGradientStop(EndPt, InWidgetStyle.GetColorAndOpacityTint() * DrawColor));
+	}
+}
diff --git a/Engine/Source/Runtime/Slate/Public/Widgets/Colors/SColorBlock.h b/Engine/Source/Runtime/Slate/Public/Widgets/Colors/SColorBlock.h
--- a/Engine/Source/Runtime/Slate/Public/Widgets/Colors/SColorBlock.h
+++ b/Engine/Source/Runtime/Slate/Public/Widgets/Colors/SColorBlock.h
@@ -7,44 +7,78 @@
 #include "Input/Reply.h"
 #include "Widgets/DeclarativeSyntaxSupport.h"
 #include "Widgets/SLeafWidget.h"
+#include "Styling/AppStyle.h"
 
 class FPaintArgs;
 class FSlateWindowElementList;
 
+enum class EColorBlockAlphaDisplayMode : uint8
+{
+	// Draw a single block that draws color and opacity as one. I.E the entire block will be semi-transparent if opacity < 1
+	Combined,
+	// The color block is split into in half. The left half draws the color with opacity and the right half draws without any opacity
+	Separate,
+	// Alpha is omitted from display
+	Ignore,
+};
+
 class SLATE_API SColorBlock : public SLeafWidget
 {
 public:
 
-	SLATE_BEGIN_ARGS( SColorBlock )
-		: _Color( FLinearColor::White )
-		, _ColorIsHSV( false )
-		, _IgnoreAlpha( false )
-		, _ShowBackgroundForAlpha( false )
+	SLATE_BEGIN_ARGS(SColorBlock)
+		: _Color(FLinearColor::White)
+		, _AlphaBackgroundBrush(FAppStyle::Get().GetBrush("ColorPicker.AlphaBackground"))
+		, _SolidBackgroundBrush(FAppStyle::Get().GetBrush("GenericWhiteBox"))
+		, _CornerRadius(0.0f)
+		, _ColorIsHSV(false)
+		, _ShowBackgroundForAlpha(false)
 		, _UseSRGB(true)
+		, _AlphaDisplayMode(EColorBlockAlphaDisplayMode::Combined)
+		, _Size(FVector2D(16, 16))
 		, _OnMouseButtonDown()
-		, _Size( FVector2D(16,16) )
-		{}
+	{}
 
 		/** The color to display for this color block */
-		SLATE_ATTRIBUTE( FLinearColor, Color )
+		SLATE_ATTRIBUTE(FLinearColor, Color)
 
-		/** Whether the color displayed is HSV or not */
-		SLATE_ATTRIBUTE( bool, ColorIsHSV )
+		/** Background to display for when there is a color with transparency. Irrelevant if ignoring alpha */
+		SLATE_ATTRIBUTE(const FSlateBrush*, AlphaBackgroundBrush)
+	
+		/** Background to use to display when there is a solid color. This background is tinted by the solid color */
+		SLATE_ATTRIBUTE(const FSlateBrush*, SolidBackgroundBrush)
 
-		/** Whether to ignore alpha entirely from the input color */
-		SLATE_ATTRIBUTE( bool, IgnoreAlpha )
+		/** Rounding to apply to the corners of the block */
+		SLATE_ATTRIBUTE(float, CornerRadius)
+
+		/** Whether the color displayed is HSV or not */
+		SLATE_ATTRIBUTE(bool, ColorIsHSV)
 
 		/** Whether to display a background for viewing opacity. Irrelevant if ignoring alpha */
-		SLATE_ATTRIBUTE( bool, ShowBackgroundForAlpha )
+		SLATE_ATTRIBUTE(bool, ShowBackgroundForAlpha)
 
 		/** Whether to display sRGB color */
-		SLATE_ATTRIBUTE( bool, UseSRGB )
+		SLATE_ATTRIBUTE(bool, UseSRGB)
 
-		/** A handler to activate when the mouse is pressed. */
-		SLATE_EVENT( FPointerEventHandler, OnMouseButtonDown )
+		/** How the color block displays color and opacity */
+		SLATE_ATTRIBUTE(EColorBlockAlphaDisplayMode, AlphaDisplayMode)
 
 		/** How big should this color block be? */
-		SLATE_ATTRIBUTE( FVector2D, Size )
+		SLATE_ATTRIBUTE(FVector2D, Size)
+
+		/** A handler to activate when the mouse is pressed. */
+		SLATE_EVENT(FPointerEventHandler, OnMouseButtonDown)
+
+		UE_DEPRECATED(5.0, "IgnoreAlpha is deprecated. Set AlphaDisplayMode to EColorBlockAlphaDisplayMode::Ignore instead")
+		FArguments& IgnoreAlpha(bool bInIgnoreAlpha)
+		{
+			if (bInIgnoreAlpha)
+			{
+				_AlphaDisplayMode = EColorBlockAlphaDisplayMode::Ignore;
+			}
+			return Me();
+		}
+
 	SLATE_END_ARGS()
 
 	/**
@@ -52,33 +86,38 @@ public:
 	 *
 	 * @param	InArgs	The declaration data for this widget
 	 */
-	void Construct( const FArguments& InArgs );
+	void Construct(const FArguments& InArgs);
 
 private:
-
 	// SWidget overrides
-
-	virtual int32 OnPaint( const FPaintArgs& Args, const FGeometry& AllottedGeometry, const FSlateRect& MyCullingRect, FSlateWindowElementList& OutDrawElements, int32 LayerId, const FWidgetStyle& InWidgetStyle, bool bParentEnabled ) const override;
-	virtual FReply OnMouseButtonDown( const FGeometry& MyGeometry, const FPointerEvent& MouseEvent ) override;
+	virtual int32 OnPaint(const FPaintArgs& Args, const FGeometry& AllottedGeometry, const FSlateRect& MyCullingRect, FSlateWindowElementList& OutDrawElements, int32 LayerId, const FWidgetStyle& InWidgetStyle, bool bParentEnabled) const override;
+	virtual FReply OnMouseButtonDown(const FGeometry& MyGeometry, const FPointerEvent& MouseEvent) override;
 	virtual FVector2D ComputeDesiredSize(float) const override;
 
+	void MakeSection(TArray<FSlateGradientStop>& OutGradientStops, FVector2D StartPt, FVector2D EndPt, FLinearColor Color, const FWidgetStyle& InWidgetStyle, bool bIgnoreAlpha) const;
 private:
 
 	/** The color to display for this color block */
-	TAttribute< FLinearColor > Color;
+	TAttribute<FLinearColor> Color;
+
+	TAttribute<const FSlateBrush*> AlphaBackgroundBrush;
+
+	TAttribute<const FSlateBrush*> SolidBackgroundBrush;
+
+	TAttribute<float> GradientCornerRadius;
 
 	/** Whether the color displayed is HSV or not */
-	TAttribute< bool > ColorIsHSV;
-	
+	TAttribute<bool> ColorIsHSV;
+
 	/** Whether to ignore alpha entirely from the input color */
-	TAttribute< bool > IgnoreAlpha;
+	TAttribute<EColorBlockAlphaDisplayMode> AlphaDisplayMode;
 
 	/** Whether to display a background for viewing opacity. Irrelevant if ignoring alpha */
-	TAttribute< bool > ShowBackgroundForAlpha;
+	TAttribute<bool> ShowBackgroundForAlpha;
 
 	/** Whether to display sRGB color */
 	TAttribute<bool> bUseSRGB;
-	
+
 	/** A handler to activate when the mouse is pressed. */
 	FPointerEventHandler MouseButtonDownHandler;
 
diff --git a/Engine/Source/Runtime/SlateCore/Private/Rendering/DrawElements.cpp b/Engine/Source/Runtime/SlateCore/Private/Rendering/DrawElements.cpp
--- a/Engine/Source/Runtime/SlateCore/Private/Rendering/DrawElements.cpp
+++ b/Engine/Source/Runtime/SlateCore/Private/Rendering/DrawElements.cpp
@@ -439,7 +439,7 @@ void FSlateDrawElement::MakeShapedText(FSlateWindowElementList& ElementList, uin
 	Element.Init(ElementList, EElementType::ET_ShapedText, InLayer, PaintGeometry, InDrawEffects);
 }
 
-void FSlateDrawElement::MakeGradient( FSlateWindowElementList& ElementList, uint32 InLayer, const FPaintGeometry& PaintGeometry, TArray<FSlateGradientStop> InGradientStops, EOrientation InGradientType, ESlateDrawEffect InDrawEffects )
+void FSlateDrawElement::MakeGradient( FSlateWindowElementList& ElementList, uint32 InLayer, const FPaintGeometry& PaintGeometry, TArray<FSlateGradientStop> InGradientStops, EOrientation InGradientType, ESlateDrawEffect InDrawEffects, float CornerRadius)
 {
 	PaintGeometry.CommitTransformsIfUsingLegacyConstructor();
 
@@ -452,7 +452,7 @@ void FSlateDrawElement::MakeGradient( FSlateWindowElementList& ElementList, uint
 
 	FSlateGradientPayload& DataPayload = ElementList.CreatePayload<FSlateGradientPayload>(Element);
 
-	DataPayload.SetGradient(InGradientStops, InGradientType);
+	DataPayload.SetGradient(InGradientStops, InGradientType, CornerRadius);
 
 	Element.Init(ElementList, EElementType::ET_Gradient, InLayer, PaintGeometry, InDrawEffects);
 }
diff --git a/Engine/Source/Runtime/SlateCore/Private/Rendering/ElementBatcher.cpp b/Engine/Source/Runtime/SlateCore/Private/Rendering/ElementBatcher.cpp
--- a/Engine/Source/Runtime/SlateCore/Private/Rendering/ElementBatcher.cpp
+++ b/Engine/Source/Runtime/SlateCore/Private/Rendering/ElementBatcher.cpp
@@ -543,7 +543,7 @@ void FSlateElementBatcher::AddBoxElement(const FSlateDrawElement& DrawElement)
 	const FColor Tint = PackVertexColor(DrawElementPayload.GetTint());
 	const FSlateRenderTransform& ElementRenderTransform = DrawElement.GetRenderTransform();
 	const FSlateRenderTransform RenderTransform = DrawElement.GetRenderTransform();// GetBoxRenderTransform(DrawElement);
-	const FVector2D& LocalSize = DrawElement.GetLocalSize();
+	const FVector2D LocalSize = DrawElement.GetLocalSize();
 
 	const ESlateDrawEffect InDrawEffects = DrawElement.GetDrawEffects();
 	const int32 Layer = DrawElement.GetLayer();
@@ -1429,21 +1429,32 @@ template<ESlateVertexRounding Rounding>
 void FSlateElementBatcher::AddGradientElement( const FSlateDrawElement& DrawElement )
 {
 	const FSlateRenderTransform& RenderTransform = DrawElement.GetRenderTransform();
-	const FVector2D& LocalSize = DrawElement.GetLocalSize();
+	const FVector2D LocalSize = DrawElement.GetLocalSize();
 	const FSlateGradientPayload& InPayload = DrawElement.GetDataPayload<FSlateGradientPayload>();
 	const ESlateDrawEffect InDrawEffects = DrawElement.GetDrawEffects();
 	const int32 Layer = DrawElement.GetLayer();
+	const float DrawScale = DrawElement.GetScale();
 
 	// There must be at least one gradient stop
 	check( InPayload.GradientStops.Num() > 0 );
 
+	FShaderParams ShaderParams;
+
+	ESlateShader ShaderType = ESlateShader::Default;
+	if (InPayload.CornerRadius > 0)
+	{
+		ShaderType = ESlateShader::RoundedBox;
+		ShaderParams.PixelParams = FVector4(InPayload.CornerRadius, 0.0f, LocalSize.X, LocalSize.Y);
+		ShaderParams.PixelParams2 = FVector4(FLinearColor::Transparent);
+	}
+
 	FSlateRenderBatch& RenderBatch = 
 		CreateRenderBatch( 
 			Layer,
-			FShaderParams(),
+			ShaderParams,
 			nullptr,
 			ESlateDrawPrimitive::TriangleList,
-			ESlateShader::Default,
+			ShaderType,
 			InDrawEffects,
 			DrawElement.GetBatchFlags(),
 			DrawElement);
@@ -1504,6 +1515,9 @@ void FSlateElementBatcher::AddGradientElement( const FSlateDrawElement& DrawElem
 		// The end vertex at this stop
 		FVector2D EndPt;
 
+		FVector2D StartUV;
+		FVector2D EndUV;
+
 		if( InPayload.GradientType == Orient_Vertical )
 		{
 			// Gradient stop is vertical so gradients to left to right
@@ -1512,28 +1526,41 @@ void FSlateElementBatcher::AddGradientElement( const FSlateDrawElement& DrawElem
 			// Gradient stops are interpreted in local space.
 			StartPt.X += CurStop.Position.X;
 			EndPt.X += CurStop.Position.X;
+
+			StartUV.X = StartPt.X / TopRight.X;
+			StartUV.Y = 0;
+
+			EndUV.X = EndPt.X / TopRight.X;
+			EndUV.Y = 1;
 		}
 		else
 		{
+
 			// Gradient stop is horizontal so gradients to top to bottom
 			StartPt = TopLeft;
 			EndPt = TopRight;
 			// Gradient stops are interpreted in local space.
 			StartPt.Y += CurStop.Position.Y;
 			EndPt.Y += CurStop.Position.Y;
+
+			StartUV.X = 0;
+			StartUV.Y = StartPt.Y / BotLeft.Y;
+
+			EndUV.X = 1;
+			EndUV.Y = StartPt.Y / BotLeft.Y;
 		}
 
 		if( StopIndex == 0 )
 		{
 			// First stop does not have a full quad yet so do not create indices
-			RenderBatch.AddVertex(FSlateVertex::Make<Rounding>(RenderTransform, StartPt, FVector2D::ZeroVector, FVector2D::ZeroVector, PackVertexColor(CurStop.Color)));
-			RenderBatch.AddVertex(FSlateVertex::Make<Rounding>(RenderTransform, EndPt, FVector2D::ZeroVector, FVector2D::ZeroVector, PackVertexColor(CurStop.Color)));
+			RenderBatch.AddVertex(FSlateVertex::Make<Rounding>(RenderTransform, StartPt, LocalSize, DrawScale, FVector4(StartUV.X, StartUV.Y, 0,0), PackVertexColor(CurStop.Color)));
+			RenderBatch.AddVertex(FSlateVertex::Make<Rounding>(RenderTransform, EndPt, LocalSize, DrawScale, FVector4(EndUV.X, EndUV.Y, 0, 0), PackVertexColor(CurStop.Color)));
 		}
 		else
 		{
 			// All stops after the first have indices and generate quads
-			RenderBatch.AddVertex(FSlateVertex::Make<Rounding>(RenderTransform, StartPt, FVector2D::ZeroVector, FVector2D::ZeroVector, PackVertexColor(CurStop.Color)));
-			RenderBatch.AddVertex(FSlateVertex::Make<Rounding>(RenderTransform, EndPt, FVector2D::ZeroVector, FVector2D::ZeroVector, PackVertexColor(CurStop.Color)));
+			RenderBatch.AddVertex(FSlateVertex::Make<Rounding>(RenderTransform, StartPt, LocalSize, DrawScale, FVector4(StartUV.X, StartUV.Y, 0, 0), PackVertexColor(CurStop.Color)));
+			RenderBatch.AddVertex(FSlateVertex::Make<Rounding>(RenderTransform, EndPt, LocalSize, DrawScale, FVector4(EndUV.X, EndUV.Y, 0, 0), PackVertexColor(CurStop.Color)));
 
 			// Connect the indices to the previous vertices
 			RenderBatch.AddIndex(IndexStart - 2);
diff --git a/Engine/Source/Runtime/SlateCore/Private/Styling/StarshipCoreStyle.cpp b/Engine/Source/Runtime/SlateCore/Private/Styling/StarshipCoreStyle.cpp
--- a/Engine/Source/Runtime/SlateCore/Private/Styling/StarshipCoreStyle.cpp
+++ b/Engine/Source/Runtime/SlateCore/Private/Styling/StarshipCoreStyle.cpp
@@ -279,6 +279,8 @@ TSharedRef<ISlateStyle> FStarshipCoreStyle::Create()
 		Style->Set("AppIconPadding", FMargin(11, 11, 3, 5));
 		Style->Set("AppIconPadding.Small", FMargin(4, 4, 0, 0));
 
+		Style->Set("Checker", new IMAGE_BRUSH("Starship/Common/Checker", Icon16x16, FLinearColor::White, ESlateBrushTileType::Both));
+
 		Style->Set("Icons.Denied", new IMAGE_BRUSH("Icons/denied_16x", Icon16x16));
 	
 		Style->Set("Icons.Help", new IMAGE_BRUSH("Icons/icon_help_16x", Icon16x16));
@@ -1542,8 +1544,11 @@ void FStarshipCoreStyle::SetupColorPickerStyles(TSharedRef<FStyle>& Style)
 {
 	// SColorPicker defaults...
 	{
-		Style->Set("ColorPicker.Border", new FSlateRoundedBoxBrush(FStyleColors::Background, 4));
-		Style->Set("ColorPicker.AlphaBackground", new IMAGE_BRUSH("Common/Checker", Icon16x16, FLinearColor::White, ESlateBrushTileType::Both));
+		Style->Set("ColorPicker.RoundedSolidBackground", new FSlateRoundedBoxBrush(FStyleColors::White, InputFocusRadius));
+		Style->Set("ColorPicker.RoundedAlphaBackground", new FSlateRoundedBoxBrush(FName(*RootToContentDir("Starship/Common/Checker", TEXT(".png"))), FLinearColor::White, InputFocusRadius, FLinearColor::White, 0.0f, Icon16x16, ESlateBrushTileType::Both));
+		Style->Set("ColorPicker.RoundedInputBorder", new FSlateRoundedBoxBrush(FStyleColors::Transparent, InputFocusRadius, FStyleColors::InputOutline, InputFocusThickness));
+		Style->Set("ColorPicker.MultipleValuesBackground", new FSlateRoundedBoxBrush(FStyleColors::Input, InputFocusRadius));
+		Style->Set("ColorPicker.AlphaBackground", new IMAGE_BRUSH("Starship/Common/Checker", Icon16x16, FLinearColor::White, ESlateBrushTileType::Both));
 		Style->Set("ColorPicker.EyeDropper", new IMAGE_BRUSH("Icons/eyedropper_16px", Icon16x16));
 		Style->Set("ColorPicker.Font", FStyleFonts::Get().Normal);
 		Style->Set("ColorPicker.Mode", new IMAGE_BRUSH("Common/ColorPicker_Mode_16x", Icon16x16));
diff --git a/Engine/Source/Runtime/SlateCore/Public/Rendering/DrawElementPayloads.h b/Engine/Source/Runtime/SlateCore/Public/Rendering/DrawElementPayloads.h
--- a/Engine/Source/Runtime/SlateCore/Public/Rendering/DrawElementPayloads.h
+++ b/Engine/Source/Runtime/SlateCore/Public/Rendering/DrawElementPayloads.h
@@ -55,7 +55,7 @@ struct FSlateBoxPayload : public FSlateDataPayload, public FSlateTintableElement
 	ESlateBrushDrawType::Type GetBrushDrawType() const { return DrawType; }
 	const FSlateShaderResourceProxy* GetResourceProxy() const { return ResourceProxy; }
 
-	void SetBrush(const FSlateBrush* InBrush)
+	void SetBrush(const FSlateBrush* InBrush, FVector2D LocalSize, float DrawScale)
 	{
 		check(InBrush);
 		ensureMsgf(InBrush->GetDrawType() != ESlateBrushDrawType::NoDrawType, TEXT("This should have been filtered out earlier in the Make... call."));
@@ -66,7 +66,7 @@ struct FSlateBoxPayload : public FSlateDataPayload, public FSlateTintableElement
 		Tiling = InBrush->GetTiling();
 		Mirroring = InBrush->GetMirroring();
 		DrawType = InBrush->GetDrawType();
-		FSlateResourceHandle Handle = InBrush->GetRenderingResource();
+		FSlateResourceHandle Handle = InBrush->GetRenderingResource(LocalSize, DrawScale);
 		if (Handle.IsValid())
 		{
 			ResourceProxy = Handle.GetResourceProxy();
@@ -82,6 +82,20 @@ struct FSlateBoxPayload : public FSlateDataPayload, public FSlateTintableElement
 	}
 };
 
+struct FSlateRoundedBoxPayload : public FSlateBoxPayload
+{
+	FLinearColor OutlineColor;
+	float Radius;
+	float OutlineWeight;
+
+	FORCEINLINE void SetRadius(float InRadius) { Radius = InRadius; }
+	FORCEINLINE float GetRadius() const { return Radius; }
+
+	FORCEINLINE void SetOutline(const FLinearColor& InOutlineColor, float InOutlineWeight) { OutlineColor = InOutlineColor; OutlineWeight = InOutlineWeight; }
+	FORCEINLINE FLinearColor GetOutlineColor() const { return OutlineColor; }
+	FORCEINLINE float GetOutlineWeight() const { return OutlineWeight; }
+};
+
 struct FSlateTextPayload : public FSlateDataPayload, public FSlateTintableElement
 {
 	// The font to use when rendering
@@ -148,11 +162,13 @@ struct FSlateGradientPayload : public FSlateDataPayload
 {
 	TArray<FSlateGradientStop> GradientStops;
 	EOrientation GradientType;
+	float CornerRadius;
 
-	void SetGradient(const TArray<FSlateGradientStop>& InGradientStops, EOrientation InGradientType)
+	void SetGradient(const TArray<FSlateGradientStop>& InGradientStops, EOrientation InGradientType, float InCornerRadius)
 	{
 		GradientStops = InGradientStops;
 		GradientType = InGradientType;
+		CornerRadius = InCornerRadius;
 	}
 };
 
diff --git a/Engine/Source/Runtime/SlateCore/Public/Rendering/DrawElements.h b/Engine/Source/Runtime/SlateCore/Public/Rendering/DrawElements.h
--- a/Engine/Source/Runtime/SlateCore/Public/Rendering/DrawElements.h
+++ b/Engine/Source/Runtime/SlateCore/Public/Rendering/DrawElements.h
@@ -168,8 +168,9 @@ public:
 	 * @param InGradientStops          List of gradient stops which define the element
 	 * @param InGradientType           The type of gradient (I.E Horizontal, vertical)
 	 * @param InDrawEffects            Optional draw effects to apply
+	 * @param CornerRadius			   Rounds the corners of the box created by the gradient by the specified radius
 	 */
-	SLATECORE_API static void MakeGradient( FSlateWindowElementList& ElementList, uint32 InLayer, const FPaintGeometry& PaintGeometry, TArray<FSlateGradientStop> InGradientStops, EOrientation InGradientType, ESlateDrawEffect InDrawEffects = ESlateDrawEffect::None );
+	SLATECORE_API static void MakeGradient( FSlateWindowElementList& ElementList, uint32 InLayer, const FPaintGeometry& PaintGeometry, TArray<FSlateGradientStop> InGradientStops, EOrientation InGradientType, ESlateDrawEffect InDrawEffects = ESlateDrawEffect::None, float CornerRadius = 0.0f );
 
 	/**
 	 * Creates a Hermite Spline element
