Remove unused solid fill brush from color blocks

diff --git a/Engine/Source/Editor/DetailCustomizations/Private/ColorStructCustomization.cpp b/Engine/Source/Editor/DetailCustomizations/Private/ColorStructCustomization.cpp
--- a/Engine/Source/Editor/DetailCustomizations/Private/ColorStructCustomization.cpp
+++ b/Engine/Source/Editor/DetailCustomizations/Private/ColorStructCustomization.cpp
@@ -101,7 +101,6 @@ TSharedRef<SWidget> FColorStructCustomization::CreateColorWidget(TWeakPtr<IPrope
 				[
 					SAssignNew(ColorPickerParentWidget, SColorBlock)
 					.AlphaBackgroundBrush(FAppStyle::Get().GetBrush("ColorPicker.RoundedAlphaBackground"))
-					.SolidBackgroundBrush(FAppStyle::Get().GetBrush("ColorPicker.RoundedSolidBackground"))
 					.Color(this, &FColorStructCustomization::OnGetColorForColorBlock)
 					.ShowBackgroundForAlpha(true)
 					.AlphaDisplayMode(bIgnoreAlpha ? EColorBlockAlphaDisplayMode::Ignore : EColorBlockAlphaDisplayMode::Separate)
diff --git a/Engine/Source/Runtime/Slate/Private/Widgets/Colors/SColorBlock.cpp b/Engine/Source/Runtime/Slate/Private/Widgets/Colors/SColorBlock.cpp
--- a/Engine/Source/Runtime/Slate/Private/Widgets/Colors/SColorBlock.cpp
+++ b/Engine/Source/Runtime/Slate/Private/Widgets/Colors/SColorBlock.cpp
@@ -14,7 +14,6 @@ void SColorBlock::Construct(const FArguments& InArgs)
 {
 	Color = InArgs._Color;
 	AlphaBackgroundBrush = InArgs._AlphaBackgroundBrush;
-	SolidBackgroundBrush = InArgs._SolidBackgroundBrush;
 	GradientCornerRadius = InArgs._CornerRadius;
 	ColorIsHSV = InArgs._ColorIsHSV;
 	AlphaDisplayMode = InArgs._AlphaDisplayMode;
diff --git a/Engine/Source/Runtime/Slate/Public/Widgets/Colors/SColorBlock.h b/Engine/Source/Runtime/Slate/Public/Widgets/Colors/SColorBlock.h
--- a/Engine/Source/Runtime/Slate/Public/Widgets/Colors/SColorBlock.h
+++ b/Engine/Source/Runtime/Slate/Public/Widgets/Colors/SColorBlock.h
@@ -29,7 +29,6 @@ public:
 	SLATE_BEGIN_ARGS(SColorBlock)
 		: _Color(FLinearColor::White)
 		, _AlphaBackgroundBrush(FAppStyle::Get().GetBrush("ColorPicker.AlphaBackground"))
-		, _SolidBackgroundBrush(FAppStyle::Get().GetBrush("GenericWhiteBox"))
 		, _CornerRadius(0.0f)
 		, _ColorIsHSV(false)
 		, _ShowBackgroundForAlpha(false)
@@ -44,9 +43,6 @@ public:
 
 		/** Background to display for when there is a color with transparency. Irrelevant if ignoring alpha */
 		SLATE_ATTRIBUTE(const FSlateBrush*, AlphaBackgroundBrush)
-	
-		/** Background to use to display when there is a solid color. This background is tinted by the solid color */
-		SLATE_ATTRIBUTE(const FSlateBrush*, SolidBackgroundBrush)
 
 		/** Rounding to apply to the corners of the block */
 		SLATE_ATTRIBUTE(float, CornerRadius)
