Minor fix to CalculatePopupWindowPosition to pass a valid rect to GetWorkArea.

diff --git a/Engine/Source/Runtime/Slate/Private/Framework/Application/SlateApplication.cpp b/Engine/Source/Runtime/Slate/Private/Framework/Application/SlateApplication.cpp
--- a/Engine/Source/Runtime/Slate/Private/Framework/Application/SlateApplication.cpp
+++ b/Engine/Source/Runtime/Slate/Private/Framework/Application/SlateApplication.cpp
@@ -3606,11 +3606,11 @@ FVector2D FSlateApplication::CalculateTooltipWindowPosition( const FSlateRect& I
 	return ToolTipLocation;
 }
 
-FVector2D FSlateApplication::CalculatePopupWindowPosition( const FSlateRect& InAnchor, const FVector2D& InSize, bool bAutoAdjustForDPIScale, const FVector2D& InProposedPlacement, const EOrientation Orientation) const
+FVector2D FSlateApplication::CalculatePopupWindowPosition(const FSlateRect& InAnchor, const FVector2D& InSize, bool bAutoAdjustForDPIScale, const FVector2D& InProposedPlacement, const EOrientation Orientation) const
 {
-	FVector2D CalculatedPopUpWindowPosition( 0, 0 );
+	FVector2D CalculatedPopUpWindowPosition(0, 0);
 
-	float DPIScale = 1.0f; 
+	float DPIScale = 1.0f;
 
 	if (bAutoAdjustForDPIScale)
 	{
@@ -3627,14 +3627,14 @@ FVector2D FSlateApplication::CalculatePopupWindowPosition( const FSlateRect& InA
 
 	EPopUpOrientation::Type PopUpOrientation = EPopUpOrientation::Horizontal;
 
-	if ( Orientation == EOrientation::Orient_Vertical )
+	if (Orientation == EOrientation::Orient_Vertical)
 	{
-		PopUpOrientation =  EPopUpOrientation::Vertical;
+		PopUpOrientation = EPopUpOrientation::Vertical;
 	}
 
-	if ( PlatformApplication->TryCalculatePopupWindowPosition( AnchorRect, AdjustedSize, InProposedPlacement, PopUpOrientation, /*OUT*/&CalculatedPopUpWindowPosition ) )
+	if (PlatformApplication->TryCalculatePopupWindowPosition(AnchorRect, AdjustedSize, InProposedPlacement, PopUpOrientation, /*OUT*/&CalculatedPopUpWindowPosition))
 	{
-		return CalculatedPopUpWindowPosition/DPIScale;
+		return CalculatedPopUpWindowPosition / DPIScale;
 	}
 	else
 	{
@@ -3642,27 +3642,27 @@ FVector2D FSlateApplication::CalculatePopupWindowPosition( const FSlateRect& InA
 		// Use our own rect.  This window as probably doesn't have a size or position yet.
 		// Use a size of 1 to get the closest monitor to the start point
 		FPlatformRect WorkAreaFinderRect(AnchorRect);
-		WorkAreaFinderRect.Left = AnchorRect.Left + 1;
-		WorkAreaFinderRect.Top = AnchorRect.Top + 1;
+		WorkAreaFinderRect.Right = AnchorRect.Left + 1;
+		WorkAreaFinderRect.Bottom = AnchorRect.Top + 1;
 		const FPlatformRect PlatformWorkArea = PlatformApplication->GetWorkArea(WorkAreaFinderRect);
 
-		const FSlateRect WorkAreaRect( 
-			PlatformWorkArea.Left, 
-			PlatformWorkArea.Top, 
-			PlatformWorkArea.Left+(PlatformWorkArea.Right - PlatformWorkArea.Left), 
-			PlatformWorkArea.Top+(PlatformWorkArea.Bottom - PlatformWorkArea.Top) );
+		const FSlateRect WorkAreaRect(
+			PlatformWorkArea.Left,
+			PlatformWorkArea.Top,
+			PlatformWorkArea.Left + (PlatformWorkArea.Right - PlatformWorkArea.Left),
+			PlatformWorkArea.Top + (PlatformWorkArea.Bottom - PlatformWorkArea.Top));
 
 		FVector2D ProposedPlacement = InProposedPlacement;
 
 		if (ProposedPlacement.IsZero())
 		{
-		// Assume natural left-to-right, top-to-bottom flow; position popup below and to the right.
+			// Assume natural left-to-right, top-to-bottom flow; position popup below and to the right.
 			ProposedPlacement = FVector2D(
-			Orientation == Orient_Horizontal ? AnchorRect.Right : AnchorRect.Left,
-			Orientation == Orient_Horizontal ? AnchorRect.Top : AnchorRect.Bottom);
+				Orientation == Orient_Horizontal ? AnchorRect.Right : AnchorRect.Left,
+				Orientation == Orient_Horizontal ? AnchorRect.Top : AnchorRect.Bottom);
 		}
 
-		return ComputePopupFitInRect(InAnchor, FSlateRect(ProposedPlacement, ProposedPlacement+AdjustedSize), Orientation, WorkAreaRect) / DPIScale;
+		return ComputePopupFitInRect(InAnchor, FSlateRect(ProposedPlacement, ProposedPlacement + AdjustedSize), Orientation, WorkAreaRect) / DPIScale;
 	}
 }
 
