Added MinSize() to SSplitter::FSlot to allow slots to have different minimum sizes than the globally-active SSplitter::MinSlotHeight.

diff --git a/Engine/Source/Runtime/Slate/Private/Widgets/Layout/SSplitter.cpp b/Engine/Source/Runtime/Slate/Private/Widgets/Layout/SSplitter.cpp
--- a/Engine/Source/Runtime/Slate/Private/Widgets/Layout/SSplitter.cpp
+++ b/Engine/Source/Runtime/Slate/Private/Widgets/Layout/SSplitter.cpp
@@ -578,7 +578,7 @@ void SSplitter::HandleResizingDelta(EOrientation SplitterOrientation, const floa
 
 			SlotInfo.Slot = &InChildren[ SlotsAfterDragHandleIndices[ SlotIndex ] ];
 			SlotInfo.Geometry = &ChildGeometries[ SlotsAfterDragHandleIndices[ SlotIndex ] ];
-			SlotInfo.NewSize = ClampChild(SlotInfo.Geometry->GetSizeInParentSpace().Component( AxisIndex ));
+			SlotInfo.NewSize = ClampChild(*SlotInfo.Slot, SlotInfo.Geometry->GetSizeInParentSpace().Component( AxisIndex ));
 
 			SlotsAfterDragHandle.Add( SlotInfo );
 		}
@@ -590,7 +590,7 @@ void SSplitter::HandleResizingDelta(EOrientation SplitterOrientation, const floa
 
 		// Compute the new sizes of the children
 		const float PrevChildLength = PrevChildGeom.GetSizeInParentSpace().Component(AxisIndex);
-		float NewPrevChildLength = ClampChild(PrevChildLength + Delta);
+		float NewPrevChildLength = ClampChild(PrevChild, PrevChildLength + Delta);
 		Delta = NewPrevChildLength - PrevChildLength;
 
 		// Distribute the Delta across the affected slots after the drag handle
@@ -614,7 +614,7 @@ void SSplitter::HandleResizingDelta(EOrientation SplitterOrientation, const floa
 			{
 				FSlotInfo& SlotInfo = SlotsAfterDragHandle[SlotIndex];
 				float CurrentSize = SlotInfo.NewSize;
-				SlotInfo.NewSize = ClampChild(CurrentSize - DividedDelta);
+				SlotInfo.NewSize = ClampChild(*SlotInfo.Slot, CurrentSize - DividedDelta);
 
 				// If one of the slots couldn't be fully adjusted by the delta due to min/max constraints then
 				// the leftover delta needs to be evenly distributed to all of the other slots
@@ -625,7 +625,7 @@ void SSplitter::HandleResizingDelta(EOrientation SplitterOrientation, const floa
 		Delta = Delta - UnusedDelta;
 
 		// PrevChildLength needs to be updated: it's value has to take into account the next child's min/max restrictions
-		NewPrevChildLength = ClampChild( PrevChildLength + Delta );
+		NewPrevChildLength = ClampChild(PrevChild, PrevChildLength + Delta);
 
 		// Cells being resized are both stretch values -> redistribute the stretch coefficients proportionately
 		// to match the new child sizes on the screen.
@@ -696,9 +696,16 @@ void SSplitter::HandleResizingByMousePosition(EOrientation SplitterOrientation,
 *
 * @return A size that is clamped against the minimum size allowed for children.
 */
-float SSplitter::ClampChild( float ProposedSize )
+float SSplitter::ClampChild(const FSlot& ChildSlot, float ProposedSize)
 {
-	return FMath::Max( MinSplitterChildLength, ProposedSize );
+	float ClampedSize = FMath::Max(MinSplitterChildLength, ProposedSize);
+
+	if (ChildSlot.MinSizeValue.IsSet())
+	{
+		ClampedSize = FMath::Max(ChildSlot.MinSizeValue.GetValue(), ProposedSize);
+	}
+
+	return ClampedSize;
 }
 
 
diff --git a/Engine/Source/Runtime/Slate/Public/Widgets/Layout/SSplitter.h b/Engine/Source/Runtime/Slate/Public/Widgets/Layout/SSplitter.h
--- a/Engine/Source/Runtime/Slate/Public/Widgets/Layout/SSplitter.h
+++ b/Engine/Source/Runtime/Slate/Public/Widgets/Layout/SSplitter.h
@@ -85,8 +85,15 @@ public:
 			return *this;
 		}
 
+		FSlot& MinSize(float InMinSize)
+		{
+			MinSizeValue = InMinSize;
+			return *this;
+		}
+
 		TAttribute<ESizeRule> SizingRule;
 		TAttribute<float> SizeValue;
+		TOptional<float> MinSizeValue;
 		FOnSlotResized OnSlotResized_Handler;
 	};
 
@@ -277,7 +284,7 @@ protected:
 	 *
 	 * @return A size that is clamped against the minimum size allowed for children.
 	 */
-	float ClampChild( float ProposedSize );
+	float ClampChild(const FSlot& ChildSlot, float ProposedSize);
 
 	/**
 	 * Given a mouse position within the splitter, figure out which resize handle we are hovering (if any).
