Slate: Splitter SizeToContent can now be resized if the user has binded to the callback.

diff --git a/Engine/Source/Runtime/Slate/Private/Widgets/Layout/SSplitter.cpp b/Engine/Source/Runtime/Slate/Private/Widgets/Layout/SSplitter.cpp
--- a/Engine/Source/Runtime/Slate/Private/Widgets/Layout/SSplitter.cpp
+++ b/Engine/Source/Runtime/Slate/Private/Widgets/Layout/SSplitter.cpp
@@ -4,12 +4,24 @@
 #include "Layout/ArrangedChildren.h"
 #include "Rendering/DrawElements.h"
 
+bool SSplitter::FSlot::CanBeResized() const
+{
+	const ESizeRule CurrentSizingRule = SizingRule.Get();
+	if (CurrentSizingRule == SSplitter::ESizeRule::FractionOfParent)
+	{
+		return !bIsResizable.IsSet() || bIsResizable.GetValue();
+	}
+	else
+	{
+		return bIsResizable.IsSet() && bIsResizable.GetValue() && OnSlotResized_Handler.IsBound();
+	}
+}
+
 SSplitter::FSlot& SSplitter::Slot()
 {
 	return (*new FSlot());
 }
 
-
 SSplitter::FSlot& SSplitter::AddSlot( int32 AtIndex )
 {
 	FSlot& NewSlot = *new FSlot();
@@ -193,7 +205,7 @@ int32 SSplitter::OnPaint( const FPaintArgs& Args, const FGeometry& AllottedGeome
 		const float HalfHitDetectionSplitterHandleSize = FMath::RoundToFloat(HitDetectionSplitterHandleSize / 2.0f);
 		const float HalfPhysicalSplitterHandleSize = PhysicalSplitterHandleSize;// FMath::RoundToInt(PhysicalSplitterHandleSize / 2.0f);
 
-		FVector2D HandleSize;		
+		FVector2D HandleSize;
 		FVector2D HandlePosition;
 		if (Orientation == Orient_Horizontal)
 		{
@@ -449,7 +461,7 @@ FCursorReply SSplitter::OnCursorQuery( const FGeometry& MyGeometry, const FPoint
 		: GetHandleBeingResizedFromMousePosition<Orient_Vertical>( PhysicalSplitterHandleSize, HitDetectionSplitterHandleSize, LocalMousePosition, LayoutChildren );
 
 	if (CurrentHoveredHandleIndex != INDEX_NONE &&
-		Children[CurrentHoveredHandleIndex].SizingRule.Get() != ESizeRule::SizeToContent)
+		Children[CurrentHoveredHandleIndex].CanBeResized())
 	{
 		return Orientation == Orient_Horizontal
 			? FCursorReply::Cursor( EMouseCursor::ResizeLeftRight )
@@ -495,7 +507,8 @@ int32 SSplitter::FindResizeableSlotBeforeHandle( int32 DraggedHandle, const TPan
 	// Resizing collapsed or autosizing slots does not make sense (their size is predetermined).
 	// Search out from the DraggedHandle to find the first non-collapsed, non-autosizing slot we can resize.
 	int32 SlotBeforeDragHandle = DraggedHandle;
-	while ( SlotBeforeDragHandle >= 0 && (Children[SlotBeforeDragHandle].GetWidget()->GetVisibility() == EVisibility::Collapsed || Children[SlotBeforeDragHandle].SizingRule.Get() == SSplitter::SizeToContent) )
+	while (SlotBeforeDragHandle >= 0
+		&& (Children[SlotBeforeDragHandle].GetWidget()->GetVisibility() == EVisibility::Collapsed || !Children[SlotBeforeDragHandle].CanBeResized()))
 	{
 		--SlotBeforeDragHandle;
 	}
@@ -511,24 +524,23 @@ int32 SSplitter::FindResizeableSlotAfterHandle( int32 DraggedHandle, const TPane
 	// But the slots list does contain collapsed children! Make sure that we are not resizing a collapsed slot.
 	// We also cannot resize auto-sizing slots.
 	int32 SlotAfterDragHandle = DraggedHandle+1;
+	while (SlotAfterDragHandle < NumChildren
+		&& (Children[SlotAfterDragHandle].GetWidget()->GetVisibility() == EVisibility::Collapsed || !Children[SlotAfterDragHandle].CanBeResized()))
 	{
-		while( SlotAfterDragHandle < NumChildren && (Children[SlotAfterDragHandle].GetWidget()->GetVisibility() == EVisibility::Collapsed || Children[SlotAfterDragHandle].SizingRule.Get() == SSplitter::SizeToContent) )
-		{
-			++SlotAfterDragHandle;
-		}
+		++SlotAfterDragHandle;
 	}
 
 	return SlotAfterDragHandle;
 }
 
 
-void SSplitter::FindAllResizeableSlotsAfterHandle( int32 DraggedHandle, const TPanelChildren<FSlot>& Children, TArray< int32 >& OutSlotIndicies )
+void SSplitter::FindAllResizeableSlotsAfterHandle( int32 DraggedHandle, const TPanelChildren<FSlot>& Children, TArray< int32, TMemStackAllocator<> >& OutSlotIndicies )
 {
 	const int32 NumChildren = Children.Num();
 
 	for (int SlotIndex = DraggedHandle+1; SlotIndex < NumChildren; SlotIndex++)
 	{
-		if (Children[ SlotIndex ].GetWidget()->GetVisibility() == EVisibility::Collapsed || Children[ SlotIndex ].SizingRule.Get() == SSplitter::SizeToContent)
+		if (Children[ SlotIndex ].GetWidget()->GetVisibility() == EVisibility::Collapsed || !Children[ SlotIndex ].CanBeResized())
 		{
 			continue;
 		}
@@ -546,14 +558,16 @@ void SSplitter::HandleResizingDelta(EOrientation SplitterOrientation, const floa
 	//  - Prev vs. Next refers to the widgets in the order they are laid out (left->right, top->bottom).
 	//  - New vs. Old refers to the Old values for width/height vs. the post-resize values.
 
-	TArray< int32 > SlotsAfterDragHandleIndices;
+	FMemMark Mark(FMemStack::Get());
+	TArray<int32, TMemStackAllocator<>> SlotsAfterDragHandleIndices;
+	SlotsAfterDragHandleIndices.Reserve(NumChildren);
 	if (InResizeMode == ESplitterResizeMode::FixedPosition)
 	{
-		const int32 SlotAfterDragHandle = FindResizeableSlotAfterHandle( DraggedHandle, InChildren);
+		int32 SlotAfterDragHandle = FindResizeableSlotAfterHandle(DraggedHandle, Children);
 
-		if ( SlotAfterDragHandle < NumChildren )
+		if (InChildren.IsValidIndex(SlotAfterDragHandle))
 		{
-			SlotsAfterDragHandleIndices.Add( SlotAfterDragHandle );
+			SlotsAfterDragHandleIndices.Add(SlotAfterDragHandle);
 		}
 	}
 	else if (InResizeMode == ESplitterResizeMode::Fill || InResizeMode == ESplitterResizeMode::FixedSize)
@@ -571,8 +585,9 @@ void SSplitter::HandleResizingDelta(EOrientation SplitterOrientation, const floa
 			float NewSize;
 		};
 
-		TArray< FSlotInfo > SlotsAfterDragHandle;
-		for (int SlotIndex = 0; SlotIndex < NumSlotsAfterDragHandle; SlotIndex++)
+		TArray<FSlotInfo, TMemStackAllocator<>> SlotsAfterDragHandle;
+		SlotsAfterDragHandle.Reserve(NumSlotsAfterDragHandle);
+		for (int32 SlotIndex = 0; SlotIndex < NumSlotsAfterDragHandle; SlotIndex++)
 		{
 			FSlotInfo SlotInfo;
 
@@ -596,14 +611,14 @@ void SSplitter::HandleResizingDelta(EOrientation SplitterOrientation, const floa
 		// Distribute the Delta across the affected slots after the drag handle
 		float UnusedDelta = Delta;
 
-		for (int DistributionCount = 0; DistributionCount < NumSlotsAfterDragHandle && UnusedDelta != 0; DistributionCount++)
+		for (int32 DistributionCount = 0; DistributionCount < NumSlotsAfterDragHandle && UnusedDelta != 0; DistributionCount++)
 		{
 			float DividedDelta = InResizeMode != ESplitterResizeMode::FixedSize ? UnusedDelta / NumSlotsAfterDragHandle : UnusedDelta;
 
 			// Reset the unused delta to keep track of any leftover space to distribute later through stretching.
 			UnusedDelta = 0;
 
-			int SlotIndex = 0;
+			int32 SlotIndex = 0;
 			// Resize only the last handle in the case of fixed size
 			if (InResizeMode == ESplitterResizeMode::FixedSize)
 			{
@@ -631,41 +646,42 @@ void SSplitter::HandleResizingDelta(EOrientation SplitterOrientation, const floa
 		// to match the new child sizes on the screen.
 		{
 			float TotalLength = NewPrevChildLength;
-			float TotalStretchCoefficients = PrevChild.SizeValue.Get();
+			float TotalStretchCoefficients = PrevChild.SizingRule.Get() == ESizeRule::FractionOfParent ? PrevChild.SizeValue.Get() : 0.f;
 
-			for (int SlotIndex = 0; SlotIndex < NumSlotsAfterDragHandle; SlotIndex++)
+			for (int32 SlotIndex = 0; SlotIndex < NumSlotsAfterDragHandle; SlotIndex++)
 			{
-				FSlotInfo SlotInfo = SlotsAfterDragHandle[ SlotIndex ];
+				const FSlotInfo& SlotInfo = SlotsAfterDragHandle[ SlotIndex ];
 
 				TotalLength += SlotInfo.NewSize;
-				TotalStretchCoefficients += SlotInfo.Slot->SizeValue.Get();
+				TotalStretchCoefficients += SlotInfo.Slot->SizingRule.Get() == ESizeRule::FractionOfParent ? SlotInfo.Slot->SizeValue.Get() : 0.f;
 			}
 
-			const float NewPrevChildSize = ( TotalStretchCoefficients * NewPrevChildLength / TotalLength );
-
-			if (PrevChild.OnSlotResized_Handler.IsBound())
-			{
-				PrevChild.OnSlotResized_Handler.Execute( NewPrevChildSize );
-			}
-			else
+			auto SetSize = [](FSlot& Slot, float TotalStretchCoefficients, float NewLength, float TotalLength)
 			{
-				PrevChild.SizeValue = NewPrevChildSize;
-			}
-
-			for (int SlotIndex = 0; SlotIndex < NumSlotsAfterDragHandle; SlotIndex++)
-			{
-				FSlotInfo SlotInfo = SlotsAfterDragHandle[ SlotIndex ];
-
-				const float NewNextChildSize = ( TotalStretchCoefficients * SlotInfo.NewSize / TotalLength );
-
-				if (SlotInfo.Slot->OnSlotResized_Handler.IsBound())
+				if (Slot.SizingRule.Get() == ESizeRule::FractionOfParent)
 				{
-					SlotInfo.Slot->OnSlotResized_Handler.Execute(NewNextChildSize);
+					const float NewFillSize = (TotalStretchCoefficients * NewLength / TotalLength);
+					if (Slot.OnSlotResized_Handler.IsBound())
+					{
+						Slot.OnSlotResized_Handler.Execute(NewFillSize);
+					}
+					else
+					{
+						Slot.SizeValue = NewFillSize;
+					}
 				}
 				else
 				{
-					SlotInfo.Slot->SizeValue = NewNextChildSize;
+					ensure(Slot.OnSlotResized_Handler.IsBound());
+					Slot.OnSlotResized_Handler.ExecuteIfBound(NewLength);
 				}
+			};
+
+			SetSize(PrevChild, TotalStretchCoefficients, NewPrevChildLength, TotalLength);
+			for (int32 SlotIndex = 0; SlotIndex < NumSlotsAfterDragHandle; SlotIndex++)
+			{
+				const FSlotInfo& SlotInfo = SlotsAfterDragHandle[ SlotIndex ];
+				SetSize(*(SlotInfo.Slot), TotalStretchCoefficients, SlotInfo.NewSize, TotalLength);
 			}
 		}
 	}
diff --git a/Engine/Source/Runtime/Slate/Private/Widgets/Views/SHeaderRow.cpp b/Engine/Source/Runtime/Slate/Private/Widgets/Views/SHeaderRow.cpp
--- a/Engine/Source/Runtime/Slate/Private/Widgets/Views/SHeaderRow.cpp
+++ b/Engine/Source/Runtime/Slate/Private/Widgets/Views/SHeaderRow.cpp
@@ -765,66 +765,16 @@ void SHeaderRow::RegenerateWidgets()
 						[
 							SNew(SBox)
 							.WidthOverride(SomeColumn.GetWidth())
-						[
-							NewHeader
-						]
+							[
+								NewHeader
+							]
 						];
 				}
 				break;
 
+
 				case EColumnSizeMode::Manual:
 				{
-					// Sizing grip to put at the end of the column - we can't use a SSplitter here as it doesn't have the resizing behavior we need
-					const float GripSize = SplitterHandleSize > 0.0f ? SplitterHandleSize + 4.0f : 5.0f;
-					TSharedRef<SBorder> SizingGrip = SNew(SBorder)
-						.Padding(0.0f)
-						.BorderImage( &Style->ColumnSplitterStyle )
-						.Cursor(EMouseCursor::ResizeLeftRight)
-						.Content()
-						[
-							SNew(SSpacer)
-							.Size(FVector2D(GripSize, GripSize))
-						];
-
-					TWeakPtr<SBorder> WeakSizingGrip = SizingGrip;
-					auto SizingGrip_OnMouseButtonDown = [&SomeColumn, WeakSizingGrip](const FGeometry&, const FPointerEvent&) -> FReply
-					{
-						TSharedPtr<SBorder> SizingGripPtr = WeakSizingGrip.Pin();
-						if (SizingGripPtr.IsValid())
-						{
-							return FReply::Handled().CaptureMouse(SizingGripPtr.ToSharedRef());
-						}
-						return FReply::Unhandled();
-					};
-
-					auto SizingGrip_OnMouseButtonUp = [&SomeColumn, WeakSizingGrip](const FGeometry&, const FPointerEvent&) -> FReply
-					{
-						TSharedPtr<SBorder> SizingGripPtr = WeakSizingGrip.Pin();
-						if (SizingGripPtr.IsValid() && SizingGripPtr->HasMouseCapture())
-						{
-							return FReply::Handled().ReleaseMouseCapture();
-						}
-						return FReply::Unhandled();
-					};
-
-					auto SizingGrip_OnMouseMove = [&SomeColumn, WeakSizingGrip](const FGeometry&, const FPointerEvent& InPointerEvent) -> FReply
-					{
-						TSharedPtr<SBorder> SizingGripPtr = WeakSizingGrip.Pin();
-						if (SizingGripPtr.IsValid() && SizingGripPtr->HasMouseCapture())
-						{
-							// The sizing grip has been moved, so update our columns size from the movement delta
-							const float NewWidth = SomeColumn.GetWidth() + InPointerEvent.GetCursorDelta().X;
-							SomeColumn.SetWidth(FMath::Max(20.0f, NewWidth));
-							return FReply::Handled();
-						}
-						return FReply::Unhandled();
-					};
-
-					// Bind the events to handle the drag sizing
-					SizingGrip->SetOnMouseButtonDown(FPointerEventHandler::CreateLambda(SizingGrip_OnMouseButtonDown));
-					SizingGrip->SetOnMouseButtonUp(FPointerEventHandler::CreateLambda(SizingGrip_OnMouseButtonUp));
-					SizingGrip->SetOnMouseMove(FPointerEventHandler::CreateLambda(SizingGrip_OnMouseMove));
-
 					auto GetColumnWidthAsOptionalSize = [&SomeColumn]() -> FOptionalSize
 					{
 						const float DesiredWidth = SomeColumn.GetWidth();
@@ -834,29 +784,22 @@ void SHeaderRow::RegenerateWidgets()
 					TAttribute<FOptionalSize> WidthBinding;
 					WidthBinding.Bind(TAttribute<FOptionalSize>::FGetter::CreateLambda(GetColumnWidthAsOptionalSize));
 
-					// Add resizable cell
 					Splitter->AddSlot()
 						.SizeRule(SSplitter::SizeToContent)
+						.Resizable(true)
+						.OnSlotResized(SSplitter::FOnSlotResized::CreateRaw(&SomeColumn, &FColumn::SetWidth))
 						[
 							SNew(SBox)
 							.WidthOverride(WidthBinding)
 							[
-								SNew(SOverlay)
-								+ SOverlay::Slot()
-								[
-									NewHeader
-								]
-								+ SOverlay::Slot()
-								.HAlign(HAlign_Right)
-								[
-									SizingGrip
-								]
+								NewHeader
 							]
 						];
 				}
 				break;
 
 				default:
+					ensure(false);
 					break;
 				}
 			}			
diff --git a/Engine/Source/Runtime/Slate/Public/Widgets/Layout/SSplitter.h b/Engine/Source/Runtime/Slate/Public/Widgets/Layout/SSplitter.h
--- a/Engine/Source/Runtime/Slate/Public/Widgets/Layout/SSplitter.h
+++ b/Engine/Source/Runtime/Slate/Public/Widgets/Layout/SSplitter.h
@@ -25,8 +25,11 @@ namespace ESplitterResizeMode
 {
 	enum Type
 	{
+		/** Resize the selected slot. If space is needed, then resize the next resizable slot. */
 		FixedPosition,
+		/** Resize the selected slot. If space is needed, then resize the last resizable slot. */
 		FixedSize,
+		/** Resize the selected slot by redistributing the available space with the following resizable slots. */
 		Fill,
 	};
 }
@@ -67,34 +70,57 @@ public:
 		{
 		}
 
+		/** When the RuleSize is set to FractionOfParent, the size of the slot is the Value percentage of its parent size. */
 		FSlot& Value( const TAttribute<float>& InValue )
 		{
 			SizeValue = InValue;
 			return *this;
 		}
+
+		/**
+		 * Can the slot be resize by the user.
+		 * @see CanBeResized()
+		 */
+		FSlot& Resizable(bool bInIsResizable)
+		{
+			bIsResizable = bInIsResizable;
+			return *this;
+		}
+
+		/** Minimun slot size when resizing. */
+		FSlot& MinSize(float InMinSize)
+		{
+			MinSizeValue = InMinSize;
+			return *this;
+		}
 		
+		/**
+		 * Callback when the slot is resized.
+		 * @see CanBeResized()
+		 */
 		FSlot& OnSlotResized( const FOnSlotResized& InHandler )
 		{
 			OnSlotResized_Handler = InHandler;
 			return *this;
 		}
 
+		/** The size rule used by the slot. */
 		FSlot& SizeRule( const TAttribute<ESizeRule>& InSizeRule ) 
 		{
 			SizingRule = InSizeRule;
 			return *this;
 		}
 
-		FSlot& MinSize(float InMinSize)
-		{
-			MinSizeValue = InMinSize;
-			return *this;
-		}
+	public:
+		/** A slot can be resize if bIsResizable and the SizeRule is a FractionOfParent or the OnSlotResized delegate is set. */
+		bool CanBeResized() const;
 
+	public:
 		TAttribute<ESizeRule> SizingRule;
 		TAttribute<float> SizeValue;
 		TOptional<float> MinSizeValue;
 		FOnSlotResized OnSlotResized_Handler;
+		TOptional<bool> bIsResizable;
 	};
 
 	/** @return a new SSplitter::FSlot() */
@@ -265,7 +291,7 @@ protected:
 	 */
 	static int32 FindResizeableSlotAfterHandle( int32 DraggedHandle, const TPanelChildren<FSlot>& Children );
 
-	static void FindAllResizeableSlotsAfterHandle( int32 DraggedHandle, const TPanelChildren<FSlot>& Children, TArray< int32 >& OutSlotIndicies );
+	static void FindAllResizeableSlotsAfterHandle( int32 DraggedHandle, const TPanelChildren<FSlot>& Children, TArray<int32, TMemStackAllocator<>>& OutSlotIndicies );
 
 	/**
 	 * Resizes the children based on user input. The template parameter Orientation corresponds to the splitter being horizontal or vertical.
