Toolbars can now support a hybrid mode where they include no text unless a specific toolbar entry is called out using the "calloutstyle"  for the toolbar entry. This allows toolbars to reduce their width significantly. The level editor now uses this by default but a few other asset editors have been set up for this

diff --git a/Engine/Source/Developer/ToolMenus/Private/ToolMenuEntry.cpp b/Engine/Source/Developer/ToolMenus/Private/ToolMenuEntry.cpp
--- a/Engine/Source/Developer/ToolMenus/Private/ToolMenuEntry.cpp
+++ b/Engine/Source/Developer/ToolMenus/Private/ToolMenuEntry.cpp
@@ -16,6 +16,7 @@ FToolMenuEntry::FToolMenuEntry() :
 	UserInterfaceActionType(EUserInterfaceActionType::Button),
 	bShouldCloseWindowAfterMenuSelection(true),
 	ScriptObject(nullptr),
+	StyleNameOverride(NAME_None),
 	bAddedDuringRegister(false)
 {
 }
@@ -27,6 +28,7 @@ FToolMenuEntry::FToolMenuEntry(const FToolMenuOwner InOwner, const FName InName,
 	UserInterfaceActionType(EUserInterfaceActionType::Button),
 	bShouldCloseWindowAfterMenuSelection(true),
 	ScriptObject(nullptr),
+	StyleNameOverride(NAME_None),
 	bAddedDuringRegister(false)
 {
 }
@@ -58,6 +60,15 @@ void FToolMenuEntry::SetCommandList(const TSharedPtr<const FUICommandList>& InCo
 	CommandList = InCommandList;
 }
 
+void FToolMenuEntry::AddOptionsDropdown(FUIAction InAction, const FOnGetContent InMenuContentGenerator, const TAttribute<FText>& InToolTip)
+{
+	ToolBarData.OptionsDropdownData = MakeShared<FToolMenuEntryOptionsDropdownData>();
+	
+	ToolBarData.OptionsDropdownData->Action = InAction;
+	ToolBarData.OptionsDropdownData->MenuContentGenerator = FNewToolMenuChoice(InMenuContentGenerator);
+	ToolBarData.OptionsDropdownData->ToolTip = InToolTip;
+}
+
 void FToolMenuEntry::SetCommand(const TSharedPtr<const FUICommandInfo>& InCommand, TOptional<FName> InName, const TAttribute<FText>& InLabel, const TAttribute<FText>& InToolTip, const TAttribute<FSlateIcon>& InIcon)
 {
 	Command = InCommand;
diff --git a/Engine/Source/Developer/ToolMenus/Private/ToolMenus.cpp b/Engine/Source/Developer/ToolMenus/Private/ToolMenus.cpp
--- a/Engine/Source/Developer/ToolMenus/Private/ToolMenus.cpp
+++ b/Engine/Source/Developer/ToolMenus/Private/ToolMenus.cpp
@@ -1297,6 +1297,12 @@ void UToolMenus::PopulateToolBarBuilder(FToolBarBuilder& ToolBarBuilder, UToolMe
 				{
 					ToolBarBuilder.AddToolBarButton(UIAction, Block.Name, Block.Label, Block.ToolTip, Block.Icon, Block.UserInterfaceActionType, Block.TutorialHighlightName);
 				}
+
+				if (Block.ToolBarData.OptionsDropdownData.IsValid())
+				{
+					FOnGetContent OnGetContent = ConvertWidgetChoice(Block.ToolBarData.OptionsDropdownData->MenuContentGenerator, MenuData->Context);
+					ToolBarBuilder.AddComboButton(Block.ToolBarData.OptionsDropdownData->Action, OnGetContent, Block.Label, Block.ToolBarData.OptionsDropdownData->ToolTip, Block.Icon, true, Block.TutorialHighlightName);
+				}
 			}
 			else if (Block.Type == EMultiBlockType::ToolBarComboButton)
 			{
diff --git a/Engine/Source/Developer/ToolMenus/Public/ToolMenuEntry.h b/Engine/Source/Developer/ToolMenus/Public/ToolMenuEntry.h
--- a/Engine/Source/Developer/ToolMenus/Public/ToolMenuEntry.h
+++ b/Engine/Source/Developer/ToolMenus/Public/ToolMenuEntry.h
@@ -34,6 +34,13 @@ public:
 	FNewToolMenuChoice ConstructMenu;
 };
 
+struct FToolMenuEntryOptionsDropdownData
+{
+	FNewToolMenuChoice MenuContentGenerator;
+	TAttribute<FText> ToolTip;
+	FUIAction Action;
+};
+
 struct FToolMenuEntryToolBarData
 {
 public:
@@ -44,9 +51,14 @@ public:
 	{
 	}
 
-	TOptional< EVisibility > LabelVisibility;
+	/** Delegate that generates a widget for this combo button's menu content.  Called when the menu is summoned. */
+	FNewToolMenuChoice ComboButtonContextMenuGenerator;
+
+	/** Legacy delegate that generates a widget for this combo button's menu content.  Called when the menu is summoned. */
+	FNewToolBarDelegateLegacy ConstructLegacy;
+
+	TSharedPtr<FToolMenuEntryOptionsDropdownData> OptionsDropdownData;
 
-	/** If true, the icon and label won't be displayed */
 	bool bSimpleComboBox;
 
 	/** Whether ToolBar will have Focusable buttons */
@@ -55,11 +67,6 @@ public:
 	/** Whether this toolbar should always use small icons, regardless of the current settings */
 	bool bForceSmallIcons;
 
-	/** Delegate that generates a widget for this combo button's menu content.  Called when the menu is summoned. */
-	FNewToolMenuChoice ComboButtonContextMenuGenerator;
-
-	/** Legacy delegate that generates a widget for this combo button's menu content.  Called when the menu is summoned. */
-	FNewToolBarDelegateLegacy ConstructLegacy;
 };
 
 
@@ -113,6 +120,8 @@ struct TOOLMENUS_API FToolMenuEntry
 
 	void SetCommandList(const TSharedPtr<const FUICommandList>& InCommandList);
 
+	void AddOptionsDropdown(FUIAction InAction, const FOnGetContent InMenuContentGenerator, const TAttribute<FText>& InToolTip = TAttribute<FText>());
+
 	friend struct FToolMenuSection;
 	friend class UToolMenuEntryScript;
 
diff --git a/Engine/Source/Editor/AnimationBlueprintEditor/Private/AnimationBlueprintInterfaceEditorMode.cpp b/Engine/Source/Editor/AnimationBlueprintEditor/Private/AnimationBlueprintInterfaceEditorMode.cpp
--- a/Engine/Source/Editor/AnimationBlueprintEditor/Private/AnimationBlueprintInterfaceEditorMode.cpp
+++ b/Engine/Source/Editor/AnimationBlueprintEditor/Private/AnimationBlueprintInterfaceEditorMode.cpp
@@ -16,7 +16,7 @@ FAnimationBlueprintInterfaceEditorMode::FAnimationBlueprintInterfaceEditorMode(c
 {
 	AnimBlueprintPtr = CastChecked<UAnimBlueprint>(InAnimationBlueprintEditor->GetBlueprintObj());
 
-	TabLayout = FTabManager::NewLayout( "Standalone_AnimationBlueprintInterfaceEditMode_Layout_v1.0" )
+	TabLayout = FTabManager::NewLayout( "Standalone_AnimationBlueprintInterfaceEditMode_Layout_v1.1" )
 		->AddArea
 		(
 			FTabManager::NewPrimaryArea()
diff --git a/Engine/Source/Editor/EditorStyle/Private/StarshipStyle.cpp b/Engine/Source/Editor/EditorStyle/Private/StarshipStyle.cpp
--- a/Engine/Source/Editor/EditorStyle/Private/StarshipStyle.cpp
+++ b/Engine/Source/Editor/EditorStyle/Private/StarshipStyle.cpp
@@ -2388,12 +2388,12 @@ void FStarshipEditorStyle::FStyle::SetupViewportStyles()
 		ViewportToolbarStyle
 			.SetBackground(FSlateNoResource())
 			.SetIconSize(Icon16x16)
+			.SetBackgroundPadding(FMargin(0))
 			.SetLabelPadding(FMargin(0))
 			.SetComboButtonPadding(FMargin(0))
 			.SetBlockPadding(FMargin(2.0f,0.0f))
 			.SetIndentedBlockPadding(FMargin(0))
 			.SetButtonPadding(FMargin(0))
-			.SetSubMenuIndicator(FSlateNoResource())
 			.SetCheckBoxPadding(FMargin(3.0f, 0.0f))
 			.SetButtonStyle(ViewportMenuButton)
 			.SetSeparatorBrush(FSlateNoResource())
@@ -2439,7 +2439,6 @@ void FStarshipEditorStyle::FStyle::SetupViewportStyles()
 			FToolBarStyle()
 			.SetBackground(BOX_BRUSH("Old/Menu_Background", FMargin(8.0f / 64.0f), FLinearColor::Transparent))
 			.SetExpandBrush(IMAGE_BRUSH("Icons/toolbar_expand_16x", Icon8x8))
-			.SetSubMenuIndicator(IMAGE_BRUSH("Common/SubmenuArrow", Icon8x8))
 			.SetComboButtonPadding(FMargin(0))
 			.SetButtonPadding(FMargin(0))
 			.SetCheckBoxPadding(FMargin(4))
@@ -5332,11 +5331,9 @@ void FStarshipEditorStyle::FStyle::SetupPersonaStyle()
 		FLinearColor GreenHover = FLinearColor(GreenHSV.R, GreenHSV.G * .5, GreenHSV.B, GreenHSV.A).HSVToLinearRGB();
 		FLinearColor GreenPress = FLinearColor(GreenHSV.R, GreenHSV.G, GreenHSV.B*.5, GreenHSV.A).HSVToLinearRGB(); 
 
-		FToolBarStyle PlayToolbar = FStarshipCoreStyle::GetCoreStyle().GetWidgetStyle<FToolBarStyle>("SlimToolBar");
+		FToolBarStyle PlayToolbar = FStarshipCoreStyle::GetCoreStyle().GetWidgetStyle<FToolBarStyle>("AssetEditorToolbar");
 
 		const FButtonStyle PlayToolbarButton = FButtonStyle(PlayToolbar.ButtonStyle)
-			.SetPressed(FSlateNoResource())
-			.SetHovered(FSlateNoResource())
 			.SetNormalForeground(FStyleColors::AccentGreen)
 			.SetPressedForeground(GreenPress)
 			.SetHoveredForeground(GreenHover);
@@ -5346,9 +5343,10 @@ void FStarshipEditorStyle::FStyle::SetupPersonaStyle()
 			.SetDownArrowImage(CORE_IMAGE_BRUSH_SVG("Starship/Common/ellipsis-vertical-narrow", FVector2D(6, 24)));
 		PlayToolbarComboButton.ButtonStyle = PlayToolbarButton;
 
-	
 		PlayToolbar.SetButtonStyle(PlayToolbarButton);
 		PlayToolbar.SetSettingsComboButtonStyle(PlayToolbarComboButton);
+		PlayToolbar.SetShowLabels(true);
+
 		Set("PlayToolBar", PlayToolbar);
 
 		Set( "PlayWorld.Simulate", new IMAGE_BRUSH( "Icons/icon_simulate_40x", Icon40x40 ) );
diff --git a/Engine/Source/Editor/Kismet/Private/SBlueprintEditorSelectedDebugObjectWidget.cpp b/Engine/Source/Editor/Kismet/Private/SBlueprintEditorSelectedDebugObjectWidget.cpp
--- a/Engine/Source/Editor/Kismet/Private/SBlueprintEditorSelectedDebugObjectWidget.cpp
+++ b/Engine/Source/Editor/Kismet/Private/SBlueprintEditorSelectedDebugObjectWidget.cpp
@@ -138,14 +138,14 @@ TSharedRef<SWidget> SBlueprintEditorSelectedDebugObjectWidget::OnGetActiveDetail
 			SNew(SHorizontalBox)
 			+ SHorizontalBox::Slot()
 			.VAlign(VAlign_Center)
-			.Padding(5.0f)
+			.Padding(0.0f)
 			.AutoWidth()
 			[
 				DebugWorldsComboBox.ToSharedRef()
 			]
 			+ SHorizontalBox::Slot()
 			.VAlign(VAlign_Center)
-			.Padding(0.0f, 5.0f)
+			.Padding(0.0f)
 			.AutoWidth()
 			[
 				DebugObjectSelectionWidget
diff --git a/Engine/Source/Editor/Kismet/Private/SBlueprintEditorToolbar.cpp b/Engine/Source/Editor/Kismet/Private/SBlueprintEditorToolbar.cpp
--- a/Engine/Source/Editor/Kismet/Private/SBlueprintEditorToolbar.cpp
+++ b/Engine/Source/Editor/Kismet/Private/SBlueprintEditorToolbar.cpp
@@ -440,23 +440,21 @@ void FBlueprintEditorToolbar::AddCompileToolbar(UToolMenu* InMenu)
 			{
 				const FFullBlueprintEditorCommands& Commands = FFullBlueprintEditorCommands::Get();
 
-				InSection.AddEntry(FToolMenuEntry::InitToolBarButton(
+				FToolMenuEntry CompileButton = FToolMenuEntry::InitToolBarButton(
 					Commands.Compile,
 					TAttribute<FText>(),
 					TAttribute<FText>(BlueprintEditorToolbar.ToSharedRef(), &FBlueprintEditorToolbar::GetStatusTooltip),
 					TAttribute<FSlateIcon>(BlueprintEditorToolbar.ToSharedRef(), &FBlueprintEditorToolbar::GetStatusImage),
-					"CompileBlueprint"
-				));
+					"CompileBlueprint");
 
-				InSection.AddEntry(FToolMenuEntry::InitComboButton(
-					"BlueprintCompileOptions",
+				CompileButton.StyleNameOverride = "CalloutToolbar";
+
+				CompileButton.AddOptionsDropdown(
 					FUIAction(),
 					FOnGetContent::CreateStatic(&BlueprintEditorToolbarImpl::GenerateCompileOptionsWidget, Context->BlueprintEditor.Pin()->GetToolkitCommands()),
-					LOCTEXT("BlupeintCompileOptions_ToolbarName",    "Compile Options"),
-					LOCTEXT("BlupeintCompileOptions_ToolbarTooltip", "Options to customize how Blueprints compile"),
-					TAttribute<FSlateIcon>(),
-					/*bSimpleComboBox =*/true
-				));
+					LOCTEXT("BlupeintCompileOptions_ToolbarTooltip", "Options to customize how Blueprints compile"));
+
+				InSection.AddEntry(CompileButton);
 			}
 		}
 	}));
diff --git a/Engine/Source/Editor/LevelEditor/Private/LevelEditorToolBar.cpp b/Engine/Source/Editor/LevelEditor/Private/LevelEditorToolBar.cpp
--- a/Engine/Source/Editor/LevelEditor/Private/LevelEditorToolBar.cpp
+++ b/Engine/Source/Editor/LevelEditor/Private/LevelEditorToolBar.cpp
@@ -16,7 +16,6 @@
 #include "Widgets/Input/SSlider.h"
 #include "Widgets/Layout/SSpacer.h"
 #include "EditorStyleSet.h"
-
 #include "Settings/EditorExperimentalSettings.h"
 #include "GameMapsSettings.h"
 #include "GameFramework/PlayerController.h"
@@ -24,7 +23,6 @@
 #include "GameFramework/HUD.h"
 #include "GameFramework/GameStateBase.h"
 #include "Engine/TextureStreamingTypes.h"
-
 #include "LevelEditor.h"
 #include "LevelEditorActions.h"
 #include "SourceCodeNavigation.h"
@@ -1184,6 +1182,7 @@ void FLevelEditorToolBar::RegisterLevelEditorToolBar( const TSharedRef<FUIComman
 #define LOCTEXT_NAMESPACE "LevelEditorToolBar"
 
 	UToolMenu* AssetsToolBar = UToolMenus::Get()->RegisterMenu("LevelEditor.LevelEditorToolBar.AssetsToolBar", NAME_None, EMultiBoxType::SlimHorizontalToolBar);
+	AssetsToolBar->StyleName = "AssetEditorToolbar";
 	{
 		{
 			FToolMenuSection& Section = AssetsToolBar->AddSection("File");
@@ -1234,6 +1233,7 @@ void FLevelEditorToolBar::RegisterLevelEditorToolBar( const TSharedRef<FUIComman
 	}
 
 	UToolMenu* ModesToolbar = UToolMenus::Get()->RegisterMenu("LevelEditor.LevelEditorToolBar.ModesToolBar", NAME_None, EMultiBoxType::SlimHorizontalToolBar);
+	ModesToolbar->StyleName = "AssetEditorToolbar";
 	{
 		FToolMenuSection& Section = ModesToolbar->AddSection("EditorModes");
 
@@ -1294,6 +1294,8 @@ void FLevelEditorToolBar::RegisterLevelEditorToolBar( const TSharedRef<FUIComman
 	}
 
 	UToolMenu* SettingsToolBar = UToolMenus::Get()->RegisterMenu("LevelEditor.LevelEditorToolBar.MarketplaceToolBar", NAME_None, EMultiBoxType::SlimHorizontalToolBar);
+
+	SettingsToolBar->StyleName = "AssetEditorToolbar";
 	{
 		FToolMenuSection& SettingsSection = SettingsToolBar->AddSection("Settings");
 		if (FLauncherPlatformModule::Get()->CanOpenLauncher(true))
@@ -1303,6 +1305,8 @@ void FLevelEditorToolBar::RegisterLevelEditorToolBar( const TSharedRef<FUIComman
 	}
 
 	UToolMenu* PlayToolBar = UToolMenus::Get()->RegisterMenu("LevelEditor.LevelEditorToolBar.PlayToolBar", NAME_None, EMultiBoxType::SlimHorizontalToolBar);
+	PlayToolBar->StyleName = "AssetEditorToolbar";
+
 	{
 		FToolMenuSection& PlaySection = PlayToolBar->AddSection("Play");
 
@@ -1377,16 +1381,19 @@ void FLevelEditorToolBar::RegisterLevelEditorToolBar( const TSharedRef<FUIComman
 		// Add the shared play-world commands that will be shown on the Kismet toolbar as well
 		FPlayWorldCommands::BuildToolbar(PlaySection, true);
 
-		PlaySection.AddEntry(FToolMenuEntry::InitComboButton(
-			"LevelToolbarQuickSettings",
-			FUIAction(),
-			FOnGetContent::CreateStatic(&FLevelEditorToolBar::GenerateQuickSettingsMenu, InCommandList, TWeakPtr<SLevelEditor>(InLevelEditor)),
-			LOCTEXT("QuickSettingsCombo", "Settings"),
-			LOCTEXT("QuickSettingsCombo_ToolTip", "Project and Editor settings"),
-			FSlateIcon(FEditorStyle::GetStyleSetName(), "LevelEditor.GameSettings"),
-			false,
-			"LevelToolbarQuickSettings"
-		));
+		FToolMenuEntry SettingsEntry =
+			FToolMenuEntry::InitComboButton(
+				"LevelToolbarQuickSettings",
+				FUIAction(),
+				FOnGetContent::CreateStatic(&FLevelEditorToolBar::GenerateQuickSettingsMenu, InCommandList, TWeakPtr<SLevelEditor>(InLevelEditor)),
+				LOCTEXT("QuickSettingsCombo", "Settings"),
+				LOCTEXT("QuickSettingsCombo_ToolTip", "Project and Editor settings"),
+				FSlateIcon(FEditorStyle::GetStyleSetName(), "LevelEditor.GameSettings"),
+				false,
+				"LevelToolbarQuickSettings");
+		SettingsEntry.StyleNameOverride = "CalloutToolbar";
+
+		PlaySection.AddEntry(SettingsEntry);
 
 	}
 
diff --git a/Engine/Source/Editor/Persona/Private/SAssetFamilyShortcutBar.cpp b/Engine/Source/Editor/Persona/Private/SAssetFamilyShortcutBar.cpp
--- a/Engine/Source/Editor/Persona/Private/SAssetFamilyShortcutBar.cpp
+++ b/Engine/Source/Editor/Persona/Private/SAssetFamilyShortcutBar.cpp
@@ -151,6 +151,7 @@ public:
 			.AutoWidth()
 			[
 				SNew(SSeparator)
+				.Visibility(this, &SAssetShortcut::GetComboVisibility)
 				.Thickness(1.0)
 				.Orientation(EOrientation::Orient_Vertical)
 			]
diff --git a/Engine/Source/Editor/SkeletalMeshEditor/Private/SkeletalMeshEditor.cpp b/Engine/Source/Editor/SkeletalMeshEditor/Private/SkeletalMeshEditor.cpp
--- a/Engine/Source/Editor/SkeletalMeshEditor/Private/SkeletalMeshEditor.cpp
+++ b/Engine/Source/Editor/SkeletalMeshEditor/Private/SkeletalMeshEditor.cpp
@@ -530,8 +530,11 @@ void FSkeletalMeshEditor::ExtendToolbar()
 
 	{
 		FToolMenuSection& Section = ToolMenu->AddSection("Mesh", FText(), SectionInsertLocation);
+		
+		FToolMenuEntry Entry = FToolMenuEntry::InitToolBarButton(FSkeletalMeshEditorCommands::Get().ReimportMesh);
 		Section.AddEntry(FToolMenuEntry::InitToolBarButton(FSkeletalMeshEditorCommands::Get().ReimportMesh));
-		Section.AddEntry(FToolMenuEntry::InitComboButton("ReimportContextMenu", FUIAction(), FNewToolMenuDelegate()));
+
+		Section.AddEntry(FToolMenuEntry::InitComboButton("ReimportContextMenu", FUIAction(), FNewToolMenuDelegate(), FText(), FText(), FSlateIcon(), true));
 	}
 
 	{
diff --git a/Engine/Source/Editor/StaticMeshEditor/Private/StaticMeshEditor.cpp b/Engine/Source/Editor/StaticMeshEditor/Private/StaticMeshEditor.cpp
--- a/Engine/Source/Editor/StaticMeshEditor/Private/StaticMeshEditor.cpp
+++ b/Engine/Source/Editor/StaticMeshEditor/Private/StaticMeshEditor.cpp
@@ -866,12 +866,15 @@ void FStaticMeshEditor::ExtendToolBar()
 					FUIAction(),
 					FOnGetContent::CreateLambda(ConstructReimportContextMenu),
 					TAttribute<FText>(),
-					TAttribute<FText>()
+					TAttribute<FText>(),
+					TAttribute<FSlateIcon>(),
+					true
 				);
 			}
 			ToolbarBuilder.EndSection();
 
 			ToolbarBuilder.BeginSection("Command");
+			ToolbarBuilder.BeginStyleOverride("CalloutToolbar");
 			{
 				TSharedPtr<const FUICommandList> CommandList = ToolbarBuilder.GetTopCommandList();
 
@@ -892,7 +895,7 @@ void FStaticMeshEditor::ExtendToolBar()
 					LOCTEXT("UVToolbarTooltip", "Toggles display of the static mesh's UVs for the specified channel."),
 					FSlateIcon(FEditorStyle::GetStyleSetName(), "StaticMeshEditor.SetDrawUVs"));
 			}
-
+			ToolbarBuilder.EndStyleOverride();
 			ToolbarBuilder.EndSection();
 		}
 	};
diff --git a/Engine/Source/Editor/StatusBar/Private/SourceControlMenuHelpers.cpp b/Engine/Source/Editor/StatusBar/Private/SourceControlMenuHelpers.cpp
--- a/Engine/Source/Editor/StatusBar/Private/SourceControlMenuHelpers.cpp
+++ b/Engine/Source/Editor/StatusBar/Private/SourceControlMenuHelpers.cpp
@@ -270,7 +270,7 @@ TSharedRef<SWidget> FSourceControlMenuHelpers::MakeSourceControlStatusWidget()
 			+ SHorizontalBox::Slot()
 			.AutoWidth()
 			.VAlign(VAlign_Center)
-			.Padding(FMargin(5, 5, 0, 5))
+			.Padding(FMargin(5, 0, 0, 0))
 			[
 				SNew(STextBlock)
 				.TextStyle(&FAppStyle::Get().GetWidgetStyle<FTextBlockStyle>("NormalText"))
diff --git a/Engine/Source/Editor/TextureEditor/Private/TextureEditorToolkit.cpp b/Engine/Source/Editor/TextureEditor/Private/TextureEditorToolkit.cpp
--- a/Engine/Source/Editor/TextureEditor/Private/TextureEditorToolkit.cpp
+++ b/Engine/Source/Editor/TextureEditor/Private/TextureEditorToolkit.cpp
@@ -1067,16 +1067,18 @@ void FTextureEditorToolkit::FillToolbar(FToolBarBuilder& ToolbarBuilder)
 		}
 		ToolbarBuilder.EndSection();
 		ToolbarBuilder.BeginSection("Settings");
+		ToolbarBuilder.BeginStyleOverride("CalloutToolbar");
 		{
 			ToolbarBuilder.AddWidget(SNew(SSpacer), NAME_None, false, HAlign_Right);
 			ToolbarBuilder.AddComboButton(
 				FUIAction(),
 				FOnGetContent::CreateSP(this, &FTextureEditorToolkit::OnGenerateSettingsMenu),
-				LOCTEXT("SettingsMenu", "Settings"),
+				LOCTEXT("SettingsMenu", "View Settings"),
 				FText::GetEmpty(),
 				FSlateIcon(FAppStyle::GetAppStyleSetName(), "Icons.Settings")
 			);
 		}
+		ToolbarBuilder.EndStyleOverride();
 	}
 }
 
diff --git a/Engine/Source/Editor/ToolMenusEditor/Private/SEditToolMenuDialog.cpp b/Engine/Source/Editor/ToolMenusEditor/Private/SEditToolMenuDialog.cpp
--- a/Engine/Source/Editor/ToolMenusEditor/Private/SEditToolMenuDialog.cpp
+++ b/Engine/Source/Editor/ToolMenusEditor/Private/SEditToolMenuDialog.cpp
@@ -266,7 +266,7 @@ FReply SMultiBlockDragHandle::OnDragDetected( const FGeometry& MyGeometry, const
 {
 	TSharedPtr<SWidget> CustomDecorator;
 
-	TSharedRef<SWidget> BlockWidget = Block->MakeWidget(BaseWidget.Pin().ToSharedRef(), EMultiBlockLocation::None, false)->AsWidget();
+	TSharedRef<SWidget> BlockWidget = Block->MakeWidget(BaseWidget.Pin().ToSharedRef(), EMultiBlockLocation::None, false, nullptr)->AsWidget();
 
 	if (Block->IsSeparator())
 	{
diff --git a/Engine/Source/Editor/TurnkeySupport/Private/TurnkeySupportModule.cpp b/Engine/Source/Editor/TurnkeySupport/Private/TurnkeySupportModule.cpp
--- a/Engine/Source/Editor/TurnkeySupport/Private/TurnkeySupportModule.cpp
+++ b/Engine/Source/Editor/TurnkeySupport/Private/TurnkeySupportModule.cpp
@@ -1460,7 +1460,7 @@ void FTurnkeySupportModule::MakeTurnkeyMenu(FToolMenuSection& MenuSection) const
 		return !FTurnkeyEditorSupport::IsPIERunning();
 	});
 
- 	MenuSection.AddEntry(FToolMenuEntry::InitComboButton(
+	FToolMenuEntry Entry = FToolMenuEntry::InitComboButton(
 		"PlatformsMenu",
 		PlatformMenuShownDelegate,
 		FOnGetContent::CreateLambda([this] { return MakeTurnkeyMenuWidget(); }),
@@ -1468,8 +1468,10 @@ void FTurnkeySupportModule::MakeTurnkeyMenu(FToolMenuSection& MenuSection) const
 		LOCTEXT("PlatformMenu_Tooltip", "Platform related actions and settings (Launching, Packaging, custom builds, etc)"),
 		FSlateIcon(FEditorStyle::GetStyleSetName(), "PlayWorld.RepeatLastLaunch"), // not great name for a good "platforms" icon
 		false,
-		"PlatformsMenu"
-	));
+		"PlatformsMenu");
+	Entry.StyleNameOverride = "CalloutToolbar";
+
+	MenuSection.AddEntry(Entry);
 }
 
 
diff --git a/Engine/Source/Editor/UnrealEd/Private/Kismet2/DebuggerCommands.cpp b/Engine/Source/Editor/UnrealEd/Private/Kismet2/DebuggerCommands.cpp
--- a/Engine/Source/Editor/UnrealEd/Private/Kismet2/DebuggerCommands.cpp
+++ b/Engine/Source/Editor/UnrealEd/Private/Kismet2/DebuggerCommands.cpp
@@ -588,28 +588,20 @@ void FPlayWorldCommands::BuildToolbar(FToolMenuSection& InSection, bool bInclude
 			TAttribute< FText >::Create(TAttribute< FText >::FGetter::CreateStatic(&FInternalPlayWorldCommandCallbacks::GetRepeatLastPlayToolTip)),
 			TAttribute< FSlateIcon >::Create(TAttribute< FSlateIcon >::FGetter::CreateStatic(&FInternalPlayWorldCommandCallbacks::GetRepeatLastPlayIcon)),
 			FName(TEXT("LevelToolbarPlay")));
-
 	PlayMenuEntry.StyleNameOverride = FName("PlayToolbar");
 
-	// Play
-	InSection.AddEntry(PlayMenuEntry);
-
 	// Play combo box
 	FUIAction SpecialPIEOptionsMenuAction;
 	SpecialPIEOptionsMenuAction.IsActionVisibleDelegate = FIsActionButtonVisible::CreateStatic(&FInternalPlayWorldCommandCallbacks::CanShowNonPlayWorldOnlyActions);
 
-
-	FToolMenuEntry PlayMenuCombo = FToolMenuEntry::InitComboButton(
-		"PlayCombo",
+	PlayMenuEntry.AddOptionsDropdown(
 		SpecialPIEOptionsMenuAction,
 		FOnGetContent::CreateStatic(&GeneratePlayMenuContent, GlobalPlayWorldActions.ToSharedRef()),
-		LOCTEXT("PlayCombo_Label", "Active Play Mode"),
-		LOCTEXT("PIEComboToolTip", "Change Play Mode and Play Settings"),
-		FSlateIcon(FEditorStyle::GetStyleSetName(), "LevelEditor.RepeatLastPlay"),
-		true);
+		LOCTEXT("PIEComboToolTip", "Change Play Mode and Play Settings")
+	);
 
-	PlayMenuCombo.StyleNameOverride = FName("PlayToolbar");
-	InSection.AddEntry(PlayMenuCombo);
+	// Play
+	InSection.AddEntry(PlayMenuEntry);
 
 	if (bIncludeLaunchButtonAndOptions)
 	{
diff --git a/Engine/Source/Editor/UnrealEd/Private/Toolkits/AssetEditorToolkit.cpp b/Engine/Source/Editor/UnrealEd/Private/Toolkits/AssetEditorToolkit.cpp
--- a/Engine/Source/Editor/UnrealEd/Private/Toolkits/AssetEditorToolkit.cpp
+++ b/Engine/Source/Editor/UnrealEd/Private/Toolkits/AssetEditorToolkit.cpp
@@ -1052,6 +1052,12 @@ void FAssetEditorToolkit::RegisterDefaultToolBar(bool bCustomToolbarPlacement)
 	if (!ToolMenus->IsMenuRegistered(DefaultToolBarName))
 	{
 		UToolMenu* ToolbarBuilder = ToolMenus->RegisterMenu(DefaultToolBarName, NAME_None, bCustomToolbarPlacement ? EMultiBoxType::ToolBar : EMultiBoxType::SlimHorizontalToolBar);
+#if 0
+		if (!bCustomToolbarPlacement)
+		{
+			ToolbarBuilder->StyleName = "AssetEditorToolbar";
+		}
+#endif
 		{
 			FToolMenuSection& Section = ToolbarBuilder->AddSection("Asset");
 			Section.AddEntry(FToolMenuEntry::InitToolBarButton(FAssetEditorCommonCommands::Get().SaveAsset));
@@ -1115,13 +1121,11 @@ void FAssetEditorToolkit::GenerateToolbar()
 	Toolbar = 
 		SNew(SHorizontalBox)
 		+SHorizontalBox::Slot()
-		.VAlign(VAlign_Center)
 		[
 			ToolBarWidget
 		]
 		+SHorizontalBox::Slot()
 		.HAlign(HAlign_Right)
-		.VAlign(VAlign_Center)
 		.AutoWidth()
 		[
 			SNew(SBorder)
diff --git a/Engine/Source/Editor/VREditor/UI/VREditorUISystem.cpp b/Engine/Source/Editor/VREditor/UI/VREditorUISystem.cpp
--- a/Engine/Source/Editor/VREditor/UI/VREditorUISystem.cpp
+++ b/Engine/Source/Editor/VREditor/UI/VREditorUISystem.cpp
@@ -1437,7 +1437,7 @@ void UVREditorUISystem::MakeUniformGridMenu(const TSharedRef<FMultiBox>& MultiBo
 		const TSharedRef<const FMultiBlock>& Block = MultiBlock;
 		if(Block->GetType() == EMultiBlockType::MenuEntry)
 		{
-			TSharedRef<SWidget> BlockWidget = Block->MakeWidget(MultiBoxWidget, EMultiBlockLocation::Middle, true)->AsWidget();
+			TSharedRef<SWidget> BlockWidget = Block->MakeWidget(MultiBoxWidget, EMultiBlockLocation::Middle, true, nullptr)->AsWidget();
 		
 			int32 Column = FMath::Fmod(static_cast<float>(Index), static_cast<float>(Columns));
 			Row = FMath::DivideAndRoundDown(Index, Columns);
@@ -1567,7 +1567,7 @@ void UVREditorUISystem::MakeRadialBoxMenu(const TSharedRef<FMultiBox>& MultiBox,
 		const TSharedRef<const FMultiBlock>& Block = MultiBlock;
 		if (Block->GetType() == EMultiBlockType::MenuEntry)
 		{
-			TSharedRef<SWidget> BlockWidget = Block->MakeWidget(MultiBoxWidget, EMultiBlockLocation::Middle, true)->AsWidget();
+			TSharedRef<SWidget> BlockWidget = Block->MakeWidget(MultiBoxWidget, EMultiBlockLocation::Middle, true, nullptr)->AsWidget();
 			TSharedRef<SOverlay> TestOverlay = SNew(SOverlay);
 
 			BlockWidget = AddHoverableButton(BlockWidget, ButtonTypeOverride, TestOverlay);
diff --git a/Engine/Source/Runtime/Slate/Private/Framework/MultiBox/MultiBox.cpp b/Engine/Source/Runtime/Slate/Private/Framework/MultiBox/MultiBox.cpp
--- a/Engine/Source/Runtime/Slate/Private/Framework/MultiBox/MultiBox.cpp
+++ b/Engine/Source/Runtime/Slate/Private/Framework/MultiBox/MultiBox.cpp
@@ -23,6 +23,8 @@
 #include "Framework/MultiBox/MultiBoxCustomization.h"
 #include "Framework/MultiBox/SClippingHorizontalBox.h"
 #include "Framework/MultiBox/SWidgetBlock.h"
+#include "Framework/MultiBox/SToolBarComboButtonBlock.h"
+
 #include "Framework/Commands/UICommandDragDropOp.h"
 #include "SUniformToolbarPanel.h"
 #include "Styling/ToolBarStyle.h"
@@ -76,11 +78,17 @@ void SMultiBlockBaseWidget::SetOwnerMultiBoxWidget(TSharedRef< SMultiBoxWidget >
 	OwnerMultiBoxWidget = InOwnerMultiBoxWidget;
 }
 
-void SMultiBlockBaseWidget::SetMultiBlock(TSharedRef< const FMultiBlock > InMultiBlock)
+void SMultiBlockBaseWidget::SetMultiBlock(TSharedRef<const FMultiBlock> InMultiBlock)
 {
 	MultiBlock = InMultiBlock;
 }
 
+void SMultiBlockBaseWidget::SetOptionsBlockWidget(TSharedPtr<SWidget> InOptionsBlock)
+{
+	OptionsBlockWidget = InOptionsBlock;
+}
+
+
 void SMultiBlockBaseWidget::SetMultiBlockLocation(EMultiBlockLocation::Type InLocation, bool bInSectionContainsIcons)
 {
 	Location = InLocation;
@@ -140,15 +148,17 @@ bool SMultiBlockBaseWidget::IsInEditMode() const
  *
  * @return  MultiBlock widget object
  */
-TSharedRef< IMultiBlockBaseWidget > FMultiBlock::MakeWidget( TSharedRef< SMultiBoxWidget > InOwnerMultiBoxWidget, EMultiBlockLocation::Type InLocation, bool bSectionContainsIcons ) const
+TSharedRef<IMultiBlockBaseWidget> FMultiBlock::MakeWidget(TSharedRef<SMultiBoxWidget> InOwnerMultiBoxWidget, EMultiBlockLocation::Type InLocation, bool bSectionContainsIcons, TSharedPtr<SWidget> OptionsBlockWidget) const
 {
-	TSharedRef< IMultiBlockBaseWidget > NewMultiBlockWidget = ConstructWidget();
+	TSharedRef<IMultiBlockBaseWidget> NewMultiBlockWidget = ConstructWidget();
 
 	// Tell the widget about its parent MultiBox widget
 	NewMultiBlockWidget->SetOwnerMultiBoxWidget( InOwnerMultiBoxWidget );
 
 	// Assign ourselves to the MultiBlock widget
-	NewMultiBlockWidget->SetMultiBlock( AsShared() );
+	NewMultiBlockWidget->SetMultiBlock(AsShared());
+
+	NewMultiBlockWidget->SetOptionsBlockWidget(OptionsBlockWidget);
 
 	// Pass location information to widget.
 	NewMultiBlockWidget->SetMultiBlockLocation(InLocation, bSectionContainsIcons);
@@ -596,7 +606,7 @@ EVisibility SMultiBoxWidget::GetCustomizationBorderDragVisibility(const FName In
 	return EVisibility::Collapsed;
 }
 
-void SMultiBoxWidget::AddBlockWidget( const FMultiBlock& Block, TSharedPtr<SHorizontalBox> HorizontalBox, TSharedPtr<SVerticalBox> VerticalBox, EMultiBlockLocation::Type InLocation,  bool bSectionContainsIcons )
+void SMultiBoxWidget::AddBlockWidget(const FMultiBlock& Block, TSharedPtr<SHorizontalBox> HorizontalBox, TSharedPtr<SVerticalBox> VerticalBox, EMultiBlockLocation::Type InLocation, bool bSectionContainsIcons, TSharedPtr<const FToolBarComboButtonBlock> OptionsBlock)
 {
 	check( MultiBox.IsValid() );
 
@@ -608,7 +618,14 @@ void SMultiBoxWidget::AddBlockWidget( const FMultiBlock& Block, TSharedPtr<SHori
 
 	bool bDisplayExtensionHooks = FMultiBoxSettings::DisplayMultiboxHooks.Get() && Block.GetExtensionHook() != NAME_None;
 
-	TSharedRef<SWidget> BlockWidget = Block.MakeWidget(SharedThis(this), InLocation, bSectionContainsIcons)->AsWidget();
+	TSharedPtr<SWidget> OptionsWidget;
+	if (OptionsBlock)
+	{
+		OptionsWidget = OptionsBlock->MakeWidget(SharedThis(this), EMultiBlockLocation::None, bSectionContainsIcons, nullptr)->AsWidget();
+	}
+
+	TSharedRef<SWidget> BlockWidget = Block.MakeWidget(SharedThis(this), InLocation, bSectionContainsIcons, OptionsWidget)->AsWidget();
+
 
 	TWeakPtr<SWidget> BlockWidgetWeakPtr = BlockWidget;
 	TWeakPtr<const FMultiBlock> BlockWeakPtr = Block.AsShared();
@@ -806,9 +823,13 @@ void SMultiBoxWidget::BuildMultiBoxWidget()
 	const FName& StyleName = MultiBox->GetStyleName();
 
 	const FSlateBrush* BackgroundBrush = nullptr;
+	FMargin BackgroundPadding;
 	if (StyleSet->HasWidgetStyle<FToolBarStyle>(StyleName))
 	{
-		BackgroundBrush = &StyleSet->GetWidgetStyle<FToolBarStyle>(StyleName).BackgroundBrush;
+		const FToolBarStyle& Style = StyleSet->GetWidgetStyle<FToolBarStyle>(StyleName);
+
+		BackgroundBrush = &Style.BackgroundBrush;
+		BackgroundPadding = Style.BackgroundPadding;
 	}
 	else
 	{
@@ -933,6 +954,13 @@ void SMultiBoxWidget::BuildMultiBoxWidget()
 		const FMultiBlock& Block = *Blocks[Index];
 		EMultiBlockLocation::Type Location = EMultiBlockLocation::None;
 		
+		TSharedPtr<const FMultiBlock> NextBlock = nullptr;
+
+		if (Blocks.IsValidIndex(Index + 1))
+		{
+			NextBlock = Blocks[Index + 1];
+		}
+
 		// Determine the location of the current block, used for group styling information
 		{
 			// Check if we are a start or end block
@@ -948,10 +976,9 @@ void SMultiBoxWidget::BuildMultiBoxWidget()
 			// Check if we are next to a start or end block
 			bool bIsNextToStartBlock = false;
 			bool bIsNextToEndBlock = false;
-			if (Index + 1 < Blocks.Num())
+			if (NextBlock)
 			{
-				const FMultiBlock& NextBlock = *Blocks[Index + 1];
-				if ( NextBlock.IsGroupEndBlock() )
+				if ( NextBlock->IsGroupEndBlock() )
 				{
 					bIsNextToEndBlock = true;
 				}
@@ -989,18 +1016,30 @@ void SMultiBoxWidget::BuildMultiBoxWidget()
 			}
 		}
 
+		TSharedPtr<const FToolBarComboButtonBlock> OptionsBlock;
+		if (NextBlock && NextBlock->GetType() == EMultiBlockType::ToolBarComboButton)
+		{
+			TSharedPtr<const FToolBarComboButtonBlock> NextToolBarComboButtonBlock = StaticCastSharedPtr<const FToolBarComboButtonBlock>(NextBlock);
+			if (NextToolBarComboButtonBlock->IsSimpleComboBox())
+			{
+				// Apply a special treatment to simple combo boxes as they represent options for the previous button
+				OptionsBlock = NextToolBarComboButtonBlock;
+				// Skip over options blocks. They are not added directly
+				++Index;
+			}
+		}
 
 		if( DragPreview.IsValid() && DragPreview.InsertIndex == Index )
 		{
 			// Add the drag preview before if we have it. This block shows where the custom block will be 
 			// added if the user drops it
-			AddBlockWidget( *DragPreview.PreviewBlock, HorizontalBox, VerticalBox, EMultiBlockLocation::None, bSectionContainsIcons );
+			AddBlockWidget(*DragPreview.PreviewBlock, HorizontalBox, VerticalBox, EMultiBlockLocation::None, bSectionContainsIcons, OptionsBlock);
 		}
 		
 		// Do not add a block if it is being dragged
 		if( !IsBlockBeingDragged( Blocks[Index] ) )
 		{
-			AddBlockWidget( Block, HorizontalBox, VerticalBox, Location, bSectionContainsIcons );
+			AddBlockWidget(Block, HorizontalBox, VerticalBox, Location, bSectionContainsIcons, OptionsBlock);
 		}
 	}
 
@@ -1027,6 +1066,7 @@ void SMultiBoxWidget::BuildMultiBoxWidget()
 	default:
 		{ 
 			BorderBrush = BackgroundBrush;
+			BorderPadding = BackgroundPadding;
 			BorderForegroundColor = FCoreStyle::Get().GetSlateColor("DefaultForeground");
 		}
 		break;
@@ -1097,7 +1137,7 @@ void SMultiBoxWidget::UpdateDropAreaPreviewBlock( TSharedRef<const FMultiBlock>
 					MakeShareable(
 						new FDropPreviewBlock( 
 							NewBlock.ToSharedRef(), 
-							NewBlock->MakeWidget( SharedThis(this), EMultiBlockLocation::None, NewBlock->HasIcon() ) )
+							NewBlock->MakeWidget( SharedThis(this), EMultiBlockLocation::None, NewBlock->HasIcon(),nullptr ) )
 					);
 
 				bAddedNewBlock = true;
diff --git a/Engine/Source/Runtime/Slate/Private/Framework/MultiBox/SToolBarButtonBlock.cpp b/Engine/Source/Runtime/Slate/Private/Framework/MultiBox/SToolBarButtonBlock.cpp
--- a/Engine/Source/Runtime/Slate/Private/Framework/MultiBox/SToolBarButtonBlock.cpp
+++ b/Engine/Source/Runtime/Slate/Private/Framework/MultiBox/SToolBarButtonBlock.cpp
@@ -11,6 +11,7 @@
 #include "Framework/MultiBox/SToolBarComboButtonBlock.h"
 #include "Styling/ToolBarStyle.h"
 #include "Widgets/Images/SLayeredImage.h"
+#include "Widgets/Layout/SSeparator.h"
 
 
 FToolBarButtonBlock::FToolBarButtonBlock( const TSharedPtr< const FUICommandInfo > InCommand, TSharedPtr< const FUICommandList > InCommandList, const TAttribute<FText>& InLabelOverride, const TAttribute<FText>& InToolTipOverride, const TAttribute<FSlateIcon>& InIconOverride )
@@ -75,11 +76,11 @@ bool FToolBarButtonBlock::HasIcon() const
 TSharedRef< class IMultiBlockBaseWidget > FToolBarButtonBlock::ConstructWidget() const
 {
 	return SNew( SToolBarButtonBlock )
-		.LabelVisibility( LabelVisibility.IsSet() ? LabelVisibility.GetValue() : TOptional< EVisibility >() )
-		.IsFocusable( bIsFocusable )
-		.ForceSmallIcons( bForceSmallIcons )
+		.LabelVisibility(LabelVisibility)
+		.IsFocusable(bIsFocusable)
+		.ForceSmallIcons(bForceSmallIcons)
 		.TutorialHighlightName(GetTutorialHighlightName())
-		.Cursor( EMouseCursor::Default );
+		.Cursor(EMouseCursor::Default);
 }
 
 
@@ -90,15 +91,7 @@ TSharedRef< class IMultiBlockBaseWidget > FToolBarButtonBlock::ConstructWidget()
  */
 void SToolBarButtonBlock::Construct( const FArguments& InArgs )
 {
-	if ( InArgs._LabelVisibility.IsSet() )
-	{
-		LabelVisibility = InArgs._LabelVisibility.GetValue();
-	}
-	else
-	{
-		LabelVisibility = TAttribute<EVisibility>::Create(TAttribute<EVisibility>::FGetter::CreateSP(SharedThis(this), &SToolBarButtonBlock::GetIconVisibility, false));
-	}
-
+	LabelVisibilityOverride = InArgs._LabelVisibility;
 	bIsFocusable = InArgs._IsFocusable;
 	bForceSmallIcons = InArgs._ForceSmallIcons;
 	TutorialHighlightName = InArgs._TutorialHighlightName;
@@ -112,6 +105,22 @@ void SToolBarButtonBlock::BuildMultiBlockWidget(const ISlateStyle* StyleSet, con
 {
 	const FToolBarStyle& ToolBarStyle = StyleSet->GetWidgetStyle<FToolBarStyle>(StyleName);
 
+	// If override is set use that
+	if (LabelVisibilityOverride.IsSet())
+	{
+		LabelVisibility = LabelVisibilityOverride.GetValue();
+	}
+	else if (!ToolBarStyle.bShowLabels)
+	{
+		// Otherwise check the style
+		LabelVisibility = EVisibility::Collapsed;
+	}
+	else
+	{
+		// Finally if the style doesnt disable labels, use the default
+		LabelVisibility = TAttribute<EVisibility>::Create(TAttribute<EVisibility>::FGetter::CreateSP(SharedThis(this), &SToolBarButtonBlock::GetIconVisibility, false));
+	}
+
 	struct Local
 	{
 		/** Appends the key binding to the end of the provided ToolTip */
@@ -132,10 +141,9 @@ void SToolBarButtonBlock::BuildMultiBlockWidget(const ISlateStyle* StyleSet, con
 		}
 	};
 
-
 	TSharedRef<const FMultiBox> MultiBox = OwnerMultiBoxWidget.Pin()->GetMultiBox();
 
-	TSharedRef<const FToolBarButtonBlock > ToolBarButtonBlock = StaticCastSharedRef< const FToolBarButtonBlock >(MultiBlock.ToSharedRef());
+	TSharedRef<const FToolBarButtonBlock> ToolBarButtonBlock = StaticCastSharedRef<const FToolBarButtonBlock>(MultiBlock.ToSharedRef());
 
 	TSharedPtr<const FUICommandInfo> UICommand = ToolBarButtonBlock->GetAction();
 
@@ -287,27 +295,48 @@ void SToolBarButtonBlock::BuildMultiBlockWidget(const ISlateStyle* StyleSet, con
 		];
 	}
 
-
-	// Often buttons have a "simple" combo box next to it.  The button + simple combo button is designed to feel like a normal combo button but 
-	// when the button part is pressed some action happens independently of the combo dropdown and the dropdown changes settings the button uses.  
-	// We want this to feel like one widget so we space them closer together.
-	FMargin Padding = ToolBarStyle.ButtonPadding;
+	if (OptionsBlockWidget.IsValid())
 	{
-		int32 MyIndex = MultiBox->GetBlocks().Find(ToolBarButtonBlock);
-		if (MultiBox->GetBlocks().IsValidIndex(MyIndex + 1))
-		{
-			const TSharedRef<const FMultiBlock>& NextBlock = MultiBox->GetBlocks()[MyIndex + 1];
-			if (NextBlock->GetType() == EMultiBlockType::ToolBarComboButton)
-			{
-				TSharedRef<const FToolBarComboButtonBlock> NextToolBarComboButtonBlock = StaticCastSharedRef<const FToolBarComboButtonBlock>(NextBlock);
-				if (NextToolBarComboButtonBlock->IsSimpleComboBox())
-				{
-					Padding.Right = 2.0f;
-				}
-			}
-		}
-
-		ChildSlot.Padding(Padding);
+		ChildSlot
+		[
+			SNew(SHorizontalBox)
+			+ SHorizontalBox::Slot()
+			.AutoWidth()
+			[
+				SAssignNew(ButtonBorder, SBorder)
+				.Padding(FMargin(10.0f, 0.0f,2.0f, 0.0f))
+				.BorderImage(this, &SToolBarButtonBlock::GetOptionsBlockLeftBrush)
+				.VAlign(VAlign_Center)
+				[
+					ChildSlot.GetWidget()
+				]
+			]
+			/*+ SHorizontalBox::Slot()
+			.AutoWidth()
+			[
+				SNew(SSeparator)
+				.Orientation(Orient_Vertical)
+				.Visibility(this, &SToolBarButtonBlock::GetOptionsSeparatorVisibility)
+				.Thickness(1.0f)
+			]*/
+			+ SHorizontalBox::Slot()
+			.Padding(0.0f,0.0f)
+			.AutoWidth()
+			[
+				SAssignNew(OptionsBorder, SBorder)
+				.Padding(FMargin(0.0f, 0.0f, 2.0f, 0.0f))
+				.BorderImage(this, &SToolBarButtonBlock::GetOptionsBlockRightBrush)
+				.VAlign(VAlign_Center)
+				[
+					OptionsBlockWidget.ToSharedRef()
+				]
+			]
+		];
+		ChildSlot.Padding(ToolBarStyle.ComboButtonPadding.Left, 0.0f, ToolBarStyle.ComboButtonPadding.Right, 0.0f);
+	}
+	else
+	{
+		ChildSlot.Padding(ToolBarStyle.ButtonPadding);
 	}
 
 	// Bind our widget's enabled state to whether or not our action can execute
@@ -515,3 +544,45 @@ FSlateColor SToolBarButtonBlock::GetIconForegroundColor() const
 
 	return FSlateColor::UseForeground();
 }
+
+const FSlateBrush* SToolBarButtonBlock::GetOptionsBlockLeftBrush() const
+{
+	if (ButtonBorder->IsHovered())
+	{
+		static const FName LeftHover("ToolbarSettingsRegion.LeftHover");
+		return FAppStyle::Get().GetBrush(LeftHover);
+	}
+	else if (OptionsBorder->IsHovered())
+	{
+		static const FName Left("ToolbarSettingsRegion.Left");
+		return FAppStyle::Get().GetBrush(Left);
+	}
+	else
+	{
+		return FStyleDefaults::GetNoBrush();
+	}
+
+}
+
+const FSlateBrush* SToolBarButtonBlock::GetOptionsBlockRightBrush() const
+{
+	if (OptionsBorder->IsHovered())
+	{
+		static const FName RightHover("ToolbarSettingsRegion.RightHover");
+		return FAppStyle::Get().GetBrush(RightHover);
+	}
+	else if (ButtonBorder->IsHovered())
+	{
+		static const FName Right("ToolbarSettingsRegion.Right");
+		return FAppStyle::Get().GetBrush(Right);
+	}
+	else
+	{
+		return FStyleDefaults::GetNoBrush();
+	}
+}
+
+EVisibility SToolBarButtonBlock::GetOptionsSeparatorVisibility() const
+{
+	return IsHovered() ? EVisibility::HitTestInvisible : EVisibility::Hidden;
+}
diff --git a/Engine/Source/Runtime/Slate/Private/Framework/MultiBox/SToolBarComboButtonBlock.cpp b/Engine/Source/Runtime/Slate/Private/Framework/MultiBox/SToolBarComboButtonBlock.cpp
--- a/Engine/Source/Runtime/Slate/Private/Framework/MultiBox/SToolBarComboButtonBlock.cpp
+++ b/Engine/Source/Runtime/Slate/Private/Framework/MultiBox/SToolBarComboButtonBlock.cpp
@@ -50,15 +50,7 @@ TSharedRef< class IMultiBlockBaseWidget > FToolBarComboButtonBlock::ConstructWid
 
 void SToolBarComboButtonBlock::Construct( const FArguments& InArgs )
 {
-	if ( InArgs._LabelVisibility.IsSet() )
-	{
-		LabelVisibility = InArgs._LabelVisibility.GetValue();
-	}
-	else
-	{
-		LabelVisibility = TAttribute< EVisibility >::Create( TAttribute< EVisibility >::FGetter::CreateSP( SharedThis( this ), &SToolBarComboButtonBlock::GetIconVisibility, false ) );
-	}
-
+	LabelVisibilityOverride = InArgs._LabelVisibility;
 	Icon = InArgs._Icon;
 	bForceSmallIcons = InArgs._ForceSmallIcons;
 }
@@ -75,6 +67,21 @@ void SToolBarComboButtonBlock::BuildMultiBlockWidget(const ISlateStyle* StyleSet
 
 	const FToolBarStyle& ToolBarStyle = StyleSet->GetWidgetStyle<FToolBarStyle>(StyleName);
 
+	// If override is set use that
+	if (LabelVisibilityOverride.IsSet())
+	{
+		LabelVisibility = LabelVisibilityOverride.GetValue();
+	}
+	else if (!ToolBarStyle.bShowLabels)
+	{
+		// Otherwise check the style
+		LabelVisibility = EVisibility::Collapsed;
+	}
+	else
+	{
+		LabelVisibility = TAttribute< EVisibility >::Create(TAttribute< EVisibility >::FGetter::CreateSP(SharedThis(this), &SToolBarComboButtonBlock::GetIconVisibility, false));
+	}
+
 	TSharedRef<SWidget> IconWidget = SNullWidget::NullWidget;
 	if (!ToolBarComboButtonBlock->bSimpleComboBox)
 	{
@@ -203,8 +210,7 @@ void SToolBarComboButtonBlock::BuildMultiBlockWidget(const ISlateStyle* StyleSet
 	FMargin Padding = ToolBarStyle.ComboButtonPadding;
 	if (ToolBarComboButtonBlock->bSimpleComboBox)
 	{
-		Padding.Left = 0;
-		Padding.Right = 10;
+		Padding = FMargin(0);
 	}
 
 	ChildSlot.Padding(Padding);
diff --git a/Engine/Source/Runtime/Slate/Public/Framework/MultiBox/MultiBox.h b/Engine/Source/Runtime/Slate/Public/Framework/MultiBox/MultiBox.h
--- a/Engine/Source/Runtime/Slate/Public/Framework/MultiBox/MultiBox.h
+++ b/Engine/Source/Runtime/Slate/Public/Framework/MultiBox/MultiBox.h
@@ -26,6 +26,7 @@ class STableViewBase;
 class SVerticalBox;
 class UToolMenuBase;
 class SUniformWrapPanel;
+class FToolBarComboButtonBlock;
 
 namespace MultiBoxConstants
 {	
@@ -152,7 +153,7 @@ public:
 	 *
 	 * @return  MultiBlock widget object
 	 */
-	TSharedRef< class IMultiBlockBaseWidget > MakeWidget( TSharedRef< class SMultiBoxWidget > InOwnerMultiBoxWidget, EMultiBlockLocation::Type InLocation, bool bSectionContainsIcons ) const;
+	TSharedRef<class IMultiBlockBaseWidget> MakeWidget(TSharedRef< class SMultiBoxWidget > InOwnerMultiBoxWidget, EMultiBlockLocation::Type InLocation, bool bSectionContainsIcons, TSharedPtr<SWidget> OptionsBlockWidget) const;
 
 	/**
 	 * Gets the type of this MultiBox
@@ -450,7 +451,7 @@ public:
 	 *
 	 * @return  Widget reference
 	 */
-	virtual TSharedRef< SWidget > AsWidget() = 0;
+	virtual TSharedRef<SWidget> AsWidget() = 0;
 
 
 	/**
@@ -458,24 +459,28 @@ public:
 	 *
 	 * @return  Widget reference
 	 */
-	virtual TSharedRef< const SWidget > AsWidget() const = 0;
-
+	virtual TSharedRef<const SWidget> AsWidget() const = 0;
 
 	/**
 	 * Associates the owner MultiBox widget with this widget
 	 *
 	 * @param	InOwnerMultiBoxWidget		The MultiBox widget that owns us
 	 */
-	virtual void SetOwnerMultiBoxWidget( TSharedRef< SMultiBoxWidget > InOwnerMultiBoxWidget ) = 0;
+	virtual void SetOwnerMultiBoxWidget(TSharedRef<SMultiBoxWidget> InOwnerMultiBoxWidget ) = 0;
 
-	
 	/**
 	 * Associates this widget with a MultiBlock
 	 *
 	 * @param	InMultiBlock	The MultiBlock we'll be associated with
 	 */
-	virtual void SetMultiBlock( TSharedRef< const FMultiBlock > InMultiBlock ) = 0;
+	virtual void SetMultiBlock(TSharedRef<const FMultiBlock> InMultiBlock) = 0;
 
+	/**
+	 * Adds a dropdown widget for options associated with this widget. The usage of this is block specific
+	 *
+	 * @param	InOptionsBlockWidget	The options block to associate with this widget
+	 */
+	virtual void SetOptionsBlockWidget(TSharedPtr<SWidget> InOptionsBlockWidget) = 0;
 
 	/**
 	 * Builds this MultiBlock widget up from the MultiBlock associated with it
@@ -519,23 +524,26 @@ public:
 	/** IMultiBlockBaseWidget interface */
 	virtual TSharedRef< SWidget > AsWidget() override;
 	virtual TSharedRef< const SWidget > AsWidget() const override;
-	virtual void SetOwnerMultiBoxWidget( TSharedRef< SMultiBoxWidget > InOwnerMultiBoxWidget ) override;
-	virtual void SetMultiBlock( TSharedRef< const FMultiBlock > InMultiBlock ) override;
-	virtual void SetMultiBlockLocation( EMultiBlockLocation::Type InLocation, bool bInSectionContainsIcons ) override;
+	virtual void SetOwnerMultiBoxWidget(TSharedRef<SMultiBoxWidget> InOwnerMultiBoxWidget) override;
+	virtual void SetMultiBlock(TSharedRef<const FMultiBlock> InMultiBlock) override;
+	virtual void SetOptionsBlockWidget(TSharedPtr<SWidget> InOptionsBlock) override;
+	virtual void SetMultiBlockLocation(EMultiBlockLocation::Type InLocation, bool bInSectionContainsIcons) override;
 	virtual EMultiBlockLocation::Type GetMultiBlockLocation() override;
 	virtual bool IsInEditMode() const override;
 
 	/** SWidget Interface */
-	virtual void OnDragEnter( const FGeometry& MyGeometry, const FDragDropEvent& DragDropEvent ) override;
-	virtual FReply OnDragOver( const FGeometry& MyGeometry, const FDragDropEvent& DragDropEvent ) override;
-	virtual FReply OnDrop( const FGeometry& MyGeometry, const FDragDropEvent& DragDropEvent ) override;
+	virtual void OnDragEnter(const FGeometry& MyGeometry, const FDragDropEvent& DragDropEvent) override;
+	virtual FReply OnDragOver(const FGeometry& MyGeometry, const FDragDropEvent& DragDropEvent) override;
+	virtual FReply OnDrop(const FGeometry& MyGeometry, const FDragDropEvent& DragDropEvent) override;
 protected:
 
 	/** Weak reference back to the MultiBox widget that owns us */
-	TWeakPtr< SMultiBoxWidget > OwnerMultiBoxWidget;
+	TWeakPtr<SMultiBoxWidget> OwnerMultiBoxWidget;
 
 	/** The MultiBlock we're associated with */
-	TSharedPtr< const FMultiBlock > MultiBlock;
+	TSharedPtr<const FMultiBlock> MultiBlock;
+
+	TSharedPtr<SWidget> OptionsBlockWidget;
 
 	/** The MultiBlocks location relative to the other blocks in the set */
 	EMultiBlockLocation::Type Location;
@@ -737,7 +745,7 @@ public:
 	/**
 	 * @return The visibility of customization widgets for a block
 	 */
-	EVisibility GetCustomizationVisibility( TWeakPtr<const FMultiBlock> BlockWeakPtr, TWeakPtr<SWidget> BlockWidgetWeakPtr ) const;
+	EVisibility GetCustomizationVisibility(TWeakPtr<const FMultiBlock> BlockWeakPtr, TWeakPtr<SWidget> BlockWidgetWeakPtr) const;
 
 	/**
 	 * @return The visibility of the drop location indicator of a drag and drop for a block
@@ -746,7 +754,7 @@ public:
 
 private:
 	/** Adds a block Widget to this widget */
-	void AddBlockWidget( const FMultiBlock& Block, TSharedPtr<SHorizontalBox> HorizontalBox, TSharedPtr<SVerticalBox> VerticalBox, EMultiBlockLocation::Type InLocation,  bool bSectionContainsIcons );
+	void AddBlockWidget(const FMultiBlock& Block, TSharedPtr<SHorizontalBox> HorizontalBox, TSharedPtr<SVerticalBox> VerticalBox, EMultiBlockLocation::Type InLocation, bool bSectionContainsIcons, TSharedPtr<const FToolBarComboButtonBlock> OptionsBlock);
 
 	/**
 	 * Updates the preview block being dragged.  The drag area is where the users dragged block will be dropped
diff --git a/Engine/Source/Runtime/Slate/Public/Framework/MultiBox/SToolBarButtonBlock.h b/Engine/Source/Runtime/Slate/Public/Framework/MultiBox/SToolBarButtonBlock.h
--- a/Engine/Source/Runtime/Slate/Public/Framework/MultiBox/SToolBarButtonBlock.h
+++ b/Engine/Source/Runtime/Slate/Public/Framework/MultiBox/SToolBarButtonBlock.h
@@ -12,6 +12,8 @@
 #include "Framework/Commands/UICommandList.h"
 #include "Framework/MultiBox/MultiBox.h"
 
+class SBorder;
+
 /**
  * Tool bar button MultiBlock
  */
@@ -103,7 +105,7 @@ public:
 		, _TutorialHighlightName(NAME_None)
 		{}
 
-		SLATE_ARGUMENT( TOptional< EVisibility >, LabelVisibility )
+		SLATE_ARGUMENT( TOptional<EVisibility>, LabelVisibility )
 		SLATE_ARGUMENT( bool, IsFocusable )
 		SLATE_ARGUMENT( bool, ForceSmallIcons )
 		SLATE_ARGUMENT( FName, TutorialHighlightName )
@@ -171,9 +173,23 @@ private:
 	const FSlateBrush* GetSmallIconBrush() const;
 
 	FSlateColor GetIconForegroundColor() const;
+
+	/** Gets the brush for the left side options block (the side with the clickable toolbar button */
+	const FSlateBrush* GetOptionsBlockLeftBrush() const;
+
+	/** Gets the brush for the right side options block (the side with the options dropdown */
+	const FSlateBrush* GetOptionsBlockRightBrush() const;
+
+	EVisibility GetOptionsSeparatorVisibility() const;
 private:
+	/** Overrides the visibility of the of label. This is used to set up the LabelVisibility attribute */
+	TOptional<EVisibility> LabelVisibilityOverride;
+
+	/** Controls the visibility of the of label, defaults to GetIconVisibility */
+	TAttribute<EVisibility> LabelVisibility;
 
-	TAttribute< EVisibility > LabelVisibility;
+	TSharedPtr<SBorder> ButtonBorder;
+	TSharedPtr<SBorder> OptionsBorder;
 
 	/** Whether ToolBar will have Focusable buttons. */
 	bool bIsFocusable;
diff --git a/Engine/Source/Runtime/Slate/Public/Framework/MultiBox/SToolBarComboButtonBlock.h b/Engine/Source/Runtime/Slate/Public/Framework/MultiBox/SToolBarComboButtonBlock.h
--- a/Engine/Source/Runtime/Slate/Public/Framework/MultiBox/SToolBarComboButtonBlock.h
+++ b/Engine/Source/Runtime/Slate/Public/Framework/MultiBox/SToolBarComboButtonBlock.h
@@ -10,6 +10,8 @@
 #include "Widgets/DeclarativeSyntaxSupport.h"
 #include "Framework/MultiBox/MultiBox.h"
 
+class SComboButton;
+
 /**
  * Tool bar combo button MultiBlock
  */
@@ -45,7 +47,7 @@ public:
 	/** Set whether this toolbar should always use small icons, regardless of the current settings */
 	void SetForceSmallIcons( const bool InForceSmallIcons ) { bForceSmallIcons = InForceSmallIcons; }
 
-
+	bool IsSimpleComboBox() const { return bSimpleComboBox; }
 private:
 
 	/**
@@ -150,21 +152,35 @@ private:
 	/** @return True if this toolbar button is using a dynamically set icon */
 	bool HasDynamicIcon() const;
 
+	/** Gets the icon brush for the toolbar block widget */
+	const FSlateBrush* GetIconBrush() const;
+
 	/** @return The icon for the toolbar button; may be dynamic, so check HasDynamicIcon */
-	const FSlateBrush* GetIcon() const;
+	const FSlateBrush* GetNormalIconBrush() const;
 
 	/** @return The small icon for the toolbar button; may be dynamic, so check HasDynamicIcon */
-	const FSlateBrush* GetSmallIcon() const;
+	const FSlateBrush* GetSmallIconBrush() const;
 
 	/** Called by Slate to determine whether icons/labels are visible */
 	EVisibility GetIconVisibility(bool bIsASmallIcon) const;
 
+	FSlateColor OnGetForegroundColor() const;
+
+private:
+	/** Overrides the visibility of the of label. This is used to set up the LabelVisibility attribute */
+	TOptional<EVisibility> LabelVisibilityOverride;
+
 	/** Controls the visibility of the of label, defaults to GetIconVisibility */
 	TAttribute< EVisibility > LabelVisibility;
 
 	/** Optional overridden icon for this tool bar button.  IF not set, then the action's icon will be used instead. */
 	TAttribute<FSlateIcon> Icon;
 
+	TSharedPtr<SComboButton> ComboButtonWidget;
+
+	/** The foreground color for button when the combo button is open */
+	FSlateColor OpenForegroundColor;
+
 	/** Whether this toolbar should always use small icons, regardless of the current settings */
 	bool bForceSmallIcons;
 };
diff --git a/Engine/Source/Runtime/SlateCore/Private/Styling/StarshipCoreStyle.cpp b/Engine/Source/Runtime/SlateCore/Private/Styling/StarshipCoreStyle.cpp
--- a/Engine/Source/Runtime/SlateCore/Private/Styling/StarshipCoreStyle.cpp
+++ b/Engine/Source/Runtime/SlateCore/Private/Styling/StarshipCoreStyle.cpp
@@ -1739,8 +1739,8 @@ void FStarshipCoreStyle::SetupMultiboxStyles(TSharedRef<FStyle>& Style)
 		FToolBarStyle NormalToolbarStyle =
 			FToolBarStyle()
 			.SetBackground(FSlateColorBrush(FStyleColors::Background))
+			.SetBackgroundPadding(FMargin(0, 4, 0, 4))
 			.SetExpandBrush(IMAGE_BRUSH("Icons/toolbar_expand_16x", Icon16x16))
-			.SetSubMenuIndicator(IMAGE_BRUSH("Common/SubmenuArrow", Icon8x8))
 			.SetComboButtonPadding(FMargin(4.0f, 0.0f))
 			.SetButtonPadding(FMargin(2.0f, 0.f))
 			.SetCheckBoxPadding(FMargin(4.0f, 0.f))
@@ -1777,15 +1777,13 @@ void FStarshipCoreStyle::SetupMultiboxStyles(TSharedRef<FStyle>& Style)
 			.SetDisabledForeground(FSlateColor::UseForeground());
 
 		NormalToolbarStyle.SetButtonStyle(ToolbarButton);
-
-		NormalToolbarStyle.SetSettingsComboButtonStyle(Style->GetWidgetStyle<FComboButtonStyle>("ComboButton"));
 		NormalToolbarStyle.SetIconSize(Icon40x40);
 
 		Style->Set("ToolBar", NormalToolbarStyle);
 
 		// Slim Toolbar
 
-		const FButtonStyle SlimToolBarButton = FButtonStyle(ToolbarButton)
+		FButtonStyle SlimToolBarButton = FButtonStyle(ToolbarButton)
 			.SetPressed(FSlateNoResource())
 			.SetHovered(FSlateNoResource())
 			.SetNormalForeground(FStyleColors::Foreground)
@@ -1801,19 +1799,19 @@ void FStarshipCoreStyle::SetupMultiboxStyles(TSharedRef<FStyle>& Style)
 		FToolBarStyle SlimToolbarStyle =
 			FToolBarStyle()
 			.SetBackground(FSlateColorBrush(FStyleColors::Background))
+			.SetBackgroundPadding(FMargin(0,4,0,4))
 			.SetExpandBrush(IMAGE_BRUSH("Icons/toolbar_expand_16x", Icon16x16))
-			.SetSubMenuIndicator(IMAGE_BRUSH("Common/SubmenuArrow", Icon8x8))
 			.SetComboButtonPadding(FMargin(6.0f, 0.0f))
 			.SetButtonPadding(FMargin(4.0f, 0.0f))
-			.SetCheckBoxPadding(FMargin(4.0f, 0.0f))
+			.SetCheckBoxPadding(FMargin(8.0f, 0.0f))
 			.SetSeparatorBrush(FSlateColorBrush(FStyleColors::Recessed))
-			.SetSeparatorPadding(FMargin(8.f, 0))
+			.SetSeparatorPadding(FMargin(8.f, -4.0f))
 			.SetLabelStyle(FTextBlockStyle(NormalText))
 			.SetComboButtonStyle(SlimToolBarComboButton)
-			.SetLabelPadding(FMargin(5, 9, 0, 9))
+			.SetLabelPadding(FMargin(5, 0, 0, 0))
 			.SetEditableTextStyle(FEditableTextBoxStyle(NormalEditableTextBoxStyle));
 
-		const FCheckBoxStyle SlimToolBarToggleButtonCheckBoxStyle = FCheckBoxStyle()
+		FCheckBoxStyle SlimToolBarToggleButtonCheckBoxStyle = FCheckBoxStyle()
 			.SetCheckBoxType(ESlateCheckBoxType::ToggleButton)
 			.SetUncheckedImage(FSlateNoResource())
 			.SetUncheckedPressedImage(FSlateNoResource())
@@ -1839,7 +1837,27 @@ void FStarshipCoreStyle::SetupMultiboxStyles(TSharedRef<FStyle>& Style)
 		SlimToolbarStyle.SetSettingsComboButtonStyle(SlimToolBarSettingsComboButton);
 		SlimToolbarStyle.SetIconSize(Icon20x20);
 
+		Style->Set("ToolbarSettingsRegion.LeftHover", new FSlateRoundedBoxBrush(FStyleColors::Dropdown, FVector4(4.0f, 0.0f, 0.0f, 4.0f)));
+		Style->Set("ToolbarSettingsRegion.RightHover", new FSlateRoundedBoxBrush(FStyleColors::Dropdown, FVector4(0.0f, 4.0f, 4.0f, 0.0f)));
+		Style->Set("ToolbarSettingsRegion.Left", new FSlateRoundedBoxBrush(FStyleColors::Header, FVector4(4.0f, 0.0f, 0.0f, 4.0f), FStyleColors::Header, 1.0f));
+		Style->Set("ToolbarSettingsRegion.Right", new FSlateRoundedBoxBrush(FStyleColors::Header, FVector4(0.0f, 4.0f, 4.0f, 0.0f), FStyleColors::Header, 1.0f));
+
 		Style->Set("SlimToolBar", SlimToolbarStyle);
+
+		SlimToolbarStyle.SetComboButtonPadding(FMargin(4.0f, 0.0f));
+		SlimToolbarStyle.SetButtonPadding(FMargin(4.0f, 4.0f));
+		SlimToolbarStyle.SetCheckBoxPadding(FMargin(8.0f, 0.0f));
+		SlimToolbarStyle.SetShowLabels(false);
+		SlimToolbarStyle.SetLabelPadding(FMargin(5, 0, 0, 0));
+	
+		Style->Set("AssetEditorToolbar", SlimToolbarStyle);
+
+		// Callout Toolbar - Used to "call out" the toolbar button with text
+		{
+			SlimToolbarStyle.SetShowLabels(true);
+
+			Style->Set("CalloutToolbar", SlimToolbarStyle);
+		}
 	}
 
 	// MenuBar
diff --git a/Engine/Source/Runtime/SlateCore/Private/Styling/ToolBarStyle.cpp b/Engine/Source/Runtime/SlateCore/Private/Styling/ToolBarStyle.cpp
--- a/Engine/Source/Runtime/SlateCore/Private/Styling/ToolBarStyle.cpp
+++ b/Engine/Source/Runtime/SlateCore/Private/Styling/ToolBarStyle.cpp
@@ -8,7 +8,6 @@ const FName FToolBarStyle::TypeName(TEXT("FToolbarStyle"));
 FToolBarStyle::FToolBarStyle()
 	: BackgroundBrush(FSlateNoResource())
 	, ExpandBrush(FSlateNoResource())
-	, SubMenuIndicator(FSlateNoResource())
 	, SeparatorBrush(FSlateNoResource())
 	, LabelStyle()
 	, EditableTextStyle()
@@ -22,7 +21,9 @@ FToolBarStyle::FToolBarStyle()
 	, CheckBoxPadding(0)
 	, BlockPadding(0)
 	, IndentedBlockPadding(0)
+	, BackgroundPadding(0)
 	, IconSize(16,16)
+	, bShowLabels(true)
 {}
 
 const FToolBarStyle& FToolBarStyle::GetDefault()
@@ -35,7 +36,6 @@ void FToolBarStyle::GetResources(TArray<const FSlateBrush*>& OutBrushes) const
 {
 	OutBrushes.Add(&BackgroundBrush);
 	OutBrushes.Add(&ExpandBrush);
-	OutBrushes.Add(&SubMenuIndicator);
 	OutBrushes.Add(&SeparatorBrush);
 
 	LabelStyle.GetResources(OutBrushes);
diff --git a/Engine/Source/Runtime/SlateCore/Public/Styling/ToolBarStyle.h b/Engine/Source/Runtime/SlateCore/Public/Styling/ToolBarStyle.h
--- a/Engine/Source/Runtime/SlateCore/Public/Styling/ToolBarStyle.h
+++ b/Engine/Source/Runtime/SlateCore/Public/Styling/ToolBarStyle.h
@@ -37,10 +37,6 @@ struct SLATECORE_API FToolBarStyle : public FSlateWidgetStyle
 	FSlateBrush ExpandBrush;
 	FToolBarStyle& SetExpandBrush(const FSlateBrush& InExpandBrush) { ExpandBrush = InExpandBrush; return *this; }
 
-	UPROPERTY(EditAnywhere, BlueprintReadWrite, Category = Appearance)
-	FSlateBrush SubMenuIndicator;
-	FToolBarStyle& SetSubMenuIndicator(const FSlateBrush& InSubMenuIndicator) { SubMenuIndicator = InSubMenuIndicator; return *this; }
-
 	UPROPERTY(EditAnywhere, BlueprintReadWrite, Category = Appearance)
 	FSlateBrush SeparatorBrush;
 	FToolBarStyle& SetSeparatorBrush(const FSlateBrush& InSeparatorBrush) { SeparatorBrush = InSeparatorBrush; return *this; }
@@ -97,7 +93,15 @@ struct SLATECORE_API FToolBarStyle : public FSlateWidgetStyle
 	FMargin IndentedBlockPadding;
 	FToolBarStyle& SetIndentedBlockPadding(const FMargin& InIndentedBlockPadding) { IndentedBlockPadding = InIndentedBlockPadding; return *this; }
 
+	UPROPERTY(EditAnywhere, BlueprintReadWrite, Category = Appearance)
+	FMargin BackgroundPadding;
+	FToolBarStyle& SetBackgroundPadding(const FMargin& InMargin) { BackgroundPadding = InMargin; return *this; }
+
 	UPROPERTY(EditAnywhere, BlueprintReadWrite, Category = Appearance)
 	FVector2D IconSize;
 	FToolBarStyle& SetIconSize(const FVector2D& InIconSize) { IconSize = InIconSize; return *this; }
+
+	UPROPERTY(EditAnywhere, BlueprintReadWrite, Category=Appearance)
+	bool bShowLabels;
+	FToolBarStyle& SetShowLabels(bool bInShowLabels) { bShowLabels = bInShowLabels; return *this; }
 };
