SSplitter now attempts to respect the MinSplitterChildLength and FSlot::MinSizeValue during arrange when possible to do so.

If there is enough space to do so, then slots that require more space to respect their clamp size will steal space from previous slots.

diff --git a/Engine/Source/Runtime/Slate/Private/Widgets/Layout/SSplitter.cpp b/Engine/Source/Runtime/Slate/Private/Widgets/Layout/SSplitter.cpp
--- a/Engine/Source/Runtime/Slate/Private/Widgets/Layout/SSplitter.cpp
+++ b/Engine/Source/Runtime/Slate/Private/Widgets/Layout/SSplitter.cpp
@@ -78,8 +78,8 @@ void SSplitter::Construct( const SSplitter::FArguments& InArgs )
 
 TArray<FLayoutGeometry> SSplitter::ArrangeChildrenForLayout(const FGeometry& AllottedGeometry) const
 {
-	TArray<FLayoutGeometry> Result;
-	Result.Empty(Children.Num());
+	TArray<FLayoutGeometry> ResultGeometries;
+	ResultGeometries.Empty(Children.Num());
 
 	const int32 AxisIndex = (Orientation == Orient_Horizontal) ? 0 : 1;
 
@@ -101,55 +101,87 @@ TArray<FLayoutGeometry> SSplitter::ArrangeChildrenForLayout(const FGeometry& All
 	int32 NumNonCollapsedChildren = 0;
 	int32 NumResizeableChildren = 0;
 	float CoefficientTotal = 0;
-	// Some space is claimed by non-resizeable elements (auto-sized elements)
-	float NonResizeableSpace = 0;
+	// Some space is claimed by non-resizable elements (auto-sized elements)
+	float NonResizableSpace = 0;
+	float MinResizableSpace = 0;
+
+	for (int32 ChildIndex = 0; ChildIndex < Children.Num(); ++ChildIndex)
 	{
-		for (int32 ChildIndex = 0; ChildIndex < Children.Num(); ++ChildIndex)
+		const FSlot& CurSlot = Children[ChildIndex];
+		if (CurSlot.GetWidget()->GetVisibility() != EVisibility::Collapsed)
 		{
-			if (Children[ChildIndex].GetWidget()->GetVisibility() != EVisibility::Collapsed)
-			{
-				++NumNonCollapsedChildren;
+			++NumNonCollapsedChildren;
 
-				if (Children[ChildIndex].SizingRule.Get() == SSplitter::SizeToContent)
-				{
-					NonResizeableSpace += Children[ChildIndex].GetWidget()->GetDesiredSize()[AxisIndex];
-				}
-				else // SizingRule == SSplitter::FractionOfParent
-				{
-					CoefficientTotal += Children[ChildIndex].SizeValue.Get();
-				}
+			if (CurSlot.SizingRule.Get() == SSplitter::SizeToContent)
+			{
+				NonResizableSpace += CurSlot.GetWidget()->GetDesiredSize()[AxisIndex];
+			}
+			else // SizingRule == SSplitter::FractionOfParent
+			{
+				MinResizableSpace += FMath::Max(MinSplitterChildLength, CurSlot.MinSizeValue.Get(0));
+				CoefficientTotal += CurSlot.SizeValue.Get();
 			}
 		}
 	}
 
 	// The user-sizable children must make room for the resize handles and for auto-sized children.
 	const float SpaceNeededForHandles = FMath::Max(0, NumNonCollapsedChildren - 1) * PhysicalSplitterHandleSize;
-	const float ResizeableSpace = AllottedGeometry.Size.Component(AxisIndex) - SpaceNeededForHandles - NonResizeableSpace;
+	const float ResizableSpace = AllottedGeometry.Size.Component(AxisIndex) - SpaceNeededForHandles - NonResizableSpace;
 
-	// Arrange the children horizontally or vertically.
-	float XOffset = 0;
+	TArray<float, TMemStackAllocator<>> SlotSizes;
+
+	// calculate slot sizes
 	for (int32 ChildIndex = 0; ChildIndex < Children.Num(); ++ChildIndex)
 	{
 		const FSlot& CurSlot = Children[ChildIndex];
 
-		const float ChildSpace = (CurSlot.SizingRule.Get() == SSplitter::SizeToContent)
+		float ChildSpace = (CurSlot.SizingRule.Get() == SSplitter::SizeToContent)
 			? CurSlot.GetWidget()->GetDesiredSize()[AxisIndex]
-			: ResizeableSpace * CurSlot.SizeValue.Get() / CoefficientTotal;
+			: ResizableSpace * CurSlot.SizeValue.Get() / CoefficientTotal;
 
-		const EVisibility ChildVisibility = CurSlot.GetWidget()->GetVisibility();
+		// we have enough space to clamp this and steal space from earlier slots
+		if (ResizableSpace >= MinResizableSpace)
+		{
+			const float ClampedSpace = ClampChild(CurSlot, ChildSpace);
+			float CurrentRequiredSpace = ClampedSpace - ChildSpace;
 
-		const FVector2D ChildOffset = Orientation == Orient_Horizontal ? FVector2D(XOffset, 0) : FVector2D(0, XOffset);
-		const FVector2D ChildSize = Orientation == Orient_Horizontal ? FVector2D(ChildSpace, AllottedGeometry.Size.Y) : FVector2D(AllottedGeometry.Size.X, ChildSpace);
-		Result.Emplace(FSlateLayoutTransform(ChildOffset), ChildSize);
+			for (int32 PrevIndex = ChildIndex - 1; PrevIndex >= 0 && CurrentRequiredSpace > 0; --PrevIndex)
+			{
+				if (Children[PrevIndex].SizingRule.Get() == SSplitter::FractionOfParent)
+				{
+					const float MinChildSize = FMath::Max(MinSplitterChildLength, Children[PrevIndex].MinSizeValue.Get(0));
+					const float AvailableSpace = SlotSizes[PrevIndex] - MinChildSize;
+					if (AvailableSpace > 0)
+					{
+						const float SpaceToSteal = FMath::Min(AvailableSpace, CurrentRequiredSpace);
+						SlotSizes[PrevIndex] -= SpaceToSteal;
+						CurrentRequiredSpace -= SpaceToSteal;
+					}
+				}
+			}
+
+			ChildSpace = ClampedSpace;
+		}
+
+		SlotSizes.Add(ChildSpace);
+	}
+
+	// create geometries and calculate offsets
+	float Offset = 0;
+	for (int32 ChildIndex = 0; ChildIndex < Children.Num(); ++ChildIndex)
+	{
+		const FVector2D ChildOffset = Orientation == Orient_Horizontal ? FVector2D(Offset, 0) : FVector2D(0, Offset);
+		const FVector2D ChildSize = Orientation == Orient_Horizontal ? FVector2D(SlotSizes[ChildIndex], AllottedGeometry.Size.Y) : FVector2D(AllottedGeometry.Size.X, SlotSizes[ChildIndex]);
+		ResultGeometries.Emplace(FSlateLayoutTransform(ChildOffset), ChildSize);
 
 		// Advance to the next slot. If the child is collapsed, it takes up no room and does not need a splitter
-		if (ChildVisibility != EVisibility::Collapsed)
+		if (Children[ChildIndex].GetWidget()->GetVisibility() != EVisibility::Collapsed)
 		{
-			XOffset += FMath::RoundToInt(ChildSpace + PhysicalSplitterHandleSize);
+			Offset += FMath::RoundToInt(SlotSizes[ChildIndex] + PhysicalSplitterHandleSize);
 		}
 	}
 
-	return Result;
+	return ResultGeometries;
 }
 
 /**
@@ -576,7 +608,7 @@ void SSplitter::HandleResizingDelta(EOrientation SplitterOrientation, const floa
 	}
 
 	const int32 NumSlotsAfterDragHandle = SlotsAfterDragHandleIndices.Num();
-	if ( NumSlotsAfterDragHandle > 0 )
+	if (NumSlotsAfterDragHandle > 0)
 	{
 		struct FSlotInfo 
 		{
@@ -712,7 +744,7 @@ void SSplitter::HandleResizingByMousePosition(EOrientation SplitterOrientation,
 *
 * @return A size that is clamped against the minimum size allowed for children.
 */
-float SSplitter::ClampChild(const FSlot& ChildSlot, float ProposedSize)
+float SSplitter::ClampChild(const FSlot& ChildSlot, float ProposedSize) const
 {
 	float ClampedSize = FMath::Max(MinSplitterChildLength, ProposedSize);
 
diff --git a/Engine/Source/Runtime/Slate/Public/Widgets/Layout/SSplitter.h b/Engine/Source/Runtime/Slate/Public/Widgets/Layout/SSplitter.h
--- a/Engine/Source/Runtime/Slate/Public/Widgets/Layout/SSplitter.h
+++ b/Engine/Source/Runtime/Slate/Public/Widgets/Layout/SSplitter.h
@@ -310,7 +310,7 @@ protected:
 	 *
 	 * @return A size that is clamped against the minimum size allowed for children.
 	 */
-	float ClampChild(const FSlot& ChildSlot, float ProposedSize);
+	float ClampChild(const FSlot& ChildSlot, float ProposedSize) const;
 
 	/**
 	 * Given a mouse position within the splitter, figure out which resize handle we are hovering (if any).
