Fix overlapping labels in numeric edit boxes

diff --git a/Engine/Source/Editor/DetailCustomizations/Private/MathStructProxyCustomizations.cpp b/Engine/Source/Editor/DetailCustomizations/Private/MathStructProxyCustomizations.cpp
--- a/Engine/Source/Editor/DetailCustomizations/Private/MathStructProxyCustomizations.cpp
+++ b/Engine/Source/Editor/DetailCustomizations/Private/MathStructProxyCustomizations.cpp
@@ -46,6 +46,7 @@ TSharedRef<SWidget> FMathStructProxyCustomization::MakeNumericProxyWidget(TShare
 		.MinSliderValue(bRotationInDegrees ? 0.0f : TOptional<NumericType>())
 		.LabelPadding(FMargin(3))
 		.ToolTipText(this, &FMathStructProxyCustomization::OnGetValueToolTip<ProxyType, NumericType>, WeakHandlePtr, ProxyValue, Label)
+		.LabelLocation(SNumericEntryBox<NumericType>::ELabelLocation::Inside)
 		.Label()
 		[
 			SNumericEntryBox<NumericType>::BuildNarrowColorLabel(LabelBackgroundColor)
diff --git a/Engine/Source/Runtime/Slate/Private/Widgets/Input/SRotatorInputBox.cpp b/Engine/Source/Runtime/Slate/Private/Widgets/Input/SRotatorInputBox.cpp
--- a/Engine/Source/Runtime/Slate/Private/Widgets/Input/SRotatorInputBox.cpp
+++ b/Engine/Source/Runtime/Slate/Private/Widgets/Input/SRotatorInputBox.cpp
@@ -21,6 +21,7 @@ void SRotatorInputBox::Construct( const SRotatorInputBox::FArguments& InArgs )
 			.MinSliderValue(0.0f)
 			.MaxSliderValue(359.999f)
 			.LabelPadding(FMargin(3))
+			.LabelLocation(SNumericEntryBox<float>::ELabelLocation::Inside)
 			.Label()
 			[
 				InArgs._bColorAxisLabels ? SNumericEntryBox<float>::BuildNarrowColorLabel(SNumericEntryBox<float>::RedLabelBackgroundColor) : SNullWidget::NullWidget
@@ -48,6 +49,7 @@ void SRotatorInputBox::Construct( const SRotatorInputBox::FArguments& InArgs )
 			.MinSliderValue(0.0f)
 			.MaxSliderValue(359.999f)
 			.LabelPadding(FMargin(3))
+			.LabelLocation(SNumericEntryBox<float>::ELabelLocation::Inside)
 			.Label()
 			[
 				InArgs._bColorAxisLabels ? SNumericEntryBox<float>::BuildNarrowColorLabel(SNumericEntryBox<float>::GreenLabelBackgroundColor) : SNullWidget::NullWidget
@@ -74,8 +76,8 @@ void SRotatorInputBox::Construct( const SRotatorInputBox::FArguments& InArgs )
 			.AllowSpin(InArgs._AllowSpin)
 			.MinSliderValue(0.0f)
 			.MaxSliderValue(359.999f)
-			.LabelPadding(0)
 			.LabelPadding(FMargin(3))
+			.LabelLocation(SNumericEntryBox<float>::ELabelLocation::Inside)
 			.Label()
 			[
 				InArgs._bColorAxisLabels ? SNumericEntryBox<float>::BuildNarrowColorLabel(SNumericEntryBox<float>::BlueLabelBackgroundColor) : SNullWidget::NullWidget
diff --git a/Engine/Source/Runtime/Slate/Private/Widgets/Input/SVectorInputBox.cpp b/Engine/Source/Runtime/Slate/Private/Widgets/Input/SVectorInputBox.cpp
--- a/Engine/Source/Runtime/Slate/Private/Widgets/Input/SVectorInputBox.cpp
+++ b/Engine/Source/Runtime/Slate/Private/Widgets/Input/SVectorInputBox.cpp
@@ -62,6 +62,7 @@ void SVectorInputBox::ConstructX(const FArguments& InArgs, TSharedRef<SHorizonta
 		.OnBeginSliderMovement(InArgs._OnBeginSliderMovement)
 		.OnEndSliderMovement(InArgs._OnEndSliderMovement)
 		.LabelPadding(FMargin(3))
+		.LabelLocation(SNumericEntryBox<float>::ELabelLocation::Inside)
 		.Label()
 		[
 			LabelWidget
@@ -109,6 +110,7 @@ void SVectorInputBox::ConstructY(const FArguments& InArgs, TSharedRef<SHorizonta
 		.OnBeginSliderMovement(InArgs._OnBeginSliderMovement)
 		.OnEndSliderMovement(InArgs._OnEndSliderMovement)
 		.LabelPadding(FMargin(3))
+		.LabelLocation(SNumericEntryBox<float>::ELabelLocation::Inside)
 		.Label()
 		[
 			LabelWidget
@@ -156,6 +158,7 @@ void SVectorInputBox::ConstructZ(const FArguments& InArgs, TSharedRef<SHorizonta
 		.OnBeginSliderMovement(InArgs._OnBeginSliderMovement)
 		.OnEndSliderMovement(InArgs._OnEndSliderMovement)
 		.LabelPadding(FMargin(3))
+		.LabelLocation(SNumericEntryBox<float>::ELabelLocation::Inside)
 		.Label()
 		[
 			LabelWidget
diff --git a/Engine/Source/Runtime/Slate/Public/Widgets/Input/SNumericEntryBox.h b/Engine/Source/Runtime/Slate/Public/Widgets/Input/SNumericEntryBox.h
--- a/Engine/Source/Runtime/Slate/Public/Widgets/Input/SNumericEntryBox.h
+++ b/Engine/Source/Runtime/Slate/Public/Widgets/Input/SNumericEntryBox.h
@@ -56,6 +56,14 @@ public:
 	/** Notification when the max/min spinner values are changed (only apply if SupportDynamicSliderMaxValue or SupportDynamicSliderMinValue are true) */
 	DECLARE_DELEGATE_FourParams(FOnDynamicSliderMinMaxValueChanged, NumericType, TWeakPtr<SWidget>, bool, bool);
 
+	enum class ELabelLocation
+	{
+		// Outside the bounds of the editable area of this box. Usually preferred for text based labels
+		Outside,
+		// Inside the bounds of the editable area of this box. Usually preferred for non-text based labels
+		// when a spin box is used the label will appear on top of the spin box in this case
+		Inside
+	};
 public:
 
 	SLATE_BEGIN_ARGS( SNumericEntryBox<NumericType> )
@@ -63,7 +71,8 @@ public:
 		, _SpinBoxStyle(&FAppStyle::Get().GetWidgetStyle<FSpinBoxStyle>("NumericEntrySpinBox") )
 		, _Label()
 		, _LabelVAlign( VAlign_Fill )
-		, _LabelPadding( FMargin(3,0) )
+		, _LabelLocation(ELabelLocation::Outside)
+		, _LabelPadding(FMargin(3,0) )
 		, _BorderForegroundColor(FAppStyle::Get().GetWidgetStyle<FSpinBoxStyle>("NumericEntrySpinBox").ForegroundColor)
 		, _BorderBackgroundColor(FLinearColor::White)
 		, _UndeterminedString( SNumericEntryBox<NumericType>::DefaultUndeterminedString )
@@ -90,6 +99,7 @@ public:
 		SLATE_NAMED_SLOT( FArguments, Label )
 		/** Vertical alignment of the label content */
 		SLATE_ARGUMENT( EVerticalAlignment, LabelVAlign )
+		SLATE_ARGUMENT(ELabelLocation, LabelLocation)
 		/** Padding around the label content */
 		SLATE_ARGUMENT( FMargin, LabelPadding )
 		/** Border Foreground Color */
@@ -239,28 +249,57 @@ public:
 			];
 
 
-		if (InArgs._Label.Widget != SNullWidget::NullWidget)
+		if (InArgs._Label.Widget == SNullWidget::NullWidget || InArgs._LabelLocation == ELabelLocation::Inside)
 		{
-			Overlay->AddSlot()
+			if(InArgs._Label.Widget != SNullWidget::NullWidget)
+			{
+				Overlay->AddSlot()
+					.HAlign(HAlign_Left)
+					.VAlign(InArgs._LabelVAlign)
+					.Padding(InArgs._LabelPadding)
+					[
+						InArgs._Label.Widget
+					];
+			}
+
+			ChildSlot
+			[
+				SNew(SBorder)
+				.BorderImage(this, &SNumericEntryBox<NumericType>::GetBorderImage)
+				.BorderBackgroundColor(InArgs._BorderBackgroundColor)
+				.ForegroundColor(InArgs._BorderForegroundColor)
+				.Padding(0)
+				[
+					Overlay
+				]
+			];
+		}
+		else
+		{
+			ChildSlot
+			[
+				SNew(SHorizontalBox)
+				+SHorizontalBox::Slot()
+				.AutoWidth()
 				.HAlign(HAlign_Left)
 				.VAlign(InArgs._LabelVAlign)
 				.Padding(InArgs._LabelPadding)
 				[
 					InArgs._Label.Widget
-				];
-		}
-
-		ChildSlot
-			[
-				SNew( SBorder )
-				.BorderImage( this, &SNumericEntryBox<NumericType>::GetBorderImage )
-				.BorderBackgroundColor( InArgs._BorderBackgroundColor )
-				.ForegroundColor( InArgs._BorderForegroundColor )
-				.Padding(0)
+				]
+				+ SHorizontalBox::Slot()
 				[
-					Overlay
+					SNew(SBorder)
+					.BorderImage(this, &SNumericEntryBox<NumericType>::GetBorderImage)
+					.BorderBackgroundColor(InArgs._BorderBackgroundColor)
+					.ForegroundColor(InArgs._BorderForegroundColor)
+					.Padding(0)
+					[
+						Overlay
+					]
 				]
 			];
+		}
 	}
 
 	static TSharedRef<SWidget> BuildLabel(TAttribute<FText> LabelText, const FSlateColor& ForegroundColor, const FSlateColor& BackgroundColor)
