Allow an options button for a button on a toolbar to be able to independently control visibility

diff --git a/Engine/Source/Runtime/Slate/Private/Framework/MultiBox/SToolBarButtonBlock.cpp b/Engine/Source/Runtime/Slate/Private/Framework/MultiBox/SToolBarButtonBlock.cpp
--- a/Engine/Source/Runtime/Slate/Private/Framework/MultiBox/SToolBarButtonBlock.cpp
+++ b/Engine/Source/Runtime/Slate/Private/Framework/MultiBox/SToolBarButtonBlock.cpp
@@ -264,6 +264,7 @@ void SToolBarButtonBlock::BuildMultiBlockWidget(const ISlateStyle* StyleSet, con
 			.ContentPadding(0)
 			// Use the tool bar item style for this button
 			.ButtonStyle(&ButtonStyle)
+			.IsEnabled(this, &SToolBarButtonBlock::IsEnabled)
 			.OnClicked(this, &SToolBarButtonBlock::OnClicked)
 			.ToolTip(FMultiBoxSettings::ToolTipConstructor.Execute(ActualToolTip, nullptr, Action.Pin()))
 			.IsFocusable(bIsFocusable)
@@ -288,6 +289,7 @@ void SToolBarButtonBlock::BuildMultiBlockWidget(const ISlateStyle* StyleSet, con
 			.ToolTip( FMultiBoxSettings::ToolTipConstructor.Execute( ActualToolTip, nullptr, Action.Pin()))		
 			.OnCheckStateChanged(this, &SToolBarButtonBlock::OnCheckStateChanged )
 			.IsChecked(this, &SToolBarButtonBlock::OnIsChecked)
+			.IsEnabled(this, &SToolBarButtonBlock::IsEnabled)
 			.Padding(ToolBarStyle.CheckBoxPadding)
 			[
 				ButtonContent
@@ -311,14 +313,6 @@ void SToolBarButtonBlock::BuildMultiBlockWidget(const ISlateStyle* StyleSet, con
 					ChildSlot.GetWidget()
 				]
 			]
-			/*+ SHorizontalBox::Slot()
-			.AutoWidth()
-			[
-				SNew(SSeparator)
-				.Orientation(Orient_Vertical)
-				.Visibility(this, &SToolBarButtonBlock::GetOptionsSeparatorVisibility)
-				.Thickness(1.0f)
-			]*/
 			+ SHorizontalBox::Slot()
 			.Padding(0.0f,0.0f)
 			.AutoWidth()
@@ -339,9 +333,6 @@ void SToolBarButtonBlock::BuildMultiBlockWidget(const ISlateStyle* StyleSet, con
 		ChildSlot.Padding(ToolBarStyle.ButtonPadding);
 	}
 
-	// Bind our widget's enabled state to whether or not our action can execute
-	SetEnabled( TAttribute< bool >( this, &SToolBarButtonBlock::IsEnabled ) );
-
 	// Bind our widget's visible state to whether or not the button should be visible
 	SetVisibility(TAttribute<EVisibility>(this, &SToolBarButtonBlock::GetBlockVisibility));
 }
