Slate: Temporarily remove stack allocator in SSplitter to fix exception in MemStack.h(108)

diff --git a/Engine/Source/Runtime/Slate/Private/Widgets/Layout/SSplitter.cpp b/Engine/Source/Runtime/Slate/Private/Widgets/Layout/SSplitter.cpp
--- a/Engine/Source/Runtime/Slate/Private/Widgets/Layout/SSplitter.cpp
+++ b/Engine/Source/Runtime/Slate/Private/Widgets/Layout/SSplitter.cpp
@@ -128,7 +128,7 @@ TArray<FLayoutGeometry> SSplitter::ArrangeChildrenForLayout(const FGeometry& All
 	const float SpaceNeededForHandles = FMath::Max(0, NumNonCollapsedChildren - 1) * PhysicalSplitterHandleSize;
 	const float ResizableSpace = AllottedGeometry.Size.Component(AxisIndex) - SpaceNeededForHandles - NonResizableSpace;
 
-	TArray<float, TMemStackAllocator<>> SlotSizes;
+	TArray<float> SlotSizes;
 
 	// calculate slot sizes
 	for (int32 ChildIndex = 0; ChildIndex < Children.Num(); ++ChildIndex)
