Allows styles to specify default Background Brushes for ListViews, TreeViews

diff --git a/Engine/Source/Editor/EditorStyle/Private/StarshipStyle.cpp b/Engine/Source/Editor/EditorStyle/Private/StarshipStyle.cpp
--- a/Engine/Source/Editor/EditorStyle/Private/StarshipStyle.cpp
+++ b/Engine/Source/Editor/EditorStyle/Private/StarshipStyle.cpp
@@ -778,6 +778,15 @@ void FStarshipEditorStyle::FStyle::SetupGeneralStyles()
 
 	// Lists, Trees
 	{
+
+		const FTableViewStyle DefaultTreeViewStyle = FTableViewStyle()
+			.SetBackgroundBrush(FSlateColorBrush(FStyleColors::Recessed));
+		Set("ListView", DefaultTreeViewStyle);
+
+		const FTableViewStyle DefaultTableViewStyle = FTableViewStyle()
+			.SetBackgroundBrush(FSlateColorBrush(FStyleColors::Recessed));
+		Set("TreeView", DefaultTableViewStyle);
+
 		Set( "TableView.Row", FTableRowStyle( NormalTableRowStyle) );
 		Set( "TableView.DarkRow",FTableRowStyle( NormalTableRowStyle)
 			.SetEvenRowBackgroundBrush(IMAGE_BRUSH("PropertyView/DetailCategoryMiddle", FVector2D(16, 16)))
diff --git a/Engine/Source/Editor/SceneOutliner/Private/SSceneOutliner.cpp b/Engine/Source/Editor/SceneOutliner/Private/SSceneOutliner.cpp
--- a/Engine/Source/Editor/SceneOutliner/Private/SSceneOutliner.cpp
+++ b/Engine/Source/Editor/SceneOutliner/Private/SSceneOutliner.cpp
@@ -165,24 +165,9 @@ void SSceneOutliner::Construct(const FArguments& InArgs, const FSceneOutlinerIni
 		Toolbar
 	];
 
-	TSharedRef<SVerticalBox> InnerBox = SNew(SVerticalBox);
-
 	VerticalBox->AddSlot()
-	.FillHeight( 1.0f )
-	.Padding(4.0, 4.0)
-	[
-		SNew(SBorder)
-		.BorderImage( FAppStyle::Get().GetBrush("Brushes.Recessed") )
-		.Padding(0.0)
-		[
-			InnerBox
-		]
-	];	
-
-	InnerBox->AddSlot()
 	.FillHeight(1.0)
 	[
-
 		SNew( SOverlay )
 		+SOverlay::Slot()
 		.HAlign( HAlign_Center )
@@ -244,7 +229,7 @@ void SSceneOutliner::Construct(const FArguments& InArgs, const FSceneOutlinerIni
 	// Bottom panel status bar, if enabled by the mode
 	if (Mode->ShowStatusBar())
 	{
-		InnerBox->AddSlot()
+		VerticalBox->AddSlot()
 		.AutoHeight()
 		[
 			SNew(SBorder)
diff --git a/Engine/Source/Runtime/Slate/Private/Widgets/Views/STableViewBase.cpp b/Engine/Source/Runtime/Slate/Private/Widgets/Views/STableViewBase.cpp
--- a/Engine/Source/Runtime/Slate/Private/Widgets/Views/STableViewBase.cpp
+++ b/Engine/Source/Runtime/Slate/Private/Widgets/Views/STableViewBase.cpp
@@ -689,9 +689,33 @@ bool STableViewBase::IsPendingRefresh() const
 	return bItemsNeedRefresh || ItemsPanel->IsRefreshPending();
 }
 
+bool STableViewBase::ComputeVolatility() const
+{
+	return BackgroundBrush.IsBound();
+}
+
 int32 STableViewBase::OnPaint( const FPaintArgs& Args, const FGeometry& AllottedGeometry, const FSlateRect& MyCullingRect, FSlateWindowElementList& OutDrawElements, int32 LayerId, const FWidgetStyle& InWidgetStyle, bool bParentEnabled ) const
 {
-	int32 NewLayerId = SCompoundWidget::OnPaint( Args, AllottedGeometry, MyCullingRect, OutDrawElements, LayerId, InWidgetStyle, bParentEnabled );
+	int32 NewLayerId = LayerId;
+
+	const FSlateBrush* BackgroundBrushResource = BackgroundBrush.Get();
+	if ( BackgroundBrushResource && BackgroundBrushResource->DrawAs != ESlateBrushDrawType::NoDrawType )
+	{
+
+		const bool bIsEnabled = ShouldBeEnabled(bParentEnabled);
+		const ESlateDrawEffect DrawEffects = bIsEnabled ? ESlateDrawEffect::None : ESlateDrawEffect::DisabledEffect;
+
+		FSlateDrawElement::MakeBox(
+				OutDrawElements,
+				++NewLayerId,
+				AllottedGeometry.ToPaintGeometry(),
+				BackgroundBrushResource,
+				DrawEffects,
+				BackgroundBrushResource->GetTint(InWidgetStyle) * InWidgetStyle.GetColorAndOpacityTint() 
+		);
+	}
+
+	NewLayerId = SCompoundWidget::OnPaint( Args, AllottedGeometry, MyCullingRect, OutDrawElements, NewLayerId, InWidgetStyle, bParentEnabled );
 
 	if( !bShowSoftwareCursor )
 	{
@@ -721,6 +745,7 @@ STableViewBase::STableViewBase( ETableViewMode::Type InTableViewMode )
 	, SoftwareCursorPosition( ForceInitToZero )
 	, bShowSoftwareCursor( false )
 	, WheelScrollMultiplier(GetGlobalScrollAmount())
+	, BackgroundBrush(FStyleDefaults::GetNoBrush())
 	, bIsScrollingActiveTimerRegistered( false )
 	, Overscroll()
 	, AllowOverscroll(EAllowOverscroll::Yes)
@@ -831,6 +856,11 @@ void STableViewBase::SetWheelScrollMultiplier(float NewWheelScrollMultiplier)
 	WheelScrollMultiplier = NewWheelScrollMultiplier;
 }
 
+void STableViewBase::SetBackgroundBrush(const TAttribute<const FSlateBrush*>& InBackgroundBrush)
+{
+	BackgroundBrush.SetImage(*this, InBackgroundBrush);
+}
+
 void STableViewBase::InsertWidget( const TSharedRef<ITableRow> & WidgetToInset )
 {
 	ItemsPanel->AddSlot(0)
diff --git a/Engine/Source/Runtime/Slate/Public/Widgets/Views/SListView.h b/Engine/Source/Runtime/Slate/Public/Widgets/Views/SListView.h
--- a/Engine/Source/Runtime/Slate/Public/Widgets/Views/SListView.h
+++ b/Engine/Source/Runtime/Slate/Public/Widgets/Views/SListView.h
@@ -9,6 +9,7 @@
 #include "Widgets/DeclarativeSyntaxSupport.h"
 #include "Widgets/SWidget.h"
 #include "Styling/SlateTypes.h"
+#include "Styling/AppStyle.h"
 #include "Framework/SlateDelegates.h"
 #include "Widgets/Text/STextBlock.h"
 #include "Framework/Views/ITypedTableView.h"
@@ -78,7 +79,8 @@ public:
 
 public:
 	SLATE_BEGIN_ARGS(SListView<ItemType>)
-		: _OnGenerateRow()
+		: _ListViewStyle(&FAppStyle::Get().GetWidgetStyle<FTableViewStyle>("ListView"))
+		, _OnGenerateRow()
 		, _OnEntryInitialized()
 		, _OnRowReleased()
 		, _ListItemsSource()
@@ -108,6 +110,8 @@ public:
 			this->_Clipping = EWidgetClipping::ClipToBounds;
 		}
 
+		SLATE_STYLE_ARGUMENT( FTableViewStyle, ListViewStyle )
+
 		SLATE_EVENT( FOnGenerateRow, OnGenerateRow )
 		
 		SLATE_EVENT( FOnEntryInitialized, OnEntryInitialized )
@@ -222,6 +226,8 @@ public:
 
 		this->OnKeyDownHandler = InArgs._OnKeyDownHandler;
 
+		this->SetStyle(InArgs._ListViewStyle);
+
 		// Check for any parameters that the coder forgot to specify.
 		FString ErrorString;
 		{
@@ -1091,6 +1097,12 @@ public:
 
 public:	
 
+	void SetStyle(const FTableViewStyle* InStyle)
+	{
+		Style = InStyle;
+		SetBackgroundBrush( Style != nullptr ? &Style->BackgroundBrush : FStyleDefaults::GetNoBrush() );
+	}
+
 	/** Sets the OnEntryInitializer delegate. This delegate is invoked after initializing an entry being generated, before it may be added to the actual widget hierarchy. */
 	void SetOnEntryInitialized(const FOnEntryInitialized& Delegate)
 	{
@@ -2174,6 +2186,9 @@ protected:
 	/** If true, the item currently slated to be scrolled into view will also be navigated to after being scrolled in */
 	bool bNavigateOnScrollIntoView = false;
 
+	/** Style resource for the list */
+	const FTableViewStyle* Style;
+
 private:
 	struct FGenerationPassGuard
 	{
diff --git a/Engine/Source/Runtime/Slate/Public/Widgets/Views/STableViewBase.h b/Engine/Source/Runtime/Slate/Public/Widgets/Views/STableViewBase.h
--- a/Engine/Source/Runtime/Slate/Public/Widgets/Views/STableViewBase.h
+++ b/Engine/Source/Runtime/Slate/Public/Widgets/Views/STableViewBase.h
@@ -13,6 +13,7 @@
 #include "Framework/Views/ITypedTableView.h"
 #include "Framework/Layout/InertialScrollManager.h"
 #include "Framework/Layout/Overscroll.h"
+#include "Styling/SlateTypes.h"
 
 #include "STableViewBase.generated.h"
 
@@ -163,6 +164,9 @@ public:
 	/** Sets the multiplier applied when wheel scrolling. Higher numbers will cover more distance per click of the wheel. */
 	void SetWheelScrollMultiplier(float NewWheelScrollMultiplier);
 
+	/** Sets the Background Brush */
+	void SetBackgroundBrush(const TAttribute<const FSlateBrush*>& InBackgroundBrush);
+
 public:
 
 	// SWidget interface
@@ -182,6 +186,7 @@ public:
 	virtual FReply OnKeyDown( const FGeometry& MyGeometry, const FKeyEvent& InKeyEvent ) override;
 	virtual FCursorReply OnCursorQuery( const FGeometry& MyGeometry, const FPointerEvent& CursorEvent ) const override;
 	virtual int32 OnPaint(const FPaintArgs& Args, const FGeometry& AllottedGeometry, const FSlateRect& MyCullingRect, FSlateWindowElementList& OutDrawElements, int32 LayerId, const FWidgetStyle& InWidgetStyle, bool bParentEnabled) const override;
+	virtual bool ComputeVolatility() const override;
 	virtual FReply OnTouchStarted( const FGeometry& MyGeometry, const FPointerEvent& InTouchEvent ) override;
 	virtual FReply OnTouchMoved( const FGeometry& MyGeometry, const FPointerEvent& InTouchEvent ) override;
 	virtual FReply OnTouchEnded( const FGeometry& MyGeometry, const FPointerEvent& InTouchEvent ) override;
@@ -412,6 +417,9 @@ protected:
 	/** Passing over the clipping to SListPanel */
 	virtual void OnClippingChanged() override;
 
+	/** Brush resource representing the background area of the view */
+	FInvalidatableBrushAttribute BackgroundBrush;
+
 protected:
 
 	/** Check whether the current state of the table warrants inertial scroll by the specified amount */
diff --git a/Engine/Source/Runtime/Slate/Public/Widgets/Views/STreeView.h b/Engine/Source/Runtime/Slate/Public/Widgets/Views/STreeView.h
--- a/Engine/Source/Runtime/Slate/Public/Widgets/Views/STreeView.h
+++ b/Engine/Source/Runtime/Slate/Public/Widgets/Views/STreeView.h
@@ -112,7 +112,8 @@ public:
 public:
 	
 	SLATE_BEGIN_ARGS( STreeView<ItemType> )
-		: _OnGenerateRow()
+		: _TreeViewStyle(&FAppStyle::Get().GetWidgetStyle<FTableViewStyle>("TreeView"))
+		, _OnGenerateRow()
 		, _OnGetChildren()
 		, _OnSetExpansionRecursive()
 		, _TreeItemsSource( static_cast< const TArray<ItemType>* >(nullptr) ) //@todo Slate Syntax: Initializing from nullptr without a cast
@@ -142,6 +143,8 @@ public:
 			this->_Clipping = EWidgetClipping::ClipToBounds;
 		}
 
+		SLATE_STYLE_ARGUMENT( FTableViewStyle, TreeViewStyle )
+
 		SLATE_EVENT( FOnGenerateRow, OnGenerateRow )
 
 		SLATE_EVENT( FOnWidgetToBeRemoved, OnRowReleased )
@@ -253,6 +256,8 @@ public:
 
 		this->bReturnFocusToSelection = InArgs._ReturnFocusToSelection;
 
+		this->SetStyle(InArgs._TreeViewStyle);
+
 		// Check for any parameters that the coder forgot to specify.
 		FString ErrorString;
 		{
@@ -712,6 +717,12 @@ public:
 		SListView<ItemType>::RebuildList();
 	}
 
+	void SetStyle(const FTableViewStyle* InStyle)
+	{
+		Style = InStyle;
+		STableViewBase::SetBackgroundBrush( Style != nullptr ? &Style->BackgroundBrush : FStyleDefaults::GetNoBrush() );
+	}
+
 	/**
 	 * Set whether some data item is expanded or not.
 	 * 
@@ -807,6 +818,9 @@ protected:
 	/** The delegate that is invoked whenever an item in the tree is expanded or collapsed. */
 	FOnExpansionChanged OnExpansionChanged;
 
+	/** Style resource for the tree */
+	const FTableViewStyle* Style;
+
 private:		
 
 	/** true when the LinearizedItems need to be regenerated. */
diff --git a/Engine/Source/Runtime/SlateCore/Private/Styling/SlateTypes.cpp b/Engine/Source/Runtime/SlateCore/Private/Styling/SlateTypes.cpp
--- a/Engine/Source/Runtime/SlateCore/Private/Styling/SlateTypes.cpp
+++ b/Engine/Source/Runtime/SlateCore/Private/Styling/SlateTypes.cpp
@@ -513,6 +513,23 @@ const FSplitterStyle& FSplitterStyle::GetDefault()
 	return Default;
 }
 
+FTableViewStyle::FTableViewStyle()
+	: BackgroundBrush(FSlateNoResource())
+{
+}
+
+void FTableViewStyle::GetResources( TArray< const FSlateBrush* >& OutBrushes ) const
+{
+	OutBrushes.Add( &BackgroundBrush );	
+}
+
+const FName FTableViewStyle::TypeName( TEXT("FTableViewStyle") );
+
+const FTableViewStyle& FTableViewStyle::GetDefault()
+{
+	static FTableViewStyle Default;
+	return Default;
+}
 
 FTableRowStyle::FTableRowStyle()
 	: TextColor(FSlateColor::UseForeground())
diff --git a/Engine/Source/Runtime/SlateCore/Private/Styling/StarshipCoreStyle.cpp b/Engine/Source/Runtime/SlateCore/Private/Styling/StarshipCoreStyle.cpp
--- a/Engine/Source/Runtime/SlateCore/Private/Styling/StarshipCoreStyle.cpp
+++ b/Engine/Source/Runtime/SlateCore/Private/Styling/StarshipCoreStyle.cpp
@@ -1638,6 +1638,12 @@ void FStarshipCoreStyle::SetupTableViewStyles(TSharedRef<FStyle>& Style)
 	const FSlateColor SelectionColor = Style->GetSlateColor("SelectionColor");
 	const FSlateColor SelectorColor = Style->GetSlateColor("SelectorColor");
 
+	const FTableViewStyle DefaultTreeViewStyle = FTableViewStyle();
+	Style->Set("ListView", DefaultTreeViewStyle);
+
+	const FTableViewStyle DefaultTableViewStyle = FTableViewStyle();
+	Style->Set("TreeView", DefaultTableViewStyle);
+
 	const FTableRowStyle DefaultTableRowStyle = FTableRowStyle()
 		.SetEvenRowBackgroundBrush(FSlateColorBrush(FStyleColors::Recessed))
 		.SetEvenRowBackgroundHoveredBrush(FSlateColorBrush(FStyleColors::SelectHover))
diff --git a/Engine/Source/Runtime/SlateCore/Public/Styling/SlateTypes.h b/Engine/Source/Runtime/SlateCore/Public/Styling/SlateTypes.h
--- a/Engine/Source/Runtime/SlateCore/Public/Styling/SlateTypes.h
+++ b/Engine/Source/Runtime/SlateCore/Public/Styling/SlateTypes.h
@@ -1409,6 +1409,41 @@ struct SLATECORE_API FSplitterStyle : public FSlateWidgetStyle
 	FSplitterStyle& SetHandleHighlightBrush( const FSlateBrush& InHandleHighlightBrush ){ HandleHighlightBrush = InHandleHighlightBrush; return *this; }
 };
 
+/**
+ * Represents the appearance of an STableView
+ */
+
+USTRUCT(BlueprintType)
+struct SLATECORE_API FTableViewStyle : public FSlateWidgetStyle
+{
+	GENERATED_USTRUCT_BODY()
+
+	FTableViewStyle();
+
+	virtual ~FTableViewStyle() {}
+
+	virtual void GetResources( TArray< const FSlateBrush* >& OutBrushes ) const override;
+
+	static const FName TypeName;
+	virtual const FName GetTypeName() const override { return TypeName; };
+
+	static const FTableViewStyle& GetDefault();
+
+	/** Brush used when a selected row is active */
+	UPROPERTY(EditAnywhere, BlueprintReadWrite, Category=Appearance)
+	FSlateBrush BackgroundBrush;
+	FTableViewStyle& SetBackgroundBrush( const FSlateBrush& InBackgroundBrush ){ BackgroundBrush = InBackgroundBrush; return *this; }
+
+	/**
+	 * Unlinks all colors in this style.
+	 * @see FSlateColor::Unlink
+	 */
+	void UnlinkColors()
+	{
+		BackgroundBrush.UnlinkColors();
+	}
+};
+
 
 /**
  * Represents the appearance of an STableRow
