Fixing PVS 7.7 Issues:

warning V1065: Expression can be simplified, check 'XYZ' and similar operands.

diff --git a/Engine/Source/Runtime/Slate/Private/Framework/Application/SlateApplication.cpp b/Engine/Source/Runtime/Slate/Private/Framework/Application/SlateApplication.cpp
--- a/Engine/Source/Runtime/Slate/Private/Framework/Application/SlateApplication.cpp
+++ b/Engine/Source/Runtime/Slate/Private/Framework/Application/SlateApplication.cpp
@@ -3574,8 +3574,8 @@ FVector2D FSlateApplication::CalculateTooltipWindowPosition( const FSlateRect& I
 	const FSlateRect WorkAreaRect( 
 		PlatformWorkArea.Left, 
 		PlatformWorkArea.Top, 
-		PlatformWorkArea.Left+(PlatformWorkArea.Right - PlatformWorkArea.Left), 
-		PlatformWorkArea.Top+(PlatformWorkArea.Bottom - PlatformWorkArea.Top) );
+		PlatformWorkArea.Right, 
+		PlatformWorkArea.Bottom);
 
 	float DPIScale = 1.0f; 
 
@@ -3650,8 +3650,8 @@ FVector2D FSlateApplication::CalculatePopupWindowPosition(const FSlateRect& InAn
 		const FSlateRect WorkAreaRect(
 			PlatformWorkArea.Left,
 			PlatformWorkArea.Top,
-			PlatformWorkArea.Left + (PlatformWorkArea.Right - PlatformWorkArea.Left),
-			PlatformWorkArea.Top + (PlatformWorkArea.Bottom - PlatformWorkArea.Top));
+			PlatformWorkArea.Right,
+			PlatformWorkArea.Bottom);
 
 		FVector2D ProposedPlacement = InProposedPlacement;
 
diff --git a/Engine/Source/Runtime/Slate/Private/Framework/Text/TextLayout.cpp b/Engine/Source/Runtime/Slate/Private/Framework/Text/TextLayout.cpp
--- a/Engine/Source/Runtime/Slate/Private/Framework/Text/TextLayout.cpp
+++ b/Engine/Source/Runtime/Slate/Private/Framework/Text/TextLayout.cpp
@@ -1967,7 +1967,7 @@ bool FTextLayout::RemoveAt( const FTextLocation& Location, int32 Count )
 	FLineModel& LineModel = LineModels[LineIndex];
 
 	//Make sure we aren't trying to remove more characters then we have
-	Count = RemoveLocation + Count > LineModel.Text->Len() ? Count - ((RemoveLocation + Count) - LineModel.Text->Len()) : Count;
+	Count = FMath::Min(Count, LineModel.Text->Len() - RemoveLocation);
 
 	if (Count == 0)
 	{
