Add option in WrapBoxSlot to force the slot to start on a new line.

diff --git a/Engine/Source/Runtime/Slate/Private/Widgets/Layout/SWrapBox.cpp b/Engine/Source/Runtime/Slate/Private/Widgets/Layout/SWrapBox.cpp
--- a/Engine/Source/Runtime/Slate/Private/Widgets/Layout/SWrapBox.cpp
+++ b/Engine/Source/Runtime/Slate/Private/Widgets/Layout/SWrapBox.cpp
@@ -169,9 +169,17 @@ void SWrapBox::FChildArranger::Arrange()
 		{
 			const float BottomBoundOfChild = ArrangementData.SlotOffset.Y + ArrangementData.SlotSize.Y;
 
+			if (Slot.bSlotForceNewLine)
+			{
+				// Begin a new line if the current one isn't empty, because the slot demanded a new line.
+				if (!IsFirstChildInCurrentLine())
+				{
+					BeginNewLine();
+				}
+			}
 			PRAGMA_DISABLE_DEPRECATION_WARNINGS
 			// Rule: If required due to a wrapping height under specified threshold, start a new line and allocate all of it to this child.
-			if ((Slot.SlotFillLineWhenSizeLessThan.IsSet() && WrapBox.PreferredSize.Get() < Slot.SlotFillLineWhenSizeLessThan.GetValue())
+			else if ((Slot.SlotFillLineWhenSizeLessThan.IsSet() && WrapBox.PreferredSize.Get() < Slot.SlotFillLineWhenSizeLessThan.GetValue())
 				|| (!Slot.SlotFillLineWhenSizeLessThan.IsSet() && Slot.SlotFillLineWhenWidthLessThan.IsSet() && WrapBox.PreferredSize.Get() < Slot.SlotFillLineWhenWidthLessThan.GetValue()))
 			{
 			PRAGMA_ENABLE_DEPRECATION_WARNINGS
@@ -204,9 +212,17 @@ void SWrapBox::FChildArranger::Arrange()
 		{
 			const float RightBoundOfChild = ArrangementData.SlotOffset.X + ArrangementData.SlotSize.X;
 
+			if (Slot.bSlotForceNewLine)
+			{
+				// Begin a new line if the current one isn't empty, because the slot demanded a new line.
+				if (!IsFirstChildInCurrentLine())
+				{
+					BeginNewLine();
+				}
+			}
 			PRAGMA_DISABLE_DEPRECATION_WARNINGS
 			// Rule: If required due to a wrapping width under specified threshold, start a new line and allocate all of it to this child.
-			if ((Slot.SlotFillLineWhenSizeLessThan.IsSet() && WrapBox.PreferredSize.Get() < Slot.SlotFillLineWhenSizeLessThan.GetValue())
+			else if ((Slot.SlotFillLineWhenSizeLessThan.IsSet() && WrapBox.PreferredSize.Get() < Slot.SlotFillLineWhenSizeLessThan.GetValue())
 				|| (!Slot.SlotFillLineWhenSizeLessThan.IsSet() && Slot.SlotFillLineWhenWidthLessThan.IsSet() && WrapBox.PreferredSize.Get() < Slot.SlotFillLineWhenWidthLessThan.GetValue()))
 			{
 			PRAGMA_ENABLE_DEPRECATION_WARNINGS
diff --git a/Engine/Source/Runtime/Slate/Public/Widgets/Layout/SWrapBox.h b/Engine/Source/Runtime/Slate/Public/Widgets/Layout/SWrapBox.h
--- a/Engine/Source/Runtime/Slate/Public/Widgets/Layout/SWrapBox.h
+++ b/Engine/Source/Runtime/Slate/Public/Widgets/Layout/SWrapBox.h
@@ -61,6 +61,7 @@ public:
 			, SlotFillLineWhenWidthLessThan()
 			, SlotFillLineWhenSizeLessThan()
 			, bSlotFillEmptySpace(false)
+			, bSlotForceNewLine(false)
 		{
 		}
 
@@ -87,11 +88,18 @@ public:
 			return *(static_cast<FSlot*>(this));
 		}
 
+		FSlot& ForceNewLine(bool bInForceNewLine)
+		{
+			bSlotForceNewLine = bInForceNewLine;
+			return *(static_cast<FSlot*>(this));
+		}
+
 		UE_DEPRECATED(4.26, "Deprecated, please use SlotFillLineWhenSizeLessThan instead")
 		TOptional<float> SlotFillLineWhenWidthLessThan;
 
 		TOptional<float> SlotFillLineWhenSizeLessThan;
 		bool bSlotFillEmptySpace;
+		bool bSlotForceNewLine;
 	};
 
 
diff --git a/Engine/Source/Runtime/UMG/Private/Components/WrapBoxSlot.cpp b/Engine/Source/Runtime/UMG/Private/Components/WrapBoxSlot.cpp
--- a/Engine/Source/Runtime/UMG/Private/Components/WrapBoxSlot.cpp
+++ b/Engine/Source/Runtime/UMG/Private/Components/WrapBoxSlot.cpp
@@ -15,6 +15,7 @@ UWrapBoxSlot::UWrapBoxSlot(const FObjectInitializer& ObjectInitializer)
 
 	bFillEmptySpace = false;
 	FillSpanWhenLessThan = 0;
+	bForceNewLine = false;
 }
 
 void UWrapBoxSlot::ReleaseSlateResources(bool bReleaseChildren)
@@ -32,6 +33,7 @@ void UWrapBoxSlot::BuildSlot(TSharedRef<SWrapBox> WrapBox)
 		.VAlign(VerticalAlignment)
 		.FillEmptySpace(bFillEmptySpace)
 		.FillLineWhenSizeLessThan(FillSpanWhenLessThan == 0 ? TOptional<float>() : TOptional<float>(FillSpanWhenLessThan))
+		.ForceNewLine(false)
 		[
 			Content == NULL ? SNullWidget::NullWidget : Content->TakeWidget()
 		];
@@ -82,6 +84,15 @@ void UWrapBoxSlot::SetVerticalAlignment(EVerticalAlignment InVerticalAlignment)
 	}
 }
 
+void UWrapBoxSlot::SetNewLine(bool InForceNewLine)
+{
+	bForceNewLine = InForceNewLine;
+	if (Slot)
+	{
+		Slot->bSlotForceNewLine = InForceNewLine;
+	}
+}
+
 void UWrapBoxSlot::SynchronizeProperties()
 {
 	SetPadding(Padding);
@@ -89,4 +100,5 @@ void UWrapBoxSlot::SynchronizeProperties()
 	SetFillSpanWhenLessThan(FillSpanWhenLessThan);
 	SetHorizontalAlignment(HorizontalAlignment);
 	SetVerticalAlignment(VerticalAlignment);
+	SetNewLine(bForceNewLine);
 }
diff --git a/Engine/Source/Runtime/UMG/Public/Components/WrapBoxSlot.h b/Engine/Source/Runtime/UMG/Public/Components/WrapBoxSlot.h
--- a/Engine/Source/Runtime/UMG/Public/Components/WrapBoxSlot.h
+++ b/Engine/Source/Runtime/UMG/Public/Components/WrapBoxSlot.h
@@ -42,6 +42,10 @@ public:
 	UPROPERTY(EditAnywhere, BlueprintReadOnly, Category="Layout|Wrap Box Slot")
 	TEnumAsByte<EVerticalAlignment> VerticalAlignment;
 
+	/** Should this slot start on a new line? */
+	UPROPERTY(EditAnywhere, BlueprintReadOnly, Category="Layout|Wrap Box Slot")
+	bool bForceNewLine;
+
 public:
 
 	UFUNCTION(BlueprintCallable, Category="Layout|Wrap Box Slot")
@@ -59,6 +63,9 @@ public:
 	UFUNCTION(BlueprintCallable, Category="Layout|Wrap Box Slot")
 	void SetVerticalAlignment(EVerticalAlignment InVerticalAlignment);
 
+	UFUNCTION(BlueprintCallable, Category="Layout|Wrap Box Slot")
+	void SetNewLine(bool InForceNewLine);
+
 public:
 
 	// UPanelSlot interface
