Slate: Reverted new SHeaderRow::Manual sizing feature to what it used to b and rename the new feature SHeaderRow::FillSized.

diff --git a/Engine/Source/Developer/SlateReflector/Private/Widgets/SWidgetReflector.cpp b/Engine/Source/Developer/SlateReflector/Private/Widgets/SWidgetReflector.cpp
--- a/Engine/Source/Developer/SlateReflector/Private/Widgets/SWidgetReflector.cpp
+++ b/Engine/Source/Developer/SlateReflector/Private/Widgets/SWidgetReflector.cpp
@@ -837,38 +837,38 @@ TSharedRef<SDockTab> SWidgetReflector::SpawnWidgetHierarchyTab(const FSpawnTabAr
 						+SHeaderRow::Column(SReflectorTreeWidgetItem::NAME_Visibility)
 						.DefaultLabel(LOCTEXT("Visibility", "Visibility"))
 						.DefaultTooltip(LOCTEXT("VisibilityTooltip", "Visibility"))
-						.ManualWidth(125.0f)
+						.FillSized(125.0f)
 						
 						+ SHeaderRow::Column(SReflectorTreeWidgetItem::NAME_Enabled)
 						.DefaultLabel(LOCTEXT("Enabled", "Enabled"))
 						.DefaultTooltip(LOCTEXT("EnabledToolTip", "Enabled"))
-						.ManualWidth(60.0f)
+						.FillSized(60.0f)
 
 						+ SHeaderRow::Column(SReflectorTreeWidgetItem::NAME_Focusable)
 						.DefaultLabel(LOCTEXT("Focus", "Focus"))
 						.DefaultTooltip(LOCTEXT("FocusableTooltip", "Focusability (Note that for hit-test directional navigation to work it must be Focusable and \"Visible\"!)"))
-						.ManualWidth(60.0f)
+						.FillSized(60.0f)
 
 						+ SHeaderRow::Column(SReflectorTreeWidgetItem::NAME_HasActiveTimer)
 						.DefaultLabel(LOCTEXT("HasActiveTimer", "Timer"))
 						.DefaultTooltip(LOCTEXT("HasActiveTimerTooltip", "Has Active Timer"))
-						.ManualWidth(60.0f)
+						.FillSized(60.0f)
 
 						+SHeaderRow::Column(SReflectorTreeWidgetItem::NAME_Clipping)
 						.DefaultLabel(LOCTEXT("Clipping", "Clipping" ))
-						.ManualWidth(100.0f)
+						.FillSized(100.0f)
 
 						+ SHeaderRow::Column(SReflectorTreeWidgetItem::NAME_ActualSize)
 						.DefaultLabel(LOCTEXT("ActualSize", "Size"))
-						.ManualWidth(100.0f)
+						.FillSized(100.0f)
 
 						+SHeaderRow::Column(SReflectorTreeWidgetItem::NAME_WidgetInfo)
 						.DefaultLabel(LOCTEXT("Source", "Source" ))
-						.ManualWidth(200.f)
+						.FillSized(200.f)
 
 						+SHeaderRow::Column(SReflectorTreeWidgetItem::NAME_Address)
 						.DefaultLabel( LOCTEXT("Address", "Address") )
-						.ManualWidth(170.0f)
+						.FillSized(170.0f)
 					)
 				]
 			]
diff --git a/Engine/Source/Runtime/Slate/Private/Widgets/Views/SHeaderRow.cpp b/Engine/Source/Runtime/Slate/Private/Widgets/Views/SHeaderRow.cpp
--- a/Engine/Source/Runtime/Slate/Private/Widgets/Views/SHeaderRow.cpp
+++ b/Engine/Source/Runtime/Slate/Private/Widgets/Views/SHeaderRow.cpp
@@ -772,8 +772,92 @@ void SHeaderRow::RegenerateWidgets()
 				}
 				break;
 
-
 				case EColumnSizeMode::Manual:
+				{
+					// Sizing grip to put at the end of the column - we can't use a SSplitter here as it doesn't have the resizing behavior we need
+					const float GripSize = SplitterHandleSize > 0.0f ? SplitterHandleSize + 4.0f : 5.0f;
+					TSharedRef<SBorder> SizingGrip = SNew(SBorder)
+						.Padding(0.0f)
+						.BorderImage( &Style->ColumnSplitterStyle )
+						.Cursor(EMouseCursor::ResizeLeftRight)
+						.Content()
+						[
+							SNew(SSpacer)
+							.Size(FVector2D(GripSize, GripSize))
+						];
+
+					TWeakPtr<SBorder> WeakSizingGrip = SizingGrip;
+					auto SizingGrip_OnMouseButtonDown = [&SomeColumn, WeakSizingGrip](const FGeometry&, const FPointerEvent&) -> FReply
+					{
+						TSharedPtr<SBorder> SizingGripPtr = WeakSizingGrip.Pin();
+						if (SizingGripPtr.IsValid())
+						{
+							return FReply::Handled().CaptureMouse(SizingGripPtr.ToSharedRef());
+						}
+						return FReply::Unhandled();
+					};
+
+					auto SizingGrip_OnMouseButtonUp = [&SomeColumn, WeakSizingGrip](const FGeometry&, const FPointerEvent&) -> FReply
+					{
+						TSharedPtr<SBorder> SizingGripPtr = WeakSizingGrip.Pin();
+						if (SizingGripPtr.IsValid() && SizingGripPtr->HasMouseCapture())
+						{
+							return FReply::Handled().ReleaseMouseCapture();
+						}
+						return FReply::Unhandled();
+					};
+
+					auto SizingGrip_OnMouseMove = [&SomeColumn, WeakSizingGrip](const FGeometry&, const FPointerEvent& InPointerEvent) -> FReply
+					{
+						TSharedPtr<SBorder> SizingGripPtr = WeakSizingGrip.Pin();
+						if (SizingGripPtr.IsValid() && SizingGripPtr->HasMouseCapture())
+						{
+							// The sizing grip has been moved, so update our columns size from the movement delta
+							const float NewWidth = SomeColumn.GetWidth() + InPointerEvent.GetCursorDelta().X;
+							SomeColumn.SetWidth(FMath::Max(20.0f, NewWidth));
+							return FReply::Handled();
+						}
+						return FReply::Unhandled();
+					};
+
+					// Bind the events to handle the drag sizing
+					SizingGrip->SetOnMouseButtonDown(FPointerEventHandler::CreateLambda(SizingGrip_OnMouseButtonDown));
+					SizingGrip->SetOnMouseButtonUp(FPointerEventHandler::CreateLambda(SizingGrip_OnMouseButtonUp));
+					SizingGrip->SetOnMouseMove(FPointerEventHandler::CreateLambda(SizingGrip_OnMouseMove));
+
+					auto GetColumnWidthAsOptionalSize = [&SomeColumn]() -> FOptionalSize
+					{
+						const float DesiredWidth = SomeColumn.GetWidth();
+						return FOptionalSize(DesiredWidth);
+					};
+
+					TAttribute<FOptionalSize> WidthBinding;
+					WidthBinding.Bind(TAttribute<FOptionalSize>::FGetter::CreateLambda(GetColumnWidthAsOptionalSize));
+
+					// Add resizable cell
+					Splitter->AddSlot()
+						.SizeRule(SSplitter::SizeToContent)
+						.Resizable(false)
+						[
+							SNew(SBox)
+							.WidthOverride(WidthBinding)
+							[
+								SNew(SOverlay)
+								+ SOverlay::Slot()
+								[
+									NewHeader
+								]
+								+ SOverlay::Slot()
+								.HAlign(HAlign_Right)
+								[
+									SizingGrip
+								]
+							]
+						];
+				}
+				break;
+
+				case EColumnSizeMode::FillSized:
 				{
 					auto GetColumnWidthAsOptionalSize = [&SomeColumn]() -> FOptionalSize
 					{
diff --git a/Engine/Source/Runtime/Slate/Public/Widgets/Views/SHeaderRow.h b/Engine/Source/Runtime/Slate/Public/Widgets/Views/SHeaderRow.h
--- a/Engine/Source/Runtime/Slate/Public/Widgets/Views/SHeaderRow.h
+++ b/Engine/Source/Runtime/Slate/Public/Widgets/Views/SHeaderRow.h
@@ -59,6 +59,9 @@ namespace EColumnSizeMode
 
 		/** Column is set to a width which can be user-sized */
 		Manual = 2,
+
+		/** Column stretches as Fill but is initialized with a width */
+		FillSized = 3,
 	};
 };
 
@@ -118,9 +121,15 @@ public:
 			SLATE_ARGUMENT( FName, ColumnId )
 			SLATE_ATTRIBUTE( FText, DefaultLabel )
 			SLATE_ATTRIBUTE( FText, DefaultTooltip )
+
+			/** Set the Column Size Mode to Fill. It's a fraction between 0 and 1 */
 			SLATE_ATTRIBUTE( float, FillWidth )
+			/** Set the Column Size Mode to Fixed. */
 			SLATE_ARGUMENT( TOptional< float >, FixedWidth )
+			/** Set the Column Size Mode to Manual. */
 			SLATE_ATTRIBUTE( float, ManualWidth )
+			/** Set the Column Size Mode to Fill Sized. */
+			SLATE_ARGUMENT(TOptional< float >, FillSized)
 			SLATE_EVENT( FOnWidthChanged, OnWidthChanged )
 
 			SLATE_DEFAULT_SLOT( FArguments, HeaderContent )
@@ -175,6 +184,11 @@ public:
 				Width = InArgs._ManualWidth;
 				SizeRule = EColumnSizeMode::Manual;
 			}
+			else if ( InArgs._FillSized.IsSet() )
+			{
+				Width = InArgs._FillSized.GetValue();
+				SizeRule = EColumnSizeMode::FillSized;
+			}
 			else
 			{
 				Width = InArgs._FillWidth;
diff --git a/Engine/Source/Runtime/Slate/Public/Widgets/Views/STableRow.h b/Engine/Source/Runtime/Slate/Public/Widgets/Views/STableRow.h
--- a/Engine/Source/Runtime/Slate/Public/Widgets/Views/STableRow.h
+++ b/Engine/Source/Runtime/Slate/Public/Widgets/Views/STableRow.h
@@ -1388,6 +1388,7 @@ protected:
 				break;
 
 				case EColumnSizeMode::Manual:
+				case EColumnSizeMode::FillSized:
 				{
 					auto GetColumnWidthAsOptionalSize = [&Column]() -> FOptionalSize
 					{
@@ -1414,6 +1415,7 @@ protected:
 				break;
 
 				default:
+					ensure(false);
 					break;
 				}
 
