Splitter doesn't try to steal space from collapsed columns.

diff --git a/Engine/Source/Runtime/Slate/Private/Widgets/Layout/SSplitter.cpp b/Engine/Source/Runtime/Slate/Private/Widgets/Layout/SSplitter.cpp
--- a/Engine/Source/Runtime/Slate/Private/Widgets/Layout/SSplitter.cpp
+++ b/Engine/Source/Runtime/Slate/Private/Widgets/Layout/SSplitter.cpp
@@ -105,22 +105,29 @@ TArray<FLayoutGeometry> SSplitter::ArrangeChildrenForLayout(const FGeometry& All
 	float NonResizableSpace = 0;
 	float MinResizableSpace = 0;
 
+	auto IsSlotCollapsed = [](const FSlot& Slot)
+	{
+		return Slot.GetWidget()->GetVisibility() == EVisibility::Collapsed;
+	};
+
 	for (int32 ChildIndex = 0; ChildIndex < Children.Num(); ++ChildIndex)
 	{
 		const FSlot& CurSlot = Children[ChildIndex];
-		if (CurSlot.GetWidget()->GetVisibility() != EVisibility::Collapsed)
+		if (IsSlotCollapsed(CurSlot))
 		{
-			++NumNonCollapsedChildren;
+			continue;
+		}
 
-			if (CurSlot.SizingRule.Get() == SSplitter::SizeToContent)
-			{
-				NonResizableSpace += CurSlot.GetWidget()->GetDesiredSize()[AxisIndex];
-			}
-			else // SizingRule == SSplitter::FractionOfParent
-			{
-				MinResizableSpace += FMath::Max(MinSplitterChildLength, CurSlot.MinSizeValue.Get(0));
-				CoefficientTotal += CurSlot.SizeValue.Get();
-			}
+		++NumNonCollapsedChildren;
+
+		if (CurSlot.SizingRule.Get() == SSplitter::SizeToContent)
+		{
+			NonResizableSpace += CurSlot.GetWidget()->GetDesiredSize()[AxisIndex];
+		}
+		else // SizingRule == SSplitter::FractionOfParent
+		{
+			MinResizableSpace += FMath::Max(MinSplitterChildLength, CurSlot.MinSizeValue.Get(0));
+			CoefficientTotal += CurSlot.SizeValue.Get();
 		}
 	}
 
@@ -128,13 +135,22 @@ TArray<FLayoutGeometry> SSplitter::ArrangeChildrenForLayout(const FGeometry& All
 	const float SpaceNeededForHandles = FMath::Max(0, NumNonCollapsedChildren - 1) * PhysicalSplitterHandleSize;
 	const float ResizableSpace = AllottedGeometry.Size.Component(AxisIndex) - SpaceNeededForHandles - NonResizableSpace;
 
-	TArray<float> SlotSizes;
+	FMemMark Mark(FMemStack::Get());
+	TArray<float, TMemStackAllocator<>> SlotSizes;
+	SlotSizes.Empty(Children.Num());
 
 	// calculate slot sizes
 	for (int32 ChildIndex = 0; ChildIndex < Children.Num(); ++ChildIndex)
 	{
 		const FSlot& CurSlot = Children[ChildIndex];
 
+		// slot is collapsed, size is automatically 0
+		if (IsSlotCollapsed(CurSlot))
+		{
+			SlotSizes.Add(0);
+			continue;
+		}
+
 		float ChildSpace = (CurSlot.SizingRule.Get() == SSplitter::SizeToContent)
 			? CurSlot.GetWidget()->GetDesiredSize()[AxisIndex]
 			: ResizableSpace * CurSlot.SizeValue.Get() / CoefficientTotal;
@@ -147,7 +163,7 @@ TArray<FLayoutGeometry> SSplitter::ArrangeChildrenForLayout(const FGeometry& All
 
 			for (int32 PrevIndex = ChildIndex - 1; PrevIndex >= 0 && CurrentRequiredSpace > 0; --PrevIndex)
 			{
-				if (Children[PrevIndex].SizingRule.Get() == SSplitter::FractionOfParent)
+				if (Children[PrevIndex].SizingRule.Get() == SSplitter::FractionOfParent && !IsSlotCollapsed(CurSlot))
 				{
 					const float MinChildSize = FMath::Max(MinSplitterChildLength, Children[PrevIndex].MinSizeValue.Get(0));
 					const float AvailableSpace = SlotSizes[PrevIndex] - MinChildSize;
@@ -175,7 +191,7 @@ TArray<FLayoutGeometry> SSplitter::ArrangeChildrenForLayout(const FGeometry& All
 		ResultGeometries.Emplace(FSlateLayoutTransform(ChildOffset), ChildSize);
 
 		// Advance to the next slot. If the child is collapsed, it takes up no room and does not need a splitter
-		if (Children[ChildIndex].GetWidget()->GetVisibility() != EVisibility::Collapsed)
+		if (!IsSlotCollapsed(Children[ChildIndex]))
 		{
 			Offset += FMath::RoundToInt(SlotSizes[ChildIndex] + PhysicalSplitterHandleSize);
 		}
