Segmented controls now uniformly space out their segments

diff --git a/Engine/Source/Editor/DetailCustomizations/Private/MobilityCustomization.cpp b/Engine/Source/Editor/DetailCustomizations/Private/MobilityCustomization.cpp
--- a/Engine/Source/Editor/DetailCustomizations/Private/MobilityCustomization.cpp
+++ b/Engine/Source/Editor/DetailCustomizations/Private/MobilityCustomization.cpp
@@ -52,6 +52,7 @@ void FMobilityCustomization::CreateMobilityCustomization(IDetailCategoryBuilder&
 
 		// Static Mobility
 		ButtonOptionsPanel->AddSlot(EComponentMobility::Static)
+		.HAlign(HAlign_Center)
 		[
 			SNew(SHorizontalBox)
 			+ SHorizontalBox::Slot()
@@ -64,7 +65,6 @@ void FMobilityCustomization::CreateMobilityCustomization(IDetailCategoryBuilder&
 				.ColorAndOpacity(FSlateColor::UseForeground())
 			]
 			+ SHorizontalBox::Slot()
-			.AutoWidth()
 			.VAlign(VAlign_Center)
 			[
 				SNew(STextBlock)
@@ -83,6 +83,7 @@ void FMobilityCustomization::CreateMobilityCustomization(IDetailCategoryBuilder&
 			: LOCTEXT("Mobility_Stationary_Object_Tooltip", "A stationary object can be changed in game but not moved, and enables cached lighting methods. \n* Cached Dynamic Shadows.");
 
 		ButtonOptionsPanel->AddSlot(EComponentMobility::Stationary)
+		.HAlign(HAlign_Center)
 		[
 			SNew(SHorizontalBox)
 			+ SHorizontalBox::Slot()
@@ -95,7 +96,6 @@ void FMobilityCustomization::CreateMobilityCustomization(IDetailCategoryBuilder&
 				.ColorAndOpacity(FSlateColor::UseForeground())
 			]
 			+ SHorizontalBox::Slot()
-			.AutoWidth()
 			.VAlign(VAlign_Center)
 			[
 				SNew(STextBlock)
@@ -112,6 +112,7 @@ void FMobilityCustomization::CreateMobilityCustomization(IDetailCategoryBuilder&
 
 	// Movable Mobility
 	ButtonOptionsPanel->AddSlot(EComponentMobility::Movable)
+	.HAlign(HAlign_Center)
 	[
 		SNew(SHorizontalBox)
 		+ SHorizontalBox::Slot()
@@ -124,7 +125,6 @@ void FMobilityCustomization::CreateMobilityCustomization(IDetailCategoryBuilder&
 			.ColorAndOpacity(FSlateColor::UseForeground())
 		]
 		+ SHorizontalBox::Slot()
-		.AutoWidth()
 		.VAlign(VAlign_Center)
 		[
 			SNew(STextBlock)
diff --git a/Engine/Source/Runtime/Slate/Public/Widgets/Input/SSegmentedControl.h b/Engine/Source/Runtime/Slate/Public/Widgets/Input/SSegmentedControl.h
--- a/Engine/Source/Runtime/Slate/Public/Widgets/Input/SSegmentedControl.h
+++ b/Engine/Source/Runtime/Slate/Public/Widgets/Input/SSegmentedControl.h
@@ -10,6 +10,7 @@
 #include "Widgets/Input/SCheckBox.h"
 #include "Widgets/Layout/SSpacer.h"
 #include "Widgets/Text/STextBlock.h"
+#include "Widgets/Layout/SUniformGridPanel.h"
 
 /**
  * A Slate Segmented Control is functionally similar to a group of Radio Buttons.
@@ -28,10 +29,11 @@ public:
 
 	/** Stores the per-child info for this panel type */
 	template< typename SlotOptionType >
-	struct FSlot : public TSlotBase<FSlot<SlotOptionType>> 
+	struct FSlot : public TSlotBase<FSlot<SlotOptionType>>, TSupportsContentAlignmentMixin<FSlot<SlotOptionType>>
 	{
 		FSlot(const SlotOptionType& InValue) 
-		: TSlotBase<FSlot<SlotOptionType>>() 
+		: TSlotBase<FSlot<SlotOptionType>>()
+		, TSupportsContentAlignmentMixin<FSlot<SlotOptionType>>(HAlign_Center, VAlign_Fill)
 		, _Text()
 		, _Tooltip()
 		, _Icon(nullptr)
@@ -103,7 +105,7 @@ public:
 
 	void RebuildChildren()
 	{
-		TSharedPtr<SHorizontalBox> HBox = SNew(SHorizontalBox);
+		TSharedPtr<SUniformGridPanel> UniformBox = SNew(SUniformGridPanel);
 
 		const int32 NumSlots = Children.Num();
 		for ( int32 SlotIndex = 0; SlotIndex < NumSlots; ++SlotIndex )
@@ -144,10 +146,12 @@ public:
 				];
 			}
 
-			HBox->AddSlot()
-			.AutoWidth()
+			UniformBox->AddSlot(SlotIndex, 0)
+			// Note HAlignment is applied at the check box level because if it were applied here it would make the slots look physically disconnected from each other 
+			.VAlign(ChildSlotPtr->VAlignment)
 			[
 				SNew(SCheckBox)
+				.HAlign(ChildSlotPtr->HAlignment)
 				.ToolTipText(ChildSlotPtr->_Tooltip)
 				.Style(SlotIndex == 0 ? &Style->FirstControlStyle : SlotIndex == (NumSlots - 1) ? &Style->LastControlStyle : &Style->ControlStyle)
 				.IsChecked(this, &SSegmentedControl::IsCurrentValue, ChildValue)
@@ -160,7 +164,7 @@ public:
 
 		ChildSlot
 		[
-			HBox.ToSharedRef()
+			UniformBox.ToSharedRef()
 		];
 
 	}
