Slate: Fix Table view header that mix Fill and Manual.
#rnx

diff --git a/Engine/Source/Runtime/Slate/Private/Widgets/Layout/SSplitter.cpp b/Engine/Source/Runtime/Slate/Private/Widgets/Layout/SSplitter.cpp
--- a/Engine/Source/Runtime/Slate/Private/Widgets/Layout/SSplitter.cpp
+++ b/Engine/Source/Runtime/Slate/Private/Widgets/Layout/SSplitter.cpp
@@ -596,7 +596,7 @@ void SSplitter::FindAllResizeableSlotsAfterHandle( int32 DraggedHandle, const TP
 		OutSlotIndicies.Add( SlotIndex );
 	}
 };
-
+PRAGMA_DISABLE_OPTIMIZATION
 void SSplitter::HandleResizingDelta(EOrientation SplitterOrientation, const float InPhysicalSplitterHandleSize, const ESplitterResizeMode::Type InResizeMode, int32 DraggedHandle, float Delta, TPanelChildren<FSlot>& InChildren, const TArray<FLayoutGeometry>& ChildGeometries)
 {
 	const int32 NumChildren = InChildren.Num();
@@ -648,6 +648,7 @@ void SSplitter::HandleResizingDelta(EOrientation SplitterOrientation, const floa
 
 		// Get references the prev and next children and their layout settings so that we can modify them.
 		const int32 SlotBeforeDragHandle = FindResizeableSlotBeforeHandle(DraggedHandle, InChildren);
+		check(InChildren.IsValidIndex(SlotBeforeDragHandle));
 		FSlot& PrevChild = InChildren[SlotBeforeDragHandle];
 		const FLayoutGeometry& PrevChildGeom = ChildGeometries[SlotBeforeDragHandle];
 
@@ -693,22 +694,31 @@ void SSplitter::HandleResizingDelta(EOrientation SplitterOrientation, const floa
 		// Cells being resized are both stretch values -> redistribute the stretch coefficients proportionately
 		// to match the new child sizes on the screen.
 		{
-			float TotalLength = NewPrevChildLength;
-			float TotalStretchCoefficients = PrevChild.SizingRule.Get() == ESizeRule::FractionOfParent ? PrevChild.SizeValue.Get() : 0.f;
+			float TotalStretchLength = 0.f;
+			float TotalStretchCoefficients = 0.f;
+			
+			if (PrevChild.SizingRule.Get() == ESizeRule::FractionOfParent)
+			{
+				TotalStretchLength = NewPrevChildLength;
+				TotalStretchCoefficients = PrevChild.SizeValue.Get();
+			}
 
 			for (int32 SlotIndex = 0; SlotIndex < NumSlotsAfterDragHandle; SlotIndex++)
 			{
-				const FSlotInfo& SlotInfo = SlotsAfterDragHandle[ SlotIndex ];
+				const FSlotInfo& SlotInfo = SlotsAfterDragHandle[SlotIndex];
 
-				TotalLength += SlotInfo.NewSize;
-				TotalStretchCoefficients += SlotInfo.Slot->SizingRule.Get() == ESizeRule::FractionOfParent ? SlotInfo.Slot->SizeValue.Get() : 0.f;
+				if (SlotInfo.Slot->SizingRule.Get() == ESizeRule::FractionOfParent)
+				{
+					TotalStretchLength += SlotInfo.NewSize;
+					TotalStretchCoefficients += SlotInfo.Slot->SizeValue.Get();
+				}
 			}
 
 			auto SetSize = [](FSlot& Slot, float TotalStretchCoefficients, float NewLength, float TotalLength)
 			{
 				if (Slot.SizingRule.Get() == ESizeRule::FractionOfParent)
 				{
-					const float NewFillSize = (TotalStretchCoefficients * NewLength / TotalLength);
+					float NewFillSize = TotalLength > 0.f ? (TotalStretchCoefficients * (NewLength / TotalLength)) : TotalStretchCoefficients;
 					if (Slot.OnSlotResized_Handler.IsBound())
 					{
 						Slot.OnSlotResized_Handler.Execute(NewFillSize);
@@ -725,16 +735,16 @@ void SSplitter::HandleResizingDelta(EOrientation SplitterOrientation, const floa
 				}
 			};
 
-			SetSize(PrevChild, TotalStretchCoefficients, NewPrevChildLength, TotalLength);
+			SetSize(PrevChild, TotalStretchCoefficients, NewPrevChildLength, TotalStretchLength);
 			for (int32 SlotIndex = 0; SlotIndex < NumSlotsAfterDragHandle; SlotIndex++)
 			{
 				const FSlotInfo& SlotInfo = SlotsAfterDragHandle[ SlotIndex ];
-				SetSize(*(SlotInfo.Slot), TotalStretchCoefficients, SlotInfo.NewSize, TotalLength);
+				SetSize(*(SlotInfo.Slot), TotalStretchCoefficients, SlotInfo.NewSize, TotalStretchLength);
 			}
 		}
 	}
 }
-
+PRAGMA_ENABLE_OPTIMIZATION
 void SSplitter::HandleResizingBySize(EOrientation SplitterOrientation, const float InPhysicalSplitterHandleSize, const ESplitterResizeMode::Type InResizeMode, int32 DraggedHandle, const FVector2D& InDesiredSize, TPanelChildren<FSlot>& InChildren, const TArray<FLayoutGeometry>& ChildGeometries)
 {
 	const int32 AxisIndex = (SplitterOrientation == Orient_Horizontal) ? 0 : 1;
