Fixed issue where SSpinBox would not respect {Min, Max}FractionalDigits.

diff --git a/Engine/Source/Runtime/Slate/Public/Widgets/Input/SSpinBox.h b/Engine/Source/Runtime/Slate/Public/Widgets/Input/SSpinBox.h
--- a/Engine/Source/Runtime/Slate/Public/Widgets/Input/SSpinBox.h
+++ b/Engine/Source/Runtime/Slate/Public/Widgets/Input/SSpinBox.h
@@ -186,6 +186,8 @@ public:
 
 		MinFractionalDigits = (InArgs._MinFractionalDigits.Get().IsSet()) ? InArgs._MinFractionalDigits : DefaultMinFractionalDigits;
 		MaxFractionalDigits = (InArgs._MaxFractionalDigits.Get().IsSet()) ? InArgs._MaxFractionalDigits : DefaultMaxFractionalDigits;
+		SetMaxFractionalDigits(MaxFractionalDigits);
+		SetMinFractionalDigits(MinFractionalDigits);
 
 		AlwaysUsesDeltaSnap = InArgs._AlwaysUsesDeltaSnap;
 
