Allow segmented controls to have a uniform padding between each slot.

diff --git a/Engine/Source/Runtime/Slate/Public/Widgets/Input/SSegmentedControl.h b/Engine/Source/Runtime/Slate/Public/Widgets/Input/SSegmentedControl.h
--- a/Engine/Source/Runtime/Slate/Public/Widgets/Input/SSegmentedControl.h
+++ b/Engine/Source/Runtime/Slate/Public/Widgets/Input/SSegmentedControl.h
@@ -82,10 +82,16 @@ public:
 		/** Slot type supported by this panel */
 		SLATE_SUPPORTS_SLOT(FSlot<OptionType>)
 
+		/** Styling for this control */
 		SLATE_STYLE_ARGUMENT(FSegmentedControlStyle, Style)
 
+		/** The current control value */
 		SLATE_ATTRIBUTE( OptionType, Value )
 
+		/** Padding to apply to each slot */
+		SLATE_ATTRIBUTE(FMargin, UniformPadding)
+
+		/** Called when the value is changed */
 		SLATE_EVENT( FOnValueChanged, OnValueChanged )
 		
 	SLATE_END_ARGS()
@@ -99,13 +105,16 @@ public:
 		CurrentValue = InArgs._Value;
 		OnValueChanged = InArgs._OnValueChanged;
 
+		UniformPadding = InArgs._UniformPadding;
+
 		Children = InArgs.Slots;
 		RebuildChildren();
 	}
 
 	void RebuildChildren()
 	{
-		TSharedPtr<SUniformGridPanel> UniformBox = SNew(SUniformGridPanel);
+		TSharedPtr<SUniformGridPanel> UniformBox =
+			SNew(SUniformGridPanel);
 
 		const int32 NumSlots = Children.Num();
 		for ( int32 SlotIndex = 0; SlotIndex < NumSlots; ++SlotIndex )
@@ -156,6 +165,7 @@ public:
 				.Style(SlotIndex == 0 ? &Style->FirstControlStyle : SlotIndex == (NumSlots - 1) ? &Style->LastControlStyle : &Style->ControlStyle)
 				.IsChecked(this, &SSegmentedControl::IsCurrentValue, ChildValue)
 				.OnCheckStateChanged(this, &SSegmentedControl::CommitValue, ChildValue)
+				.Padding(UniformPadding)
 				[
 					Child
 				]
@@ -219,5 +229,7 @@ protected:
 
 	TAttribute<OptionType> CurrentValue;
 
+	TAttribute<FMargin> UniformPadding;
+
 	const FSegmentedControlStyle* Style;
 };
