Slate: Remove deprecated invalidation flag. Remove unwanted pragma disable optimization.

diff --git a/Engine/Source/Runtime/Slate/Private/Widgets/Layout/SSplitter.cpp b/Engine/Source/Runtime/Slate/Private/Widgets/Layout/SSplitter.cpp
--- a/Engine/Source/Runtime/Slate/Private/Widgets/Layout/SSplitter.cpp
+++ b/Engine/Source/Runtime/Slate/Private/Widgets/Layout/SSplitter.cpp
@@ -596,7 +596,7 @@ void SSplitter::FindAllResizeableSlotsAfterHandle( int32 DraggedHandle, const TP
 		OutSlotIndicies.Add( SlotIndex );
 	}
 };
-PRAGMA_DISABLE_OPTIMIZATION
+
 void SSplitter::HandleResizingDelta(EOrientation SplitterOrientation, const float InPhysicalSplitterHandleSize, const ESplitterResizeMode::Type InResizeMode, int32 DraggedHandle, float Delta, TPanelChildren<FSlot>& InChildren, const TArray<FLayoutGeometry>& ChildGeometries)
 {
 	const int32 NumChildren = InChildren.Num();
@@ -744,7 +744,7 @@ void SSplitter::HandleResizingDelta(EOrientation SplitterOrientation, const floa
 		}
 	}
 }
-PRAGMA_ENABLE_OPTIMIZATION
+
 void SSplitter::HandleResizingBySize(EOrientation SplitterOrientation, const float InPhysicalSplitterHandleSize, const ESplitterResizeMode::Type InResizeMode, int32 DraggedHandle, const FVector2D& InDesiredSize, TPanelChildren<FSlot>& InChildren, const TArray<FLayoutGeometry>& ChildGeometries)
 {
 	const int32 AxisIndex = (SplitterOrientation == Orient_Horizontal) ? 0 : 1;
diff --git a/Engine/Source/Runtime/Slate/Public/Widgets/Layout/SSpacer.h b/Engine/Source/Runtime/Slate/Public/Widgets/Layout/SSpacer.h
--- a/Engine/Source/Runtime/Slate/Public/Widgets/Layout/SSpacer.h
+++ b/Engine/Source/Runtime/Slate/Public/Widgets/Layout/SSpacer.h
@@ -49,7 +49,7 @@ public:
 	void SetSize( TAttribute<FVector2D> InSpacerSize )
 	{
 		SpacerSize = InSpacerSize;
-		Invalidate(EInvalidateWidget::Layout);
+		Invalidate(EInvalidateWidgetReason::Layout);
 	}
 
 protected:
