Fixed white progress bars when the progress is indeterminate

diff --git a/Engine/Source/Runtime/Slate/Private/Widgets/Notifications/SProgressBar.cpp b/Engine/Source/Runtime/Slate/Private/Widgets/Notifications/SProgressBar.cpp
--- a/Engine/Source/Runtime/Slate/Private/Widgets/Notifications/SProgressBar.cpp
+++ b/Engine/Source/Runtime/Slate/Private/Widgets/Notifications/SProgressBar.cpp
@@ -155,9 +155,6 @@ int32 SProgressBar::OnPaint( const FPaintArgs& Args, const FGeometry& AllottedGe
 	const ESlateDrawEffect DrawEffects = bEnabled ? ESlateDrawEffect::None : ESlateDrawEffect::DisabledEffect;
 	
 	const FSlateBrush* CurrentFillImage = GetFillImage();
-	
-	const FLinearColor FillColorAndOpacitySRGB(InWidgetStyle.GetColorAndOpacityTint() * FillColorAndOpacity.Get().GetColor(InWidgetStyle) * CurrentFillImage->GetTint(InWidgetStyle));
-	const FLinearColor ColorAndOpacitySRGB = InWidgetStyle.GetColorAndOpacityTint();
 
 	TOptional<float> ProgressFraction = Percent.Get();	
 	FVector2D BorderPaddingRef = BorderPadding.Get();
@@ -175,6 +172,8 @@ int32 SProgressBar::OnPaint( const FPaintArgs& Args, const FGeometry& AllottedGe
 	
 	if( ProgressFraction.IsSet() )
 	{
+		const FLinearColor FillColorAndOpacitySRGB(InWidgetStyle.GetColorAndOpacityTint() * FillColorAndOpacity.Get().GetColor(InWidgetStyle) * CurrentFillImage->GetTint(InWidgetStyle));
+
 		EProgressBarFillType::Type ComputedBarFillType = BarFillType;
 		if (GSlateFlowDirection == EFlowDirection::RightToLeft)
 		{
@@ -338,6 +337,9 @@ int32 SProgressBar::OnPaint( const FPaintArgs& Args, const FGeometry& AllottedGe
 	{
 		const FSlateBrush* CurrentMarqueeImage = GetMarqueeImage();
 		
+		const FLinearColor FillColorAndOpacitySRGB(InWidgetStyle.GetColorAndOpacityTint()* FillColorAndOpacity.Get().GetColor(InWidgetStyle)* CurrentMarqueeImage->GetTint(InWidgetStyle));
+
+
 		// Draw Marquee
 		const float MarqueeAnimOffset = CurrentMarqueeImage->ImageSize.X * MarqueeOffset;
 		const float MarqueeImageSize = CurrentMarqueeImage->ImageSize.X;
@@ -352,7 +354,7 @@ int32 SProgressBar::OnPaint( const FPaintArgs& Args, const FGeometry& AllottedGe
 					FVector2D( AllottedGeometry.GetLocalSize().X + MarqueeImageSize, AllottedGeometry.GetLocalSize().Y )),
 				CurrentMarqueeImage,
 				DrawEffects,
-				ColorAndOpacitySRGB
+				FillColorAndOpacitySRGB
 				);
 
 			OutDrawElements.PopClip();
