Menu WidgetBlock items close submenus and have enchanced search functionality

diff --git a/Engine/Source/Runtime/Slate/Private/Framework/MultiBox/SWidgetBlock.cpp b/Engine/Source/Runtime/Slate/Private/Framework/MultiBox/SWidgetBlock.cpp
--- a/Engine/Source/Runtime/Slate/Private/Framework/MultiBox/SWidgetBlock.cpp
+++ b/Engine/Source/Runtime/Slate/Private/Framework/MultiBox/SWidgetBlock.cpp
@@ -64,6 +64,7 @@ void SWidgetBlock::Construct( const FArguments& InArgs )
  */
 void SWidgetBlock::BuildMultiBlockWidget(const ISlateStyle* StyleSet, const FName& StyleName)
 {
+	TSharedPtr< const FMultiBox > MultiBox = OwnerMultiBoxWidget.Pin()->GetMultiBox();
 	TSharedRef< const FWidgetBlock > WidgetBlock = StaticCastSharedRef< const FWidgetBlock >( MultiBlock.ToSharedRef() );
 
 	// Support menus which do not have a defined widget style yet
@@ -94,12 +95,36 @@ void SWidgetBlock::BuildMultiBlockWidget(const ISlateStyle* StyleSet, const FNam
 
 		// If there is no label, allow the custom menu widget to consume the entire space
 		if (!bHasLabel)
+		{
 			Padding = FMargin(0);
+		}
+	}
+
+	// For searchable menus with a custom widgets without a label, find a TextBlock to connect with for search
+	// This is similar to how SMenuEntryBlock works
+	FText SearchLabel = WidgetBlock->Label;
+	TAttribute<FText> SearchHighlightText;
+
+	if( OwnerMultiBoxWidgetPinned->GetMultiBox()->GetType() == EMultiBoxType::Menu && OwnerMultiBoxWidgetPinned->GetSearchable())
+	{
+		SearchHighlightText.Bind(OwnerMultiBoxWidgetPinned.Get(), &SMultiBoxWidget::GetSearchText);
 
+		if (!bHasLabel)
+		{
+			TSharedRef<SWidget> TextBlock = FindTextBlockWidget(WidgetBlock->ContentWidget);
+			if (TextBlock != SNullWidget::NullWidget )
+			{
+				// Bind the search text to the widgets text to highlight
+				TSharedRef<STextBlock> TheTextBlock = StaticCastSharedRef<STextBlock>(TextBlock);
+				TheTextBlock->SetHighlightText(SearchHighlightText);
+
+				SearchLabel = TheTextBlock.Get().GetText();
+			}
+		}
 	}
 
 	// Add this widget to the search list of the multibox
-	OwnerMultiBoxWidgetPinned->AddElement(this->AsWidget(), WidgetBlock->Label, MultiBlock->GetSearchable());
+	OwnerMultiBoxWidgetPinned->AddElement(this->AsWidget(), SearchLabel, MultiBlock->GetSearchable());
 
 	// This widget holds the search text, set it as the search block widget
 	if (OwnerMultiBoxWidgetPinned->GetSearchTextWidget() == WidgetBlock->ContentWidget)
@@ -125,6 +150,7 @@ void SWidgetBlock::BuildMultiBlockWidget(const ISlateStyle* StyleSet, const FNam
 				SNew( STextBlock )
 				.TextStyle(LabelStyle)
 				.Text(WidgetBlock->Label)
+				.HighlightText(SearchHighlightText)
 				.ColorAndOpacity(FAppStyle::Get().GetSlateColor("Colors.ForegroundHover"))
 			]
 		]
@@ -136,3 +162,41 @@ void SWidgetBlock::BuildMultiBlockWidget(const ISlateStyle* StyleSet, const FNam
 		]
 	];
 }
+
+TSharedRef<SWidget> SWidgetBlock::FindTextBlockWidget( TSharedRef<SWidget> Content )
+{
+	if (Content->GetType() == FName(TEXT("STextBlock")))
+	{
+		return Content;
+	}
+
+	FChildren* Children = Content->GetChildren();
+	int32 NumChildren = Children->Num();
+
+	for( int32 Index = 0; Index < NumChildren; ++Index )
+	{
+		TSharedRef<SWidget> Found = FindTextBlockWidget( Children->GetChildAt( Index ));
+		if (Found != SNullWidget::NullWidget)
+		{
+			return Found;
+		}
+	}
+	return SNullWidget::NullWidget;
+}
+
+void SWidgetBlock::OnMouseEnter(const FGeometry& MyGeometry, const FPointerEvent& MouseEvent)
+{
+	// If this widget is in a Menu, attempt to close any other open submenus within that menu
+	TSharedPtr<SMultiBoxWidget> OwnerMultiBoxWidgetPinned = OwnerMultiBoxWidget.Pin();
+	if(OwnerMultiBoxWidgetPinned->GetMultiBox()->GetType() == EMultiBoxType::Menu)
+	{
+		const TSharedPtr< const SMenuAnchor > OpenedMenuAnchor( OwnerMultiBoxWidgetPinned->GetOpenMenu() );
+		if ( OpenedMenuAnchor.IsValid() && OpenedMenuAnchor->IsOpen() )
+		{
+			OwnerMultiBoxWidgetPinned->CloseSummonedMenus();
+		}
+	}
+
+	SMultiBlockBaseWidget::OnMouseEnter(MyGeometry, MouseEvent);
+}
+
diff --git a/Engine/Source/Runtime/Slate/Private/Framework/MultiBox/SWidgetBlock.h b/Engine/Source/Runtime/Slate/Private/Framework/MultiBox/SWidgetBlock.h
--- a/Engine/Source/Runtime/Slate/Private/Framework/MultiBox/SWidgetBlock.h
+++ b/Engine/Source/Runtime/Slate/Private/Framework/MultiBox/SWidgetBlock.h
@@ -79,4 +79,17 @@ public:
 	 * @param	InArgs	The declaration data for this widget
 	 */
 	void Construct( const FArguments& InArgs );
+
+	
+	virtual void OnMouseEnter( const FGeometry& MyGeometry, const FPointerEvent& MouseEvent ) override;
+
+protected:
+
+	/**
+	* Finds the STextBlock that gets displayed in the UI
+	*
+	* @param Content	Widget to check for an STextBlock
+	* @return	The STextBlock widget found
+	*/
+	TSharedRef<SWidget> FindTextBlockWidget(TSharedRef<SWidget> Content);
 };
