Fix browse to content browser sometimes not working in asset editors. In UE4, the toolbar was a tab so clicking on the browse button would force the toolbar tab to be active. Now that the toolbars are not tabs you could get into a situation where no tab is activated so if the content browser was previously activated it would never switch because it thought it was already activated.

diff --git a/Engine/Source/Runtime/Slate/Private/Framework/Docking/TabManager.cpp b/Engine/Source/Runtime/Slate/Private/Framework/Docking/TabManager.cpp
--- a/Engine/Source/Runtime/Slate/Private/Framework/Docking/TabManager.cpp
+++ b/Engine/Source/Runtime/Slate/Private/Framework/Docking/TabManager.cpp
@@ -1314,7 +1314,14 @@ TSharedPtr<SDockTab> FTabManager::InvokeTab_Internal( const FTabId& TabId )
 
 		if (ExistingTab.IsValid())
 		{
-			if (!ExistingTab->IsActive())
+			TSharedPtr<SDockTab> MajorTab = FGlobalTabmanager::Get()->GetMajorTabForTabManager(ExistingTab->GetTabManager());
+
+			// Rules for drawing attention to a tab:
+			// 1. Tab is not active
+			// 2. Tab's owning major tab is not in the foreground (making the tab we want to draw attention to is not visible)
+			// 3. Tab is nomad and is not in the foreground
+			// If the tab is not active or the tabs major tab is not in the foreground, activate it
+			if (!ExistingTab->IsActive() || (MajorTab && !MajorTab->IsForeground()) || !ExistingTab->IsForeground())
 			{
 				// Draw attention to this tab if it didn't already have focus
 				DrawAttention(ExistingTab.ToSharedRef());
