Fixed bug where a recurring "unrecognized tab" could show up in Slate editors (typically seen in the BP editor), and also where an "unrecognized tab" could block the user from their intended functionality (e.g Find Results). Unrecognized tabs are no longer saved in the serialized layout data. Also improved log messaging to help track down what causes an unrecognized tab in the first place.

diff --git a/Engine/Source/Editor/UnrealEd/Private/ApplicationMode.cpp b/Engine/Source/Editor/UnrealEd/Private/ApplicationMode.cpp
--- a/Engine/Source/Editor/UnrealEd/Private/ApplicationMode.cpp
+++ b/Engine/Source/Editor/UnrealEd/Private/ApplicationMode.cpp
@@ -5,6 +5,8 @@
 #include "Framework/Docking/LayoutService.h"
 #include "UnrealEdMisc.h"
 
+DEFINE_LOG_CATEGORY_STATIC(LogApplicationMode, Log, All);
+
 #define LOCTEXT_NAMESPACE "ApplicationMode"
 
 /////////////////////////////////////////////////////
@@ -26,6 +28,8 @@ FApplicationMode::FApplicationMode(FName InModeName, FText(*GetLocalizedMode)(co
 
 void FApplicationMode::DeactivateMode(TSharedPtr<FTabManager> InTabManager)
 {
+	UE_LOG(LogApplicationMode, Log, TEXT("Deactivating ApplicationMode '%s'"), *ModeName.ToString());
+
 	// Save the layout to INI
 	check(InTabManager.IsValid());
 	if (FUnrealEdMisc::Get().IsSavingLayoutOnClosedAllowed())
@@ -45,6 +49,8 @@ void FApplicationMode::DeactivateMode(TSharedPtr<FTabManager> InTabManager)
 
 TSharedRef<FTabManager::FLayout> FApplicationMode::ActivateMode(TSharedPtr<FTabManager> InTabManager)
 {
+	UE_LOG(LogApplicationMode, Log, TEXT("Activating ApplicationMode '%s'"), *ModeName.ToString());
+
 	check(InTabManager.IsValid());
 	RegisterTabFactories(InTabManager);
 
diff --git a/Engine/Source/Runtime/Slate/Private/Framework/Docking/TabManager.cpp b/Engine/Source/Runtime/Slate/Private/Framework/Docking/TabManager.cpp
--- a/Engine/Source/Runtime/Slate/Private/Framework/Docking/TabManager.cpp
+++ b/Engine/Source/Runtime/Slate/Private/Framework/Docking/TabManager.cpp
@@ -436,22 +436,30 @@ TSharedRef<FJsonObject> FTabManager::FLayout::PersistToString_Helper(const TShar
 	{
 		JsonObj->SetStringField( TEXT("Type"), TEXT("Stack") );
 		JsonObj->SetBoolField( TEXT("HideTabWell"), NodeAsStack->bHideTabWell );
-		JsonObj->SetStringField( TEXT("ForegroundTab"), NodeAsStack->ForegroundTabId.ToString() );
+
+		if (NodeAsStack->ForegroundTabId.ShouldSaveLayout())
+		{
+			JsonObj->SetStringField(TEXT("ForegroundTab"), NodeAsStack->ForegroundTabId.ToString());
+		}
+		
 
 		TArray< TSharedPtr<FJsonValue> > TabsAsJson;
 		for(const FTab& Tab : NodeAsStack->Tabs)
 		{
-			TSharedRef<FJsonObject> TabAsJson = MakeShareable( new FJsonObject() );
-			TabAsJson->SetStringField( TEXT("TabId"), Tab.TabId.ToString() );
-			TabAsJson->SetStringField(TEXT("TabState"), StringFromTabState(Tab.TabState));
+			if (Tab.TabId.ShouldSaveLayout())
+			{	
+				TSharedRef<FJsonObject> TabAsJson = MakeShareable( new FJsonObject() );
+				TabAsJson->SetStringField( TEXT("TabId"), Tab.TabId.ToString() );
+				TabAsJson->SetStringField(TEXT("TabState"), StringFromTabState(Tab.TabState));
 
-			if (Tab.TabState == ETabState::SidebarTab && Tab.SidebarLocation != ESidebarLocation::None)
-			{
-				TabAsJson->SetStringField(TEXT("SidebarLocation"), StringFromSidebarLocation(Tab.SidebarLocation));
-				TabAsJson->SetNumberField(TEXT("SidebarCoeff"), Tab.SidebarSizeCoefficient);
-			}
+				if (Tab.TabState == ETabState::SidebarTab && Tab.SidebarLocation != ESidebarLocation::None)
+				{
+					TabAsJson->SetStringField(TEXT("SidebarLocation"), StringFromSidebarLocation(Tab.SidebarLocation));
+					TabAsJson->SetNumberField(TEXT("SidebarCoeff"), Tab.SidebarSizeCoefficient);
+				}
 
-			TabsAsJson.Add( MakeShareable( new FJsonValueObject(TabAsJson) ) );
+				TabsAsJson.Add( MakeShareable( new FJsonValueObject(TabAsJson) ) );
+			}
 		}
 		JsonObj->SetArrayField( TEXT("Tabs"), TabsAsJson );
 	}
@@ -854,6 +862,8 @@ void FTabManager::CloseAllAreas()
 
 TSharedRef<FTabManager::FLayout> FTabManager::PersistLayout() const
 {
+	UE_LOG(LogSlate, Log, TEXT("Saving tab layout to %s"), *ActiveLayoutName.ToString());
+
 	TSharedRef<FLayout> PersistentLayout = FTabManager::NewLayout( this->ActiveLayoutName );
 	
 	// Persist layout for all LiveAreas
@@ -908,6 +918,8 @@ FTabSpawnerEntry& FTabManager::RegisterTabSpawner(const FName TabId, const FOnSp
 	ensure(!TabSpawner.Contains(TabId));
 	ensure(!FGlobalTabmanager::Get()->IsLegacyTabType(TabId));
 
+	UE_LOG(LogSlate, Log, TEXT("Registering tab spawner for layout '%s': '%s'"), *ActiveLayoutName.ToString(), *TabId.ToString());
+
 	TSharedRef<FTabSpawnerEntry> NewSpawnerEntry = MakeShareable(new FTabSpawnerEntry(TabId, OnSpawnTab, CanSpawnTab));
 	TabSpawner.Add(TabId, NewSpawnerEntry);
 	return NewSpawnerEntry.Get();
@@ -915,11 +927,19 @@ FTabSpawnerEntry& FTabManager::RegisterTabSpawner(const FName TabId, const FOnSp
 
 bool FTabManager::UnregisterTabSpawner( const FName TabId )
 {
+	UE_LOG(LogSlate, Log, TEXT("Unregistering tab spawner for layout '%s': '%s'"), *ActiveLayoutName.ToString(), *TabId.ToString());
+
 	return TabSpawner.Remove( TabId ) > 0;
 }
 
 void FTabManager::UnregisterAllTabSpawners()
 {
+	if (TabSpawner.Num())
+	{
+		FString AllTabIdsString = FString::JoinBy(TabSpawner, TEXT(", "), [](const TPair<FName, TSharedRef<FTabSpawnerEntry>> Entry) { return Entry.Key.ToString(); });
+		UE_LOG(LogSlate, Log, TEXT("Unregistering all tab spawners for layout '%s': %s"), *ActiveLayoutName.ToString(), *AllTabIdsString);
+	}
+
 	TabSpawner.Empty();
 }
 
@@ -928,6 +948,8 @@ TSharedPtr<SWidget> FTabManager::RestoreFrom(const TSharedRef<FLayout>& Layout,
 {
 	ActiveLayoutName = Layout->LayoutName;
 
+	UE_LOG(LogSlate, Log, TEXT("Restoring tab layout: %s"), *ActiveLayoutName.ToString());
+
 	TSharedPtr<SDockingArea> PrimaryDockArea;
 	for (int32 AreaIndex=0; AreaIndex < Layout->Areas.Num(); ++AreaIndex )
 	{
@@ -1913,7 +1935,7 @@ TSharedPtr<SDockTab> FTabManager::SpawnTab(const FTabId& TabId, const TSharedPtr
 		if (!bCanOutputBeNullptr)
 		{
 			UE_LOG(LogSlate, Log,
-				TEXT("The tab \"%s\" attempted to spawn but failed for some reason. An \"unrecognized tab\" will be returned instead."), *StringToDisplay
+				TEXT("The tab \"%s\" attempted to spawn in layout '%s' but failed for some reason. An \"unrecognized tab\" will be returned instead."), *StringToDisplay, *ActiveLayoutName.ToString()
 			);
 
 			NewTabWidget = SNew(SDockTab)
@@ -1929,13 +1951,14 @@ TSharedPtr<SDockTab> FTabManager::SpawnTab(const FTabId& TabId, const TSharedPtr
 					]
 				];
 
-			NewTabWidget->SetLayoutIdentifier(TabId);
+			const FTabId UnrecognizedId(FName(TEXT("Unrecognized")), ETabIdFlags::None);
+			NewTabWidget->SetLayoutIdentifier(UnrecognizedId);
 		}
 		// If we can return nullptr, log it
 		else
 		{
 			UE_LOG(LogSlate, Log,
-				TEXT("The tab \"%s\" attempted to spawn but failed for some reason. It will not be displayed but it will still be saved in the layout settings file."), *StringToDisplay
+				TEXT("The tab \"%s\" attempted to spawn in layout '%s' but failed for some reason. It will not be displayed."), *StringToDisplay, *ActiveLayoutName.ToString()
 			);
 		}
 	}
@@ -2400,8 +2423,12 @@ FTabSpawnerEntry& FGlobalTabmanager::RegisterNomadTabSpawner(const FName TabId,
 	// Remove TabId if it was previously loaded. This allows re-loading the Editor UI layout without restarting the whole Editor (Window->Load Layout)
 	if (NomadTabSpawner->Contains(TabId))
 	{
+		UE_LOG(LogSlate, Log, TEXT("Unregistering nomad tab spawner for layout '%s': '%s'"), *ActiveLayoutName.ToString(), *TabId.ToString());
 		UnregisterNomadTabSpawner(TabId);
 	}
+
+	UE_LOG(LogSlate, Log, TEXT("Registering nomad tab spawner for layout '%s': '%s'"), *ActiveLayoutName.ToString(), *TabId.ToString());
+
 	// (Re)create and return NewSpawnerEntry
 	TSharedRef<FTabSpawnerEntry> NewSpawnerEntry = MakeShareable(new FTabSpawnerEntry(TabId, OnSpawnTab, CanSpawnTab));
 	NomadTabSpawner->Add(TabId, NewSpawnerEntry);
@@ -2410,6 +2437,7 @@ FTabSpawnerEntry& FGlobalTabmanager::RegisterNomadTabSpawner(const FName TabId,
 
 void FGlobalTabmanager::UnregisterNomadTabSpawner( const FName TabId )
 {
+	UE_LOG(LogSlate, Log, TEXT("Unregistering nomad tab spawner for layout '%s': '%s'"), *ActiveLayoutName.ToString(), *TabId.ToString());
 	const int32 NumRemoved = NomadTabSpawner->Remove(TabId);
 }
 
diff --git a/Engine/Source/Runtime/Slate/Public/Framework/Docking/TabManager.h b/Engine/Source/Runtime/Slate/Public/Framework/Docking/TabManager.h
--- a/Engine/Source/Runtime/Slate/Public/Framework/Docking/TabManager.h
+++ b/Engine/Source/Runtime/Slate/Public/Framework/Docking/TabManager.h
@@ -40,20 +40,38 @@ enum class ESidebarLocation : uint8
 	None,
 };
 
+
+enum class ETabIdFlags : uint8
+{
+	None            = 0x0,  // No flags
+	SaveLayout      = 0x1,  // This tab should be included when saving the Slate layout
+};
+
+ENUM_CLASS_FLAGS(ETabIdFlags);
+
 struct FTabId
 {
 	FTabId()
 		: InstanceId(INDEX_NONE)
+		, Flags(ETabIdFlags::SaveLayout)
 	{ }
 
 	FTabId(const FName InTabType, const int32 InInstanceId)
 		: TabType(InTabType)
 		, InstanceId(InInstanceId)
+		, Flags(ETabIdFlags::SaveLayout)
 	{ }
 
 	FTabId(const FName InTabType)
 		: TabType(InTabType)
 		, InstanceId(INDEX_NONE)
+		, Flags(ETabIdFlags::SaveLayout)
+	{ }
+
+	FTabId(const FName InTabType, const ETabIdFlags InFlags)
+		: TabType(InTabType)
+		, InstanceId(INDEX_NONE)
+		, Flags(InFlags)
 	{ }
 
 	/** Document tabs allow multiple instances of the same tab type. The placement rules for these tabs are left up for the specific use-cases. These tabs are not persisted. */
@@ -67,6 +85,8 @@ struct FTabId
 		return TabType == Other.TabType && (InstanceId == INDEX_NONE || Other.InstanceId == INDEX_NONE || InstanceId == Other.InstanceId) ;
 	}
 
+	bool ShouldSaveLayout() const { return EnumHasAnyFlags(Flags, ETabIdFlags::SaveLayout); }
+
 	FString ToString() const
 	{
 		// This function is useful to allow saving the layout on disk.
@@ -99,6 +119,9 @@ struct FTabId
 
 	FName TabType;
 	int32 InstanceId;
+
+	private:
+	ETabIdFlags Flags;
 };
 
 
