UMG: ComboBoxKey. Add a combobox to replace ComboBoxString. String are not localized. The key version allow the user to build a widget from a TMap or a DataTable in BP.

diff --git a/Engine/Source/Runtime/Slate/Public/Framework/Views/TableViewTypeTraits.h b/Engine/Source/Runtime/Slate/Public/Framework/Views/TableViewTypeTraits.h
--- a/Engine/Source/Runtime/Slate/Public/Framework/Views/TableViewTypeTraits.h
+++ b/Engine/Source/Runtime/Slate/Public/Framework/Views/TableViewTypeTraits.h
@@ -87,6 +87,14 @@ struct TIsValidListItem<TFieldPath<T>>
 		Value = true
 	};
 };
+template <>
+struct TIsValidListItem<FName>
+{
+	enum
+	{
+		Value = true
+	};
+};
 
 /**
  * Furthermore, ListViews of TSharedPtr<> work differently from lists of UObject*.
@@ -95,12 +103,12 @@ struct TIsValidListItem<TFieldPath<T>>
  */
 template <typename T, typename Enable=void> struct TListTypeTraits
 {
-	static_assert(TIsValidListItem<T>::Value, "Item type T must be a UObjectBase pointer, TFieldPath, TSharedRef, or TSharedPtr.");
+	static_assert(TIsValidListItem<T>::Value, "Item type T must be a UObjectBase pointer, TFieldPath, TSharedRef, TSharedPtr or FName.");
 };
 
 
 /**
- * Pointer-related functionality (e.g. setting to null, testing for null) specialized for SharedPointers.
+ * Pointer-related functionality (e.g. setting to nullptr, testing for nullptr) specialized for SharedPointers.
  */
 template <typename T> struct TListTypeTraits< TSharedPtr<T, ESPMode::NotThreadSafe> >
 {
@@ -131,7 +139,7 @@ public:
 
 	static TSharedPtr<T> MakeNullPtr()
 	{
-		return TSharedPtr<T>(NULL);
+		return TSharedPtr<T>(nullptr);
 	}
 
 	static TSharedPtr<T> NullableItemTypeConvertToItemType( const TSharedPtr<T>& InPtr )
@@ -177,7 +185,7 @@ public:
 
 	static TSharedPtr<T, ESPMode::ThreadSafe> MakeNullPtr()
 	{
-		return TSharedPtr<T, ESPMode::ThreadSafe>(NULL);
+		return TSharedPtr<T, ESPMode::ThreadSafe>(nullptr);
 	}
 
 	static TSharedPtr<T, ESPMode::ThreadSafe> NullableItemTypeConvertToItemType( const TSharedPtr<T, ESPMode::ThreadSafe>& InPtr )
@@ -223,7 +231,7 @@ public:
 
 	static TSharedPtr<T> MakeNullPtr()
 	{
-		return TSharedPtr<T>(NULL);
+		return TSharedPtr<T>(nullptr);
 	}
 
 	static TSharedRef<T> NullableItemTypeConvertToItemType( const TSharedPtr<T>& InPtr )
@@ -269,7 +277,7 @@ public:
 
 	static TSharedPtr<T, ESPMode::ThreadSafe> MakeNullPtr()
 	{
-		return TSharedPtr<T, ESPMode::ThreadSafe>(NULL);
+		return TSharedPtr<T, ESPMode::ThreadSafe>(nullptr);
 	}
 
 	static TSharedRef<T, ESPMode::ThreadSafe> NullableItemTypeConvertToItemType( const TSharedPtr<T, ESPMode::ThreadSafe>& InPtr )
@@ -287,7 +295,7 @@ public:
 
 
 /**
- * Pointer-related functionality (e.g. setting to null, testing for null) specialized for SharedPointers.
+ * Pointer-related functionality (e.g. setting to nullptr, testing for nullptr) specialized for SharedPointers.
  */
 template <typename T> struct TListTypeTraits< TWeakObjectPtr<T> >
 {
@@ -338,7 +346,7 @@ public:
 
 /**
  * Lists of pointer types only work if the pointers are deriving from UObject*.
- * In addition to testing and setting the pointers to null, Lists of UObjects
+ * In addition to testing and setting the pointers to nullptr, Lists of UObjects
  * will serialize the objects they are holding onto.
  */
 template <typename T>
@@ -371,11 +379,11 @@ public:
 		Collector.AddReferencedObjects(SelectedItems);
 	}
 
-	static bool IsPtrValid( T* InPtr ) { return InPtr != NULL; }
+	static bool IsPtrValid( T* InPtr ) { return InPtr != nullptr; }
 
-	static void ResetPtr( T*& InPtr ) { InPtr = NULL; }
+	static void ResetPtr( T*& InPtr ) { InPtr = nullptr; }
 
-	static T* MakeNullPtr() { return NULL; }
+	static T* MakeNullPtr() { return nullptr; }
 
 	static T* NullableItemTypeConvertToItemType( T* InPtr ) { return InPtr; }
 
@@ -417,11 +425,11 @@ public:
 		Collector.AddReferencedObjects(SelectedItems);
 	}
 
-	static bool IsPtrValid( const T* InPtr ) { return InPtr != NULL; }
+	static bool IsPtrValid( const T* InPtr ) { return InPtr != nullptr; }
 
-	static void ResetPtr( const T*& InPtr ) { InPtr = NULL; }
+	static void ResetPtr( const T*& InPtr ) { InPtr = nullptr; }
 
-	static const T* MakeNullPtr() { return NULL; }
+	static const T* MakeNullPtr() { return nullptr; }
 
 	static const T* NullableItemTypeConvertToItemType( const T* InPtr ) { return InPtr; }
 
@@ -435,7 +443,7 @@ public:
 
 
 /**
- * Pointer-related functionality (e.g. setting to null, testing for null) specialized for TFieldPaths.
+ * Pointer-related functionality (e.g. setting to nullptr, testing for nullptr) specialized for TFieldPaths.
  */
 template <typename T> struct TListTypeTraits< TFieldPath<T> >
 {
@@ -507,4 +515,53 @@ public:
 	}
 
 	class SerializerType {};
-};
\ No newline at end of file
+};
+
+/**
+ * Functionality (e.g. setting the value, testing invalid value) specialized for FName.
+ */
+template <>
+struct TListTypeTraits<FName>
+{
+public:
+	typedef FName NullableType;
+
+	using MapKeyFuncs = TDefaultMapHashableKeyFuncs<FName, TSharedRef<ITableRow>, false>;
+	using MapKeyFuncsSparse = TDefaultMapHashableKeyFuncs<FName, FSparseItemInfo, false>;
+	using SetKeyFuncs = DefaultKeyFuncs<FName>;
+
+	template<typename U>
+	static void AddReferencedObjects(FReferenceCollector& Collector,
+		TArray<FName>& ItemsWithGeneratedWidgets,
+		TSet<FName>& SelectedItems,
+		TMap<const U*, FName>& WidgetToItemMap)
+	{
+	}
+
+	static bool IsPtrValid(const FName& InValue)
+	{
+		return !InValue.IsNone();
+	}
+
+	static void ResetPtr(FName& InValue)
+	{
+		InValue = FName();
+	}
+
+	static FName MakeNullPtr()
+	{
+		return FName();
+	}
+
+	static FName NullableItemTypeConvertToItemType(const FName& InValue)
+	{
+		return InValue;
+	}
+
+	static FString DebugDump(FName InValue)
+	{
+		return InValue.ToString();
+	}
+
+	class SerializerType {};
+};
diff --git a/Engine/Source/Runtime/Slate/Public/Widgets/Input/SComboBox.h b/Engine/Source/Runtime/Slate/Public/Widgets/Input/SComboBox.h
--- a/Engine/Source/Runtime/Slate/Public/Widgets/Input/SComboBox.h
+++ b/Engine/Source/Runtime/Slate/Public/Widgets/Input/SComboBox.h
@@ -97,6 +97,7 @@ class SComboBox : public SComboButton
 {
 public:
 
+	typedef TListTypeTraits< OptionType > ListTypeTraits;
 	typedef typename TListTypeTraits< OptionType >::NullableType NullableOptionType;
 
 	/** Type of list used for showing menu options. */
@@ -115,7 +116,7 @@ public:
 		, _OptionsSource()
 		, _OnSelectionChanged()
 		, _OnGenerateWidget()
-		, _InitiallySelectedItem(nullptr)
+		, _InitiallySelectedItem(ListTypeTraits::MakeNullPtr())
 		, _Method()
 		, _MaxListHeight(450.0f)
 		, _HasDownArrow( true )
diff --git a/Engine/Source/Runtime/Slate/Public/Widgets/Views/SListView.h b/Engine/Source/Runtime/Slate/Public/Widgets/Views/SListView.h
--- a/Engine/Source/Runtime/Slate/Public/Widgets/Views/SListView.h
+++ b/Engine/Source/Runtime/Slate/Public/Widgets/Views/SListView.h
@@ -266,15 +266,15 @@ public:
 		}
 	}
 
-	SListView( ETableViewMode::Type InListMode = ETableViewMode::List )
-		: STableViewBase( InListMode )
+	SListView(ETableViewMode::Type InListMode = ETableViewMode::List)
+		: STableViewBase(InListMode)
 		, WidgetGenerator(this)
-		, SelectorItem( NullableItemType(nullptr) )
-		, RangeSelectionStart( NullableItemType(nullptr) )
-		, ItemsSource( nullptr )
-		, ItemToScrollIntoView( NullableItemType(nullptr) )
-		, UserRequestingScrollIntoView( 0 )
-		, ItemToNotifyWhenInView( NullableItemType(nullptr) ) 
+		, SelectorItem(TListTypeTraits<ItemType>::MakeNullPtr())
+		, RangeSelectionStart(TListTypeTraits<ItemType>::MakeNullPtr())
+		, ItemsSource(nullptr)
+		, ItemToScrollIntoView(TListTypeTraits<ItemType>::MakeNullPtr())
+		, UserRequestingScrollIntoView(0)
+		, ItemToNotifyWhenInView(TListTypeTraits<ItemType>::MakeNullPtr())
 		, IsFocusable(true)
 	{ 
 #if WITH_ACCESSIBILITY
@@ -309,7 +309,7 @@ public:
 		if ( ItemsSourceRef.Num() > 0 && !InKeyEvent.IsAltDown() )
 		{
 			bool bWasHandled = false;
-			NullableItemType ItemNavigatedTo( nullptr );
+			NullableItemType ItemNavigatedTo = TListTypeTraits<ItemType>::MakeNullPtr();
 
 			// Check for selection manipulation keys (Up, Down, Home, End, PageUp, PageDown)
 			if ( InKeyEvent.GetKey() == EKeys::Home )
@@ -2028,7 +2028,7 @@ protected:
 
 		if (OnIsSelectableOrNavigable.IsBound())
 		{
-			NullableItemType LastSelectedItem = nullptr;
+			NullableItemType LastSelectedItem = TListTypeTraits<ItemType>::MakeNullPtr();
 			if (SelectedItems.Num() == 1)
 			{
 				LastSelectedItem = *SelectedItems.CreateIterator();
diff --git a/Engine/Source/Runtime/UMG/Private/Components/ComboBoxKey.cpp b/Engine/Source/Runtime/UMG/Private/Components/ComboBoxKey.cpp
new file mode 100644
--- /dev/null
+++ b/Engine/Source/Runtime/UMG/Private/Components/ComboBoxKey.cpp
@@ -0,0 +1,251 @@
+// Copyright Epic Games, Inc. All Rights Reserved.
+
+#include "Components/ComboBoxKey.h"
+
+#include "UObject/ConstructorHelpers.h"
+#include "Widgets/Layout/SBox.h"
+#include "Widgets/SNullWidget.h"
+#include "Widgets/Text/STextBlock.h"
+
+#define LOCTEXT_NAMESPACE "UMG"
+
+/////////////////////////////////////////////////////
+// UComboBoxKey
+
+UComboBoxKey::UComboBoxKey()
+{
+	// We don't want to try and load fonts on the server.
+	if (!IsRunningDedicatedServer())
+	{
+		WidgetStyle = FCoreStyle::Get().GetWidgetStyle<FComboBoxStyle>("ComboBox");
+		ItemStyle = FCoreStyle::Get().GetWidgetStyle<FTableRowStyle>("TableView.Row");
+	}
+
+	ContentPadding = FMargin(4.0, 2.0);
+	ForegroundColor = FLinearColor::Black;
+
+	MaxListHeight = 450.0f;
+	bHasDownArrow = true;
+	bEnableGamepadNavigationMode = true;
+	bIsFocusable = true;
+}
+
+
+void UComboBoxKey::ReleaseSlateResources(bool bReleaseChildren)
+{
+	Super::ReleaseSlateResources(bReleaseChildren);
+
+	MyComboBox.Reset();
+	ComboBoxContent.Reset();
+}
+
+
+TSharedRef<SWidget> UComboBoxKey::RebuildWidget()
+{
+	MyComboBox =
+		SNew(SComboBox<FName>)
+		.ComboBoxStyle(&WidgetStyle)
+		.ItemStyle(&ItemStyle)
+		.ForegroundColor(ForegroundColor)
+		.OptionsSource(&Options)
+		.InitiallySelectedItem(SelectedOption)
+		.ContentPadding(ContentPadding)
+		.MaxListHeight(MaxListHeight)
+		.HasDownArrow(bHasDownArrow)
+		.EnableGamepadNavigationMode(bEnableGamepadNavigationMode)
+		.OnGenerateWidget_UObject(this, &UComboBoxKey::HandleGenerateItemWidget)
+		.OnSelectionChanged_UObject(this, &UComboBoxKey::HandleSelectionChanged)
+		.OnComboBoxOpening_UObject(this, &UComboBoxKey::HandleOpening)
+		.IsFocusable(bIsFocusable)
+		[
+			SAssignNew(ComboBoxContent, SBox)
+		];
+
+	GenerateContent();
+
+	return MyComboBox.ToSharedRef();
+}
+
+
+void UComboBoxKey::AddOption(FName Key)
+{
+	if (Options.AddUnique(Key) >= 0)
+	{
+		if (MyComboBox)
+		{
+			MyComboBox->RefreshOptions();
+		}
+	}
+}
+
+
+bool UComboBoxKey::RemoveOption(FName Key)
+{
+	bool bResult = Options.RemoveSingle(Key) > 0;
+	if (bResult)
+	{
+		if (Key == SelectedOption)
+		{
+			ClearSelection();
+		}
+
+		if (MyComboBox)
+		{
+			MyComboBox->RefreshOptions();
+		}
+	}
+	return bResult;
+}
+
+
+void UComboBoxKey::ClearOptions()
+{
+	Options.Reset();
+	ClearSelection();
+	if (MyComboBox)
+	{
+		MyComboBox->RefreshOptions();
+	}
+}
+
+
+void UComboBoxKey::ClearSelection()
+{
+	if (MyComboBox)
+	{
+		MyComboBox->ClearSelection();
+	}
+}
+
+
+void UComboBoxKey::SetSelectedOption(FName Option)
+{
+	if (SelectedOption != Option)
+	{
+		if (MyComboBox)
+		{
+			MyComboBox->SetSelectedItem(Option);
+		}
+		else
+		{
+			SelectedOption = Option;
+		}
+	}
+}
+
+
+FName UComboBoxKey::GetSelectedOption() const
+{
+	return SelectedOption;
+}
+
+
+bool UComboBoxKey::IsOpen() const
+{
+	return MyComboBox && MyComboBox->IsOpen();
+}
+
+
+void UComboBoxKey::GenerateContent()
+{
+	bool bSetToNull = true;
+	if (IsDesignTime())
+	{
+		if (!SelectedOption.IsNone())
+		{
+			ComboBoxContent->SetContent(SNew(STextBlock)
+				.Text(FText::FromName(SelectedOption)));
+			bSetToNull = false;
+		}
+	}
+	else
+	{
+		// Call the user's delegate to see if they want to generate a custom widget bound to the data source.
+		if (OnGenerateContentWidget.IsBound())
+		{
+			if (UWidget* Widget = OnGenerateContentWidget.Execute(SelectedOption))
+			{
+				ComboBoxContent->SetContent(Widget->TakeWidget());
+				bSetToNull = false;
+			}
+		}
+		else
+		{
+			// Warn the user that he needs to generate a widget for the item
+			ComboBoxContent->SetContent(
+				SNew(STextBlock)
+				.Text(FText::Format(LOCTEXT("WarningMissingOnGenerateContentWidget", "{0} Missing OnGenerateContentWidget"), FText::FromName(SelectedOption)))
+				);
+			bSetToNull = false;
+		}
+	}
+
+	if (bSetToNull)
+	{
+		ComboBoxContent->SetContent(SNullWidget::NullWidget);
+	}
+}
+
+
+TSharedRef<SWidget> UComboBoxKey::HandleGenerateItemWidget(FName Item)
+{
+	if (IsDesignTime())
+	{
+		return SNew(STextBlock)
+			.Text(FText::FromName(Item));
+	}
+	else
+	{
+		// Call the user's delegate to see if they want to generate a custom widget bound to the data source.
+		if (OnGenerateItemWidget.IsBound())
+		{
+			if (UWidget* Widget = OnGenerateItemWidget.Execute(Item))
+			{
+				return Widget->TakeWidget();
+			}
+			else
+			{
+				return SNullWidget::NullWidget;
+			}
+		}
+	}
+
+	return SNew(STextBlock)
+		.Text(FText::Format(LOCTEXT("WarningMissingOnGenerateItemWidget", "{0} Missing OnGenerateItemWidget"), FText::FromName(Item)));
+}
+
+
+void UComboBoxKey::HandleSelectionChanged(FName Item, ESelectInfo::Type SelectionType)
+{
+	if (SelectedOption != Item)
+	{
+		SelectedOption = Item;
+
+		if (!IsDesignTime())
+		{
+			OnSelectionChanged.Broadcast(Item, SelectionType);
+		}
+		GenerateContent();
+	}
+}
+
+void UComboBoxKey::HandleOpening()
+{
+	if (!IsDesignTime())
+	{
+		OnOpening.Broadcast();
+	}
+}
+
+#if WITH_EDITOR
+
+const FText UComboBoxKey::GetPaletteCategory()
+{
+	return LOCTEXT("Input", "Input");
+}
+
+#endif
+
+/////////////////////////////////////////////////////
+
+#undef LOCTEXT_NAMESPACE
diff --git a/Engine/Source/Runtime/UMG/Private/UMGStyle.cpp b/Engine/Source/Runtime/UMG/Private/UMGStyle.cpp
--- a/Engine/Source/Runtime/UMG/Private/UMGStyle.cpp
+++ b/Engine/Source/Runtime/UMG/Private/UMGStyle.cpp
@@ -70,6 +70,7 @@ TSharedRef< FSlateStyleSet > FUMGStyle::Create()
 	Style->Set("ClassIcon.MultiLineEditableTextBox", new IMAGE_BRUSH(TEXT("MultiLineEditableTextBox"), Icon16x16));
 	Style->Set("ClassIcon.Viewport", new IMAGE_BRUSH(TEXT("Viewport"), Icon16x16));
 	Style->Set("ClassIcon.ComboBox", new IMAGE_BRUSH(TEXT("icon_umg_ComboBox_16x"), Icon16x16));
+	Style->Set("ClassIcon.ComboBoxKey", new IMAGE_BRUSH(TEXT("icon_umg_ComboBox_16x"), Icon16x16));
 	Style->Set("ClassIcon.ComboBoxString", new IMAGE_BRUSH(TEXT("icon_umg_ComboBox_16x"), Icon16x16));
 	Style->Set("ClassIcon.ListView", new IMAGE_BRUSH(TEXT("icon_umg_ListView_16x"), Icon16x16));
 	Style->Set("ClassIcon.TileView", new IMAGE_BRUSH(TEXT("icon_umg_TileView_16x"), Icon16x16));
diff --git a/Engine/Source/Runtime/UMG/Public/Components/ComboBoxKey.h b/Engine/Source/Runtime/UMG/Public/Components/ComboBoxKey.h
new file mode 100644
--- /dev/null
+++ b/Engine/Source/Runtime/UMG/Public/Components/ComboBoxKey.h
@@ -0,0 +1,157 @@
+// Copyright Epic Games, Inc. All Rights Reserved.
+
+#pragma once
+
+#include "CoreMinimal.h"
+#include "UObject/ObjectMacros.h"
+#include "Components/Widget.h"
+#include "Widgets/Input/SComboBox.h"
+
+#include "ComboBoxKey.generated.h"
+
+
+/**
+ * The combobox allows you to display a list of options to the user in a dropdown menu for them to select one.
+ * Use OnGenerateConentWidgetEvent to return a custom built widget.
+ */
+UCLASS(meta = (DisplayName = "ComboBox (Key)"))
+class UMG_API UComboBoxKey : public UWidget
+{
+	GENERATED_BODY()
+
+public:
+	UComboBoxKey();
+
+	DECLARE_DYNAMIC_MULTICAST_DELEGATE_TwoParams(FOnSelectionChangedEvent, FName, SelectedItem, ESelectInfo::Type, SelectionType);
+	DECLARE_DYNAMIC_MULTICAST_DELEGATE(FOnOpeningEvent);
+	DECLARE_DYNAMIC_DELEGATE_RetVal_OneParam(UWidget*, FGenerateWidgetEvent, FName, Item);
+
+private:
+	/** . */
+	UPROPERTY(EditAnywhere, Category = Content)
+	TArray<FName> Options;
+
+	/** */
+	UPROPERTY(EditAnywhere, Category = Content)
+	FName SelectedOption;
+
+public:
+	/** The combobox style. */
+	UPROPERTY(EditAnywhere, BlueprintReadOnly, Category = Style, meta = (DisplayName = "Style"))
+	FComboBoxStyle WidgetStyle;
+
+	/** The item row style. */
+	UPROPERTY(EditAnywhere, BlueprintReadOnly, Category = Style)
+	FTableRowStyle ItemStyle;
+
+	/** The foreground color to pass through the hierarchy. */
+	UPROPERTY(EditAnywhere, BlueprintReadOnly, Category = Style, meta = (DesignerRebuild))
+	FSlateColor ForegroundColor;
+
+	UPROPERTY(EditAnywhere, BlueprintReadOnly, Category = Style)
+	FMargin ContentPadding;
+
+	/** The max height of the combobox list that opens */
+	UPROPERTY(EditAnywhere, BlueprintReadOnly, Category = Style, AdvancedDisplay)
+	float MaxListHeight;
+
+	/**
+	 * When false, the down arrow is not generated and it is up to the API consumer
+	 * to make their own visual hint that this is a drop down.
+	 */
+	UPROPERTY(EditAnywhere, BlueprintReadOnly, Category = Style, AdvancedDisplay)
+	bool bHasDownArrow;
+
+	/**
+	 * When false, directional keys will change the selection. When true, ComboBox
+	 * must be activated and will only capture arrow input while activated.
+	 */
+	UPROPERTY(EditAnywhere, BlueprintReadOnly, Category = Style, AdvancedDisplay)
+	bool bEnableGamepadNavigationMode;
+
+	/** When true, allows the combo box to receive keyboard focus */
+	UPROPERTY(EditAnywhere, BlueprintReadOnly, Category = Style)
+	bool bIsFocusable;
+
+public: // Events
+
+	/** Called when the widget is needed for the content. */
+	UPROPERTY(EditAnywhere, Category = Events, meta = (IsBindableEvent = "True"))
+	FGenerateWidgetEvent OnGenerateContentWidget;
+	
+	/** Called when the widget is needed for the item. */
+	UPROPERTY(EditAnywhere, Category = Events, meta = (IsBindableEvent = "True"))
+	FGenerateWidgetEvent OnGenerateItemWidget;
+
+	/** Called when a new item is selected in the combobox. */
+	UPROPERTY(BlueprintAssignable, Category = Events)
+	FOnSelectionChangedEvent OnSelectionChanged;
+
+	/** Called when the combobox is opening */
+	UPROPERTY(BlueprintAssignable, Category = Events)
+	FOnOpeningEvent OnOpening;
+
+public:
+
+	
+
+	/** Add an element to the option list. */
+	UFUNCTION(BlueprintCallable, Category = "ComboBox")
+	void AddOption(FName Option);
+
+	/** Remove an element to the option list. */
+	UFUNCTION(BlueprintCallable, Category = "ComboBox")
+	bool RemoveOption(FName Option);
+
+	/** Remove all the elements of the option list. */
+	UFUNCTION(BlueprintCallable, Category = "ComboBox")
+	void ClearOptions();
+
+	/** Clear the current selection. */
+	UFUNCTION(BlueprintCallable, Category = "ComboBox")
+	void ClearSelection();
+
+	/** Set the current selected option. */
+	UFUNCTION(BlueprintCallable, Category = "ComboBox")
+	void SetSelectedOption(FName Option);
+
+	/** Get the current selected option */
+	UFUNCTION(BlueprintCallable, Category = "ComboBox")
+	FName GetSelectedOption() const;
+
+	/** Is the combobox menu openned. */
+	UFUNCTION(BlueprintCallable, Category = "ComboBox", Meta = (ReturnDisplayName = "bOpen"))
+	bool IsOpen() const;
+
+	//~ Begin UVisual Interface
+	virtual void ReleaseSlateResources(bool bReleaseChildren) override;
+	//~ End UVisual Interface
+
+#if WITH_EDITOR
+	virtual const FText GetPaletteCategory() override;
+#endif
+
+private:
+	/** Called by slate when it needs to generate the widget in the content box */
+	void GenerateContent();
+
+	/** Called by slate when it needs to generate a new item for the combobox */
+	TSharedRef<SWidget> HandleGenerateItemWidget(FName Item);
+
+	/** Called by slate when the underlying combobox selection changes */
+	void HandleSelectionChanged(FName Item, ESelectInfo::Type SelectionType);
+
+	/** Called by slate when the underlying combobox is opening */
+	void HandleOpening();
+
+	//~ Begin UWidget Interface
+	virtual TSharedRef<SWidget> RebuildWidget() override;
+	//~ End UWidget Interface
+
+private:
+	/** A shared pointer to the underlying slate combobox */
+	TSharedPtr<SComboBox<FName>> MyComboBox;
+
+	/** A shared pointer to a container that holds the combobox content that is selected */
+	TSharedPtr< SBox > ComboBoxContent;
+};
