Added "Slate.TestNotifications" to test slate notifications for reskin

diff --git a/Engine/Source/Runtime/Slate/Private/Framework/Notifications/NotificationManager.cpp b/Engine/Source/Runtime/Slate/Private/Framework/Notifications/NotificationManager.cpp
--- a/Engine/Source/Runtime/Slate/Private/Framework/Notifications/NotificationManager.cpp
+++ b/Engine/Source/Runtime/Slate/Private/Framework/Notifications/NotificationManager.cpp
@@ -250,3 +250,172 @@ void FSlateNotificationManager::ForceNotificationsInFront( const TSharedRef<SWin
 		}
 	}
 }
+
+#if (WITH_EDITOR || IS_PROGRAM) && !UE_BUILD_SHIPPING
+
+static void TestNotifications()
+{
+	FTicker::GetCoreTicker().AddTicker(TEXT("TestNotifications"), 0.0f, [](float DeltaTime)
+		{
+			FNotificationInfo NotificationInfo(FText::FromString(TEXT("Test Notification 1")));
+			NotificationInfo.FadeInDuration = 2.0f;
+			NotificationInfo.FadeOutDuration = 2.0f;
+			NotificationInfo.ExpireDuration = 10.0f;
+
+			auto Notification = FSlateNotificationManager::Get().AddNotification(NotificationInfo);
+			Notification->ExpireAndFadeout();
+
+			return false;
+		});
+
+	FTicker::GetCoreTicker().AddTicker(TEXT("TestNotifications"), 1.0f, [](float DeltaTime)
+		{
+			FNotificationInfo NotificationInfo(FText::FromString(TEXT("Test Notification 2")));
+			NotificationInfo.FadeInDuration = 2.0f;
+			NotificationInfo.FadeOutDuration = 2.0f;
+			NotificationInfo.ExpireDuration = 10.0f;
+
+			NotificationInfo.bUseLargeFont = false;
+			auto Notification = FSlateNotificationManager::Get().AddNotification(NotificationInfo);
+			Notification->ExpireAndFadeout();
+
+			return false;
+		});
+
+	FTicker::GetCoreTicker().AddTicker(TEXT("TestNotifications"), 2.0f, [](float DeltaTime)
+		{
+			FNotificationInfo NotificationInfo(FText::FromString(TEXT("Test Notification 3")));
+			NotificationInfo.FadeInDuration = 2.0f;
+			NotificationInfo.FadeOutDuration = 2.0f;
+			NotificationInfo.ExpireDuration = 10.0f;
+
+			NotificationInfo.bUseThrobber = true;
+			auto Notification = FSlateNotificationManager::Get().AddNotification(NotificationInfo);
+			Notification->SetCompletionState(SNotificationItem::CS_Pending);
+			Notification->ExpireAndFadeout();
+
+			return false;
+		});
+
+	FTicker::GetCoreTicker().AddTicker(TEXT("TestNotifications"), 3.0f, [](float DeltaTime)
+		{
+			FNotificationInfo NotificationInfo(FText::FromString(TEXT("Test Notification 4")));
+			NotificationInfo.FadeInDuration = 2.0f;
+			NotificationInfo.FadeOutDuration = 2.0f;
+			NotificationInfo.ExpireDuration = 10.0f;
+
+			NotificationInfo.bUseSuccessFailIcons = true;
+			auto Notification = FSlateNotificationManager::Get().AddNotification(NotificationInfo);
+			Notification->SetCompletionState(SNotificationItem::CS_Success);
+			Notification->ExpireAndFadeout();
+
+			return false;
+		});
+
+	FTicker::GetCoreTicker().AddTicker(TEXT("TestNotifications"), 4.0f, [](float DeltaTime)
+		{
+			FNotificationInfo NotificationInfo(FText::FromString(TEXT("Test Notification 5")));
+			NotificationInfo.FadeInDuration = 2.0f;
+			NotificationInfo.FadeOutDuration = 2.0f;
+			NotificationInfo.ExpireDuration = 10.0f;
+
+			NotificationInfo.bUseSuccessFailIcons = true;
+			auto Notification = FSlateNotificationManager::Get().AddNotification(NotificationInfo);
+			Notification->SetCompletionState(SNotificationItem::CS_Fail);
+			Notification->ExpireAndFadeout();
+
+			return false;
+		});
+
+	FTicker::GetCoreTicker().AddTicker(TEXT("TestNotifications"), 5.0f, [](float DeltaTime)
+		{
+			FNotificationInfo NotificationInfo(FText::FromString(TEXT("Test Notification 6")));
+			NotificationInfo.FadeInDuration = 2.0f;
+			NotificationInfo.FadeOutDuration = 2.0f;
+			NotificationInfo.ExpireDuration = 10.0f;
+
+			NotificationInfo.CheckBoxText = FText::FromString(TEXT("Don't ask again"));
+			NotificationInfo.CheckBoxState = ECheckBoxState::Checked;
+			NotificationInfo.CheckBoxStateChanged = FOnCheckStateChanged::CreateStatic([](ECheckBoxState NewState) {});
+
+			auto Notification = FSlateNotificationManager::Get().AddNotification(NotificationInfo);
+
+			Notification->ExpireAndFadeout();
+
+			return false;
+		});
+
+	FTicker::GetCoreTicker().AddTicker(TEXT("TestNotifications"), 6.0f, [](float DeltaTime)
+		{
+			FNotificationInfo NotificationInfo(FText::FromString(TEXT("Test Notification 7")));
+			NotificationInfo.FadeInDuration = 2.0f;
+			NotificationInfo.FadeOutDuration = 2.0f;
+			NotificationInfo.ExpireDuration = 10.0f;
+
+			NotificationInfo.Hyperlink = FSimpleDelegate::CreateLambda([]() {});
+			NotificationInfo.HyperlinkText = FText::FromString(TEXT("This is a hyperlink"));
+
+			auto Notification = FSlateNotificationManager::Get().AddNotification(NotificationInfo);
+
+			Notification->ExpireAndFadeout();
+
+			return false;
+		});
+
+	FTicker::GetCoreTicker().AddTicker(TEXT("TestNotifications"), 7.0f, [](float DeltaTime)
+		{
+			FNotificationInfo NotificationInfo(FText::FromString(TEXT("Test Notification 8")));
+			NotificationInfo.FadeInDuration = 2.0f;
+			NotificationInfo.FadeOutDuration = 2.0f;
+			NotificationInfo.ExpireDuration = 10.0f;
+
+			FNotificationButtonInfo Button1(FText::FromString("Ok"), FText::GetEmpty(), FSimpleDelegate(), SNotificationItem::ECompletionState::CS_None);
+			FNotificationButtonInfo Button2(FText::FromString("Cancel"), FText::GetEmpty(), FSimpleDelegate(), SNotificationItem::ECompletionState::CS_None);
+
+			NotificationInfo.ButtonDetails.Add(Button1);
+			NotificationInfo.ButtonDetails.Add(Button2);
+
+			auto Notification = FSlateNotificationManager::Get().AddNotification(NotificationInfo);
+
+			Notification->ExpireAndFadeout();
+
+			return false;
+		});
+
+	FTicker::GetCoreTicker().AddTicker(TEXT("TestNotifications"), 8.0f, [](float DeltaTime)
+		{
+			FNotificationInfo NotificationInfo(FText::FromString(TEXT("Everything Under The Sun")));
+			NotificationInfo.FadeInDuration = 2.0f;
+			NotificationInfo.FadeOutDuration = 2.0f;
+			NotificationInfo.ExpireDuration = 10.0f;
+
+			NotificationInfo.CheckBoxText = FText::FromString(TEXT("Don't ask again"));
+			NotificationInfo.CheckBoxState = ECheckBoxState::Checked;
+			NotificationInfo.CheckBoxStateChanged = FOnCheckStateChanged::CreateStatic([](ECheckBoxState NewState) {});
+
+			NotificationInfo.Hyperlink = FSimpleDelegate::CreateLambda([]() {});
+			NotificationInfo.HyperlinkText = FText::FromString(TEXT("This is a hyperlink"));
+
+			NotificationInfo.bUseSuccessFailIcons = true;
+			NotificationInfo.bUseThrobber = true;
+
+
+			FNotificationButtonInfo Button1(FText::FromString("OK"), FText::GetEmpty(), FSimpleDelegate());
+			FNotificationButtonInfo Button2(FText::FromString("Cancel"), FText::GetEmpty(), FSimpleDelegate());
+
+			NotificationInfo.ButtonDetails.Add(Button1);
+			NotificationInfo.ButtonDetails.Add(Button2);
+
+
+			auto Notification = FSlateNotificationManager::Get().AddNotification(NotificationInfo);
+			Notification->SetCompletionState(SNotificationItem::CS_Pending);
+
+			Notification->ExpireAndFadeout();
+
+			return false;
+		});
+}
+
+FAutoConsoleCommand TestNotificationCommand(TEXT("Slate.TestNotifications"), TEXT(""), FConsoleCommandDelegate::CreateStatic(&TestNotifications));
+
+#endif // (WITH_EDITOR || IS_PROGRAM) && !UE_BUILD_SHIPPING
