Slate: Fix memory leak in SegmentedControl.

diff --git a/Engine/Source/Runtime/Slate/Public/Widgets/Input/SSegmentedControl.h b/Engine/Source/Runtime/Slate/Public/Widgets/Input/SSegmentedControl.h
--- a/Engine/Source/Runtime/Slate/Public/Widgets/Input/SSegmentedControl.h
+++ b/Engine/Source/Runtime/Slate/Public/Widgets/Input/SSegmentedControl.h
@@ -107,7 +107,12 @@ public:
 
 		UniformPadding = InArgs._UniformPadding;
 
-		Children = InArgs.Slots;
+		const int32 NumSlots = InArgs.Slots.Num();
+		Children.Reserve(NumSlots);
+		for (int32 SlotIndex = 0; SlotIndex < NumSlots; ++SlotIndex)
+		{
+			Children.Add(InArgs.Slots[SlotIndex]);
+		}
 		RebuildChildren();
 	}
 
@@ -119,9 +124,9 @@ public:
 		const int32 NumSlots = Children.Num();
 		for ( int32 SlotIndex = 0; SlotIndex < NumSlots; ++SlotIndex )
 		{
-			TSharedRef<SWidget> Child = Children[SlotIndex]->GetWidget();
-			FSlot<OptionType>* ChildSlotPtr = Children[SlotIndex];
-			OptionType& ChildValue = ChildSlotPtr->_Value;			
+			TSharedRef<SWidget> Child = Children[SlotIndex].GetWidget();
+			FSlot<OptionType>* ChildSlotPtr = &Children[SlotIndex];
+			OptionType& ChildValue = ChildSlotPtr->_Value;
 
 			if (Child == SNullWidget::NullWidget)
 			{
@@ -187,7 +192,7 @@ public:
 		Children.Add( &NewSlot );
 		RebuildChildren();
 
-		return NewSlot;	
+		return NewSlot;
 	}
 
 	int32 NumSlots() const
@@ -223,7 +228,7 @@ public:
 
 
 protected:
-	TArray<FSlot<OptionType>*> Children;
+	TIndirectArray<FSlot<OptionType>> Children;
 
 	FOnValueChanged OnValueChanged;
 
