Slate: Convert SBox Attributes into SlateAttributes.

diff --git a/Engine/Source/Runtime/Slate/Private/Widgets/Layout/SBox.cpp b/Engine/Source/Runtime/Slate/Private/Widgets/Layout/SBox.cpp
--- a/Engine/Source/Runtime/Slate/Private/Widgets/Layout/SBox.cpp
+++ b/Engine/Source/Runtime/Slate/Private/Widgets/Layout/SBox.cpp
@@ -5,8 +5,29 @@
 #include "Layout/ArrangedChildren.h"
 #include "Layout/LayoutUtils.h"
 
+SLATE_IMPLEMENT_WIDGET(SBox)
+void SBox::PrivateRegisterAttributes(FSlateAttributeInitializer& AttributeInitializer)
+{
+	SLATE_ADD_MEMBER_ATTRIBUTE_DEFINITION(AttributeInitializer, WidthOverride, EInvalidateWidgetReason::Layout);
+	SLATE_ADD_MEMBER_ATTRIBUTE_DEFINITION(AttributeInitializer, HeightOverride, EInvalidateWidgetReason::Layout);
+	SLATE_ADD_MEMBER_ATTRIBUTE_DEFINITION(AttributeInitializer, MinDesiredWidth, EInvalidateWidgetReason::Layout);
+	SLATE_ADD_MEMBER_ATTRIBUTE_DEFINITION(AttributeInitializer, MinDesiredHeight, EInvalidateWidgetReason::Layout);
+	SLATE_ADD_MEMBER_ATTRIBUTE_DEFINITION(AttributeInitializer, MaxDesiredWidth, EInvalidateWidgetReason::Layout);
+	SLATE_ADD_MEMBER_ATTRIBUTE_DEFINITION(AttributeInitializer, MaxDesiredHeight, EInvalidateWidgetReason::Layout);
+	SLATE_ADD_MEMBER_ATTRIBUTE_DEFINITION(AttributeInitializer, MinAspectRatio, EInvalidateWidgetReason::Paint);
+	SLATE_ADD_MEMBER_ATTRIBUTE_DEFINITION(AttributeInitializer, MaxAspectRatio, EInvalidateWidgetReason::Paint);
+}
+
 SBox::SBox()
-: ChildSlot(this)
+	: ChildSlot(this)
+	, WidthOverride(*this)
+	, HeightOverride(*this)
+	, MinDesiredWidth(*this)
+	, MinDesiredHeight(*this)
+	, MaxDesiredWidth(*this)
+	, MaxDesiredHeight(*this)
+	, MinAspectRatio(*this)
+	, MaxAspectRatio(*this)
 {
 	SetCanTick(false);
 	bCanSupportFocus = false;
@@ -14,16 +35,16 @@ SBox::SBox()
 
 void SBox::Construct( const FArguments& InArgs )
 {
-	WidthOverride = InArgs._WidthOverride;
-	HeightOverride = InArgs._HeightOverride;
+	SetWidthOverride(InArgs._WidthOverride);
+	SetHeightOverride(InArgs._HeightOverride);
 
-	MinDesiredWidth = InArgs._MinDesiredWidth;
-	MinDesiredHeight = InArgs._MinDesiredHeight;
-	MaxDesiredWidth = InArgs._MaxDesiredWidth;
-	MaxDesiredHeight = InArgs._MaxDesiredHeight;
+	SetMinDesiredWidth(InArgs._MinDesiredWidth);
+	SetMinDesiredHeight(InArgs._MinDesiredHeight);
+	SetMaxDesiredWidth(InArgs._MaxDesiredWidth);
+	SetMaxDesiredHeight(InArgs._MaxDesiredHeight);
 
-	MinAspectRatio = InArgs._MinAspectRatio;
-	MaxAspectRatio = InArgs._MaxAspectRatio;
+	SetMinAspectRatio(InArgs._MinAspectRatio);
+	SetMaxAspectRatio( InArgs._MaxAspectRatio);
 
 	ChildSlot
 		.HAlign( InArgs._HAlign )
@@ -77,46 +98,42 @@ void SBox::SetPadding(const TAttribute<FMargin>& InPadding)
 
 void SBox::SetWidthOverride(TAttribute<FOptionalSize> InWidthOverride)
 {
-	SetAttribute(WidthOverride, InWidthOverride, EInvalidateWidgetReason::Layout);
+	WidthOverride.Assign(*this, InWidthOverride);
 }
 
 void SBox::SetHeightOverride(TAttribute<FOptionalSize> InHeightOverride)
 {
-	SetAttribute(HeightOverride, InHeightOverride, EInvalidateWidgetReason::Layout);
+	HeightOverride.Assign(*this, InHeightOverride);
 }
 
 void SBox::SetMinDesiredWidth(TAttribute<FOptionalSize> InMinDesiredWidth)
 {
-	SetAttribute(MinDesiredWidth, InMinDesiredWidth, EInvalidateWidgetReason::Layout);
+	MinDesiredWidth.Assign(*this, InMinDesiredWidth);
 }
 
 void SBox::SetMinDesiredHeight(TAttribute<FOptionalSize> InMinDesiredHeight)
 {
-	SetAttribute(MinDesiredHeight, InMinDesiredHeight, EInvalidateWidgetReason::Layout);
+	MinDesiredHeight.Assign(*this, InMinDesiredHeight);
 }
 
 void SBox::SetMaxDesiredWidth(TAttribute<FOptionalSize> InMaxDesiredWidth)
 {
-	SetAttribute(MaxDesiredWidth, InMaxDesiredWidth, EInvalidateWidgetReason::Layout);
+	MaxDesiredWidth.Assign(*this, InMaxDesiredWidth);
 }
 
 void SBox::SetMaxDesiredHeight(TAttribute<FOptionalSize> InMaxDesiredHeight)
 {
-	SetAttribute(MaxDesiredHeight, InMaxDesiredHeight, EInvalidateWidgetReason::Layout);
+	MaxDesiredHeight.Assign(*this, InMaxDesiredHeight);
 }
 
 void SBox::SetMinAspectRatio(TAttribute<FOptionalSize> InMinAspectRatio)
 {
-	if (!MinAspectRatio.IdenticalTo(InMinAspectRatio))
-	{
-		MinAspectRatio = InMinAspectRatio;
-		Invalidate(EInvalidateWidget::LayoutAndVolatility);
-	}
+	MinAspectRatio.Assign(*this, InMinAspectRatio);
 }
 
 void SBox::SetMaxAspectRatio(TAttribute<FOptionalSize> InMaxAspectRatio)
 {
-	SetAttribute(MaxAspectRatio, InMaxAspectRatio, EInvalidateWidgetReason::Layout);
+	MaxAspectRatio.Assign(*this, InMaxAspectRatio);
 }
 
 FVector2D SBox::ComputeDesiredSize( float ) const
diff --git a/Engine/Source/Runtime/Slate/Public/Widgets/Layout/SBox.h b/Engine/Source/Runtime/Slate/Public/Widgets/Layout/SBox.h
--- a/Engine/Source/Runtime/Slate/Public/Widgets/Layout/SBox.h
+++ b/Engine/Source/Runtime/Slate/Public/Widgets/Layout/SBox.h
@@ -21,6 +21,8 @@ class FSlateWindowElementList;
  */
 class SLATE_API SBox : public SPanel
 {
+	SLATE_DECLARE_WIDGET(SBox, SPanel)
+
 public:
 	class FBoxSlot : public TSupportsOneChildMixin<FBoxSlot>, public TSupportsContentAlignmentMixin<FBoxSlot>, public TSupportsContentPaddingMixin<FBoxSlot>
 	{
@@ -140,24 +142,24 @@ protected:
 
 private:
 	/** When specified, ignore the content's desired size and report the.WidthOverride as the Box's desired width. */
-	TAttribute<FOptionalSize> WidthOverride;
+	TSlateAttribute<FOptionalSize> WidthOverride;
 
 	/** When specified, ignore the content's desired size and report the.HeightOverride as the Box's desired height. */
-	TAttribute<FOptionalSize> HeightOverride;
+	TSlateAttribute<FOptionalSize> HeightOverride;
 
 	/** When specified, will report the MinDesiredWidth if larger than the content's desired width. */
-	TAttribute<FOptionalSize> MinDesiredWidth;
+	TSlateAttribute<FOptionalSize> MinDesiredWidth;
 
 	/** When specified, will report the MinDesiredHeight if larger than the content's desired height. */
-	TAttribute<FOptionalSize> MinDesiredHeight;
+	TSlateAttribute<FOptionalSize> MinDesiredHeight;
 
 	/** When specified, will report the MaxDesiredWidth if smaller than the content's desired width. */
-	TAttribute<FOptionalSize> MaxDesiredWidth;
+	TSlateAttribute<FOptionalSize> MaxDesiredWidth;
 
 	/** When specified, will report the MaxDesiredHeight if smaller than the content's desired height. */
-	TAttribute<FOptionalSize> MaxDesiredHeight;
+	TSlateAttribute<FOptionalSize> MaxDesiredHeight;
 
-	TAttribute<FOptionalSize> MinAspectRatio;
+	TSlateAttribute<FOptionalSize> MinAspectRatio;
 
-	TAttribute<FOptionalSize> MaxAspectRatio;
+	TSlateAttribute<FOptionalSize> MaxAspectRatio;
 };
