Removes unnecessary foreground color styling in SBreadcrumbTrail

diff --git a/Engine/Source/Developer/SlateFileDialogs/Private/SlateFileDlgWindow.cpp b/Engine/Source/Developer/SlateFileDialogs/Private/SlateFileDlgWindow.cpp
--- a/Engine/Source/Developer/SlateFileDialogs/Private/SlateFileDlgWindow.cpp
+++ b/Engine/Source/Developer/SlateFileDialogs/Private/SlateFileDlgWindow.cpp
@@ -414,7 +414,6 @@ void SSlateFileOpenDlg::Construct(const FArguments& InArgs)
 					.DelimiterImage(StyleSet->GetBrush("SlateFileDialogs.PathDelimiter"))
 					.TextStyle(StyleSet->Get(), "SlateFileDialogs.PathText")
 					.ShowLeadingDelimiter(false)
-					.InvertTextColorOnHover(false)
 					.OnCrumbClicked(this, &SSlateFileOpenDlg::OnPathClicked)
 					.GetCrumbMenuContent(this, &SSlateFileOpenDlg::OnGetCrumbDelimiterContent)
 					.AddMetaData<FTagMetaData>(FTagMetaData(TEXT("ContentBrowserPath")))
diff --git a/Engine/Source/Developer/TraceInsights/Private/Insights/Table/Widgets/STableTreeView.cpp b/Engine/Source/Developer/TraceInsights/Private/Insights/Table/Widgets/STableTreeView.cpp
--- a/Engine/Source/Developer/TraceInsights/Private/Insights/Table/Widgets/STableTreeView.cpp
+++ b/Engine/Source/Developer/TraceInsights/Private/Insights/Table/Widgets/STableTreeView.cpp
@@ -210,15 +210,30 @@ void STableTreeView::ConstructWidget(TSharedPtr<FTable> InTablePtr)
 			.FillWidth(1.0f)
 			.VAlign(VAlign_Center)
 			[
-				SAssignNew(GroupingBreadcrumbTrail, SBreadcrumbTrail<TSharedPtr<FTreeNodeGrouping>>)
-				.ButtonContentPadding(FMargin(1.0f, 1.0f))
-				//.DelimiterImage(FEditorStyle::GetBrush("SlateFileDialogs.PathDelimiter"))
-				//.TextStyle(FEditorStyle::Get(), "Tutorials.Browser.PathText")
-				//.ShowLeadingDelimiter(true)
-				//.PersistentBreadcrumbs(true)
-				.InvertTextColorOnHover(true)
-				.OnCrumbClicked(this, &STableTreeView::OnGroupingCrumbClicked)
-				.GetCrumbMenuContent(this, &STableTreeView::GetGroupingCrumbMenuContent)
+				SAssignNew(ToolbarBox, SHorizontalBox)
+
+				+ SHorizontalBox::Slot()
+				.AutoWidth()
+				.VAlign(VAlign_Center)
+				[
+					SNew(STextBlock)
+					.Text(LOCTEXT("GroupByText", "Hierarchy:"))
+					.Margin(FMargin(0.0f, 0.0f, 4.0f, 0.0f))
+				]
+
+				+ SHorizontalBox::Slot()
+				.FillWidth(1.0f)
+				.VAlign(VAlign_Center)
+				[
+					SAssignNew(GroupingBreadcrumbTrail, SBreadcrumbTrail<TSharedPtr<FTreeNodeGrouping>>)
+					.ButtonContentPadding(FMargin(1.0f, 1.0f))
+					//.DelimiterImage(FEditorStyle::GetBrush("SlateFileDialogs.PathDelimiter"))
+					//.TextStyle(FEditorStyle::Get(), "Tutorials.Browser.PathText")
+					//.ShowLeadingDelimiter(true)
+					//.PersistentBreadcrumbs(true)
+					.OnCrumbClicked(this, &STableTreeView::OnGroupingCrumbClicked)
+					.GetCrumbMenuContent(this, &STableTreeView::GetGroupingCrumbMenuContent)
+				]
 			]
 		]
 	]
diff --git a/Engine/Source/Editor/ContentBrowser/Private/SContentBrowser.cpp b/Engine/Source/Editor/ContentBrowser/Private/SContentBrowser.cpp
--- a/Engine/Source/Editor/ContentBrowser/Private/SContentBrowser.cpp
+++ b/Engine/Source/Editor/ContentBrowser/Private/SContentBrowser.cpp
@@ -319,7 +319,6 @@ void SContentBrowser::Construct( const FArguments& InArgs, const FName& InInstan
 					.DelimiterImage(FAppStyle::Get().GetBrush("Icons.ChevronRight"))
 					.TextStyle(FAppStyle::Get(), "NormalText")
 					.ShowLeadingDelimiter(false)
-					.InvertTextColorOnHover(false)
 					.OnCrumbClicked(this, &SContentBrowser::OnPathClicked)
 					.HasCrumbMenuContent(this, &SContentBrowser::OnHasCrumbDelimiterContent)
 					.GetCrumbMenuContent(this, &SContentBrowser::OnGetCrumbDelimiterContent)
diff --git a/Engine/Source/Editor/IntroTutorials/Private/STutorialsBrowser.cpp b/Engine/Source/Editor/IntroTutorials/Private/STutorialsBrowser.cpp
--- a/Engine/Source/Editor/IntroTutorials/Private/STutorialsBrowser.cpp
+++ b/Engine/Source/Editor/IntroTutorials/Private/STutorialsBrowser.cpp
@@ -593,7 +593,6 @@ void STutorialsBrowser::Construct(const FArguments& InArgs)
 						.DelimiterImage(FEditorStyle::GetBrush("Tutorials.Browser.Breadcrumb"))
 						.TextStyle(FEditorStyle::Get(), "Tutorials.Browser.PathText")
 						.ShowLeadingDelimiter( true )
-						.InvertTextColorOnHover( false )
 						.OnCrumbClicked(this, &STutorialsBrowser::OnBreadcrumbClicked)
 					]
 					+SVerticalBox::Slot()
diff --git a/Engine/Source/Runtime/Slate/Public/Widgets/Navigation/SBreadcrumbTrail.h b/Engine/Source/Runtime/Slate/Public/Widgets/Navigation/SBreadcrumbTrail.h
--- a/Engine/Source/Runtime/Slate/Public/Widgets/Navigation/SBreadcrumbTrail.h
+++ b/Engine/Source/Runtime/Slate/Public/Widgets/Navigation/SBreadcrumbTrail.h
@@ -63,11 +63,8 @@ public:
 	/** Callback for getting the menu content to be displayed when clicking on a crumb's delimiter arrow */
 	DECLARE_DELEGATE_RetVal_OneParam( TSharedRef< SWidget >, FGetCrumbMenuContent, const ItemType& /*CrumbData*/ );
 
-	DECLARE_DELEGATE_RetVal_TwoParams( FSlateColor, FOnGetCrumbColor, int32 /*CrumbId*/, bool /*bInvert*/ );
-
 	SLATE_BEGIN_ARGS( SBreadcrumbTrail )
-		: _InvertTextColorOnHover(true)
-		, _ButtonStyle( &FCoreStyle::Get().GetWidgetStyle< FButtonStyle >( "BreadcrumbButton" ) )
+		: _ButtonStyle( &FCoreStyle::Get().GetWidgetStyle< FButtonStyle >( "BreadcrumbButton" ) )
 		, _TextStyle( &FCoreStyle::Get().GetWidgetStyle<FTextBlockStyle>("NormalText") )
 		, _ButtonContentPadding(FMargin(4.0, 2.0))
 		, _DelimiterImage( FCoreStyle::Get().GetBrush("BreadcrumbTrail.Delimiter") )
@@ -75,12 +72,8 @@ public:
 		, _PersistentBreadcrumbs(false)
 		, _HasCrumbMenuContent()
 		, _GetCrumbMenuContent()
-		, _OnGetCrumbColor()
 	    {}
 
-		/** When true, will invert the button text color when a crumb button is hovered */
-		SLATE_ARGUMENT( bool, InvertTextColorOnHover )
-
 		/** The name of the style to use for the crumb buttons */
 		SLATE_STYLE_ARGUMENT( FButtonStyle, ButtonStyle )
 		
@@ -112,14 +105,11 @@ public:
 
 		SLATE_EVENT( FGetCrumbMenuContent, GetCrumbMenuContent )
 
-		SLATE_EVENT( FOnGetCrumbColor, OnGetCrumbColor )
-
 	SLATE_END_ARGS()
 
 	/** Constructs this widget with InArgs */
 	void Construct( const FArguments& InArgs )
 	{
-		bInvertTextColorOnHover = InArgs._InvertTextColorOnHover;
 		ButtonStyle = InArgs._ButtonStyle;
 		TextStyle = InArgs._TextStyle;
 		ButtonContentPadding = InArgs._ButtonContentPadding;
@@ -131,7 +121,6 @@ public:
 		bHasStaticBreadcrumbs = InArgs._PersistentBreadcrumbs;
 		HasCrumbMenuContentCallback = InArgs._HasCrumbMenuContent;
 		GetCrumbMenuContentCallback = InArgs._GetCrumbMenuContent;
-		OnGetCrumbColor = InArgs._OnGetCrumbColor;
 
 		NextValidCrumbID = 0;
 
@@ -170,7 +159,7 @@ public:
 				.TextStyle(TextStyle)
 				.Text(CrumbText)
 				.OnClicked( this, &SBreadcrumbTrail::CrumbButtonClicked, NextValidCrumbID )
-				.ForegroundColor( this, &SBreadcrumbTrail::GetButtonForegroundColor, NextValidCrumbID )
+				// .ForegroundColor( FSlateColor::Uthis, &SBreadcrumbTrail::GetButtonForegroundColor, NextValidCrumbID )
 			]
 		];
 
@@ -377,51 +366,6 @@ private:
 		return ShowLeadingDelimiter.Get() ? EVisibility::Visible : EVisibility::Collapsed;
 	}
 
-	/** Handler to determine the text color of crumb buttons. Will invert the text color if allowed. */
-	FSlateColor GetButtonForegroundColor(int32 CrumbID) const
-	{
-		TSharedPtr< SButton > CrumbButton;
-
-		if ( !OnGetCrumbColor.IsBound() )
-		{
-			if ( bInvertTextColorOnHover )
-			{
-				for ( int32 CrumbListIdx = 0; CrumbListIdx < NumCrumbs(); ++CrumbListIdx )
-				{
-					if ( CrumbList[ CrumbListIdx ].CrumbID == CrumbID )
-					{
-						CrumbButton = CrumbList[ CrumbListIdx ].Button;
-						break;
-					}
-				}
-
-				if ( CrumbButton.IsValid() && CrumbButton->IsHovered() )
-				{
-					static const FName InvertedForegroundName( "InvertedForeground" );
-					return FCoreStyle::Get().GetSlateColor( InvertedForegroundName );
-				}
-			}
-
-			static const FName DefaultForegroundName("DefaultForeground");
-			return FCoreStyle::Get().GetSlateColor( DefaultForegroundName );
-		}
-		else
-		{
-			int32 CrumbPosition = INDEX_NONE;
-			for ( int32 CrumbListIdx = 0; CrumbListIdx < NumCrumbs(); ++CrumbListIdx )
-			{
-				if ( CrumbList[ CrumbListIdx ].CrumbID == CrumbID )
-				{
-					CrumbButton = CrumbList[ CrumbListIdx ].Button;
-					CrumbPosition = CrumbListIdx;
-					break;
-				}
-			}
-
-			return OnGetCrumbColor.Execute( CrumbPosition, ( CrumbButton.IsValid() && CrumbButton->IsHovered() ) );
-		}
-	}
-
 	/** Handler for when a crumb is clicked. Will pop crumbs down to the selected one. */
 	FReply CrumbButtonClicked(int32 CrumbID)
 	{
@@ -485,9 +429,6 @@ private:
 	/** The next ID to assign to a crumb when it is created */
 	int32 NextValidCrumbID;
 
-	/** When true, will invert the button text color when a crumb button is hovered */
-	bool bInvertTextColorOnHover;
-
 	/** The button style to apply to all crumbs */
 	const FButtonStyle* ButtonStyle;
 
@@ -521,5 +462,4 @@ private:
 	/** If true, don't dynamically remove items when clicking */
 	bool bHasStaticBreadcrumbs;
 
-	FOnGetCrumbColor OnGetCrumbColor;
 };
