Slate: Convert STextBlock Attributes into SlateAttributes

diff --git a/Engine/Plugins/VirtualProduction/DMX/DMXEngine/Source/DMXEditor/Private/Widgets/OutputConsole/SDMXFader.cpp b/Engine/Plugins/VirtualProduction/DMX/DMXEngine/Source/DMXEditor/Private/Widgets/OutputConsole/SDMXFader.cpp
--- a/Engine/Plugins/VirtualProduction/DMX/DMXEngine/Source/DMXEditor/Private/Widgets/OutputConsole/SDMXFader.cpp
+++ b/Engine/Plugins/VirtualProduction/DMX/DMXEngine/Source/DMXEditor/Private/Widgets/OutputConsole/SDMXFader.cpp
@@ -33,7 +33,7 @@ void SDMXFader::Construct(const FArguments& InArgs)
 {
 	OnRequestDelete = InArgs._OnRequestDelete;
 	OnRequestSelect = InArgs._OnRequestSelect;
-	FaderName = InArgs._FaderName.ToString();
+	FaderName = InArgs._FaderName;
 
 	// Init args and test them for general DMX validity, no protocol specifics
 	check(InArgs._UniverseID >= 0 && InArgs._UniverseID < DMX_MAX_UNIVERSE);
@@ -494,8 +494,7 @@ void SDMXFader::OnFaderNameCommitted(const FText& NewFaderName, ETextCommit::Typ
 	check(FaderNameTextBox.IsValid());
 
 	FaderNameTextBox->SetText(FaderName);
-
-	FaderName = NewFaderName.ToString();
+	FaderName = NewFaderName;
 }
 
 FReply SDMXFader::OnUniverseIDBorderDoubleClicked(const FGeometry& InMyGeometry, const FPointerEvent& InMouseEvent)
diff --git a/Engine/Plugins/VirtualProduction/DMX/DMXEngine/Source/DMXEditor/Private/Widgets/OutputConsole/SDMXFader.h b/Engine/Plugins/VirtualProduction/DMX/DMXEngine/Source/DMXEditor/Private/Widgets/OutputConsole/SDMXFader.h
--- a/Engine/Plugins/VirtualProduction/DMX/DMXEngine/Source/DMXEditor/Private/Widgets/OutputConsole/SDMXFader.h
+++ b/Engine/Plugins/VirtualProduction/DMX/DMXEngine/Source/DMXEditor/Private/Widgets/OutputConsole/SDMXFader.h
@@ -126,7 +126,7 @@ private:
 
 public:
 	/** Returns the name of the fader */
-	FString GetFaderName() const { return FaderName; };
+	FString GetFaderName() const { return FaderName.ToString(); };
 
 	/** Returns the universe ID to which to should send DMX to */
 	int32 GetUniverseID() const { return UniverseID; }
@@ -151,7 +151,7 @@ public:
 
 private:
 	/** Cached Name of the Fader */
-	FString FaderName;
+	FText FaderName;
 
 	/** The universe the should send to fader */
 	int32 UniverseID;
diff --git a/Engine/Source/Developer/SlateFileDialogs/Private/SlateFileDlgWindow.cpp b/Engine/Source/Developer/SlateFileDialogs/Private/SlateFileDlgWindow.cpp
--- a/Engine/Source/Developer/SlateFileDialogs/Private/SlateFileDlgWindow.cpp
+++ b/Engine/Source/Developer/SlateFileDialogs/Private/SlateFileDlgWindow.cpp
@@ -1456,7 +1456,7 @@ void SSlateFileOpenDlg::SetDefaultFile(FString DefaultFile)
 		SaveFilename = FileList;
 	}
 
-	SaveFilenameEditBox->SetText(SaveFilename);
+	SaveFilenameEditBox->SetText(FText::FromString(SaveFilename));
 }
 
 
@@ -1630,7 +1630,7 @@ FReply SSlateFileOpenDlg::OnNewDirectoryClick()
 {
 	NewDirectorySizeBox->SetVisibility(EVisibility::Visible);
 	NewDirCancelButton->SetVisibility(EVisibility::Visible);
-	NewDirectoryEditBox->SetText(FString(""));
+	NewDirectoryEditBox->SetText(FText::GetEmpty());
 	
 	FSlateApplication::Get().SetKeyboardFocus(NewDirectoryEditBox);	
 	NewDirectoryEditBox->EnterEditingMode();
@@ -1673,7 +1673,7 @@ FReply SSlateFileOpenDlg::OnNewDirectoryAcceptCancelClick(FSlateFileDlgWindow::E
 	NewDirCancelButton->SetVisibility(EVisibility::Hidden);
 	DirErrorMsg->SetVisibility(EVisibility::Collapsed);
 
-	NewDirectoryEditBox->SetText(FString(""));
+	NewDirectoryEditBox->SetText(FText::GetEmpty());
 
 	return FReply::Handled();
 }
diff --git a/Engine/Source/Runtime/Slate/Private/Widgets/Text/SInlineEditableTextBlock.cpp b/Engine/Source/Runtime/Slate/Private/Widgets/Text/SInlineEditableTextBlock.cpp
--- a/Engine/Source/Runtime/Slate/Private/Widgets/Text/SInlineEditableTextBlock.cpp
+++ b/Engine/Source/Runtime/Slate/Private/Widgets/Text/SInlineEditableTextBlock.cpp
@@ -7,6 +7,17 @@
 #include "Widgets/Input/SMultiLineEditableTextBox.h"
 #include "Widgets/Input/SEditableTextBox.h"
 
+SLATE_IMPLEMENT_WIDGET(SInlineEditableTextBlock)
+void SInlineEditableTextBlock::PrivateRegisterAttributes(FSlateAttributeInitializer& AttributeInitializer)
+{
+	SLATE_ADD_MEMBER_ATTRIBUTE_DEFINITION_WITH_NAME(AttributeInitializer, "Text", TextAttribute, EInvalidateWidgetReason::Layout);
+	SLATE_ADD_MEMBER_ATTRIBUTE_DEFINITION_WITH_NAME(AttributeInitializer, "bIsReadOnly", bIsReadOnlyAttribute, EInvalidateWidgetReason::None);
+}
+
+SInlineEditableTextBlock::SInlineEditableTextBlock()
+	: TextAttribute(*this)
+	, bIsReadOnlyAttribute(*this)
+{}
 
 void SInlineEditableTextBlock::Construct( const FArguments& InArgs )
 {
@@ -16,8 +27,8 @@ void SInlineEditableTextBlock::Construct( const FArguments& InArgs )
 	OnTextCommittedDelegate = InArgs._OnTextCommitted;
 	IsSelected = InArgs._IsSelected;
 	OnVerifyTextChanged= InArgs._OnVerifyTextChanged;
-	Text = InArgs._Text;
-	bIsReadOnly = InArgs._IsReadOnly;
+	TextAttribute.Assign(*this, InArgs._Text);
+	bIsReadOnlyAttribute.Assign(*this, InArgs._IsReadOnly);
 	bIsMultiLine = InArgs._MultiLine;
 	DoubleSelectDelay = 0.0f;
 
@@ -32,7 +43,7 @@ void SInlineEditableTextBlock::Construct( const FArguments& InArgs )
 		.VAlign(VAlign_Center)
 		[
 			SAssignNew(TextBlock, STextBlock)
-			.Text(Text)
+			.Text(InArgs._Text)
 			.TextStyle( &InArgs._Style->TextStyle )
 			.Font(InArgs._Font)
 			.ColorAndOpacity( InArgs._ColorAndOpacity )
@@ -92,7 +103,7 @@ void SInlineEditableTextBlock::CancelEditMode()
 	ExitEditingMode();
 
 	// Get the text from source again.
-	SetEditableText(Text);
+	SetEditableText(TextAttribute.Get());
 }
 
 bool SInlineEditableTextBlock::SupportsKeyboardFocus() const
@@ -103,7 +114,7 @@ bool SInlineEditableTextBlock::SupportsKeyboardFocus() const
 
 void SInlineEditableTextBlock::EnterEditingMode()
 {
-	if(!bIsReadOnly.Get() && FSlateApplication::Get().HasAnyMouseCaptor() == false)
+	if(!bIsReadOnlyAttribute.Get() && FSlateApplication::Get().HasAnyMouseCaptor() == false)
 	{
 		if(TextBlock->GetVisibility() == EVisibility::Visible)
 		{
@@ -156,26 +167,26 @@ bool SInlineEditableTextBlock::IsInEditMode() const
 
 void SInlineEditableTextBlock::SetReadOnly(bool bInIsReadOnly)
 {
-	bIsReadOnly = bInIsReadOnly;
+	bIsReadOnlyAttribute.Set(*this, bInIsReadOnly);
 }
 
-void SInlineEditableTextBlock::SetText( const TAttribute< FText >& InText )
+void SInlineEditableTextBlock::SetText( TAttribute<FText> InText )
 {
-	Text = InText;
-	TextBlock->SetText( Text );
-	SetEditableText( Text );
+	TextBlock->SetText( InText );
+	SetEditableText( InText );
+	TextAttribute.Assign(*this, MoveTemp(InText));
 }
 
 void SInlineEditableTextBlock::SetText( const FString& InText )
 {
-	Text = FText::FromString( InText );
-	TextBlock->SetText( Text );
-	SetEditableText( Text );
+	TextAttribute.Set(*this, FText::FromString( InText ));
+	TextBlock->SetText( TextAttribute.Get() );
+	SetEditableText( TextAttribute.Get() );
 }
 
-void SInlineEditableTextBlock::SetWrapTextAt( const TAttribute<float>& InWrapTextAt )
+void SInlineEditableTextBlock::SetWrapTextAt( TAttribute<float> InWrapTextAt )
 {
-	TextBlock->SetWrapTextAt( InWrapTextAt );
+	TextBlock->SetWrapTextAt( MoveTemp(InWrapTextAt) );
 }
 
 FReply SInlineEditableTextBlock::OnMouseButtonDown( const FGeometry& MyGeometry, const FPointerEvent& MouseEvent )
@@ -187,7 +198,7 @@ FReply SInlineEditableTextBlock::OnMouseButtonDown( const FGeometry& MyGeometry,
 
 	if (IsSelected.IsBound())
 	{
-		if (IsSelected.Execute() && !bIsReadOnly.Get() && !ActiveTimerHandle.IsValid())
+		if (IsSelected.Execute() && !bIsReadOnlyAttribute.Get() && !ActiveTimerHandle.IsValid())
 		{
 			ActiveTimerHandle = RegisterActiveTimer(0.5f, FWidgetActiveTimerDelegate::CreateSP(this, &SInlineEditableTextBlock::TriggerEditMode));
 		}
@@ -195,7 +206,7 @@ FReply SInlineEditableTextBlock::OnMouseButtonDown( const FGeometry& MyGeometry,
 	else
 	{
 		// The widget is not managed by another widget, so handle the mouse input and enter edit mode if ready.
-		if (HasKeyboardFocus() && !bIsReadOnly.Get())
+		if (HasKeyboardFocus() && !bIsReadOnlyAttribute.Get())
 		{
 			EnterEditingMode();
 			return FReply::Handled();
@@ -274,7 +285,7 @@ void SInlineEditableTextBlock::OnTextBoxCommitted(const FText& InText, ETextComm
 {
 	if(InCommitType == ETextCommit::OnCleared)
 	{
-		FText SourceTextInEditMode = Text.Get();
+		FText SourceTextInEditMode = TextAttribute.Get();
 		CancelEditMode();
 		// Commit the name, certain actions might need to be taken by the bound function
 		OnTextCommittedDelegate.ExecuteIfBound(SourceTextInEditMode, InCommitType);
@@ -301,7 +312,7 @@ void SInlineEditableTextBlock::OnTextBoxCommitted(const FText& InText, ETextComm
 					CancelEditMode();
 
 					// Commit the name, certain actions might need to be taken by the bound function
-					OnTextCommittedDelegate.ExecuteIfBound(Text.Get(), InCommitType);
+					OnTextCommittedDelegate.ExecuteIfBound(TextAttribute.Get(), InCommitType);
 				
 					return;
 				}
@@ -311,7 +322,7 @@ void SInlineEditableTextBlock::OnTextBoxCommitted(const FText& InText, ETextComm
 				CancelEditMode();
 
 				// Commit the name, certain actions might need to be taken by the bound function
-				OnTextCommittedDelegate.ExecuteIfBound(Text.Get(), InCommitType);
+				OnTextCommittedDelegate.ExecuteIfBound(TextAttribute.Get(), InCommitType);
 				return;
 			}
 		}
@@ -320,9 +331,9 @@ void SInlineEditableTextBlock::OnTextBoxCommitted(const FText& InText, ETextComm
 
 		OnTextCommittedDelegate.ExecuteIfBound(InText, InCommitType);
 
-		if ( !Text.IsBound() )
+		if ( !TextAttribute.IsBound(*this) )
 		{
-			TextBlock->SetText( Text );
+			TextBlock->SetText( TextAttribute.Get() );
 		}
 	}
 }
@@ -341,12 +352,12 @@ TSharedPtr<SWidget> SInlineEditableTextBlock::GetEditableTextWidget() const
 	}
 }
 
-void SInlineEditableTextBlock::SetEditableText( const TAttribute< FText >& InNewText )
+void SInlineEditableTextBlock::SetEditableText( TAttribute<FText> InNewText )
 {
 #if WITH_FANCY_TEXT
-	bIsMultiLine ? MultiLineTextBox->SetText( Text ) :
+	bIsMultiLine ? MultiLineTextBox->SetText( MoveTemp(InNewText) ) :
 #endif //WITH_FANCY_TEXT
-				 TextBox->SetText( Text );
+				 TextBox->SetText( MoveTemp(InNewText) );
 }
 
 void SInlineEditableTextBlock::SetTextBoxError( const FText& ErrorText )
diff --git a/Engine/Source/Runtime/Slate/Private/Widgets/Text/SRichTextBlock.cpp b/Engine/Source/Runtime/Slate/Private/Widgets/Text/SRichTextBlock.cpp
--- a/Engine/Source/Runtime/Slate/Private/Widgets/Text/SRichTextBlock.cpp
+++ b/Engine/Source/Runtime/Slate/Private/Widgets/Text/SRichTextBlock.cpp
@@ -9,7 +9,32 @@
 #include "Framework/Text/RichTextLayoutMarshaller.h"
 #include "Types/ReflectionMetadata.h"
 
+SLATE_IMPLEMENT_WIDGET(SRichTextBlock)
+void SRichTextBlock::PrivateRegisterAttributes(FSlateAttributeInitializer& AttributeInitializer)
+{
+	SLATE_ADD_MEMBER_ATTRIBUTE_DEFINITION(AttributeInitializer, BoundText, EInvalidateWidgetReason::Layout);
+	SLATE_ADD_MEMBER_ATTRIBUTE_DEFINITION(AttributeInitializer, HighlightText, EInvalidateWidgetReason::Layout);
+	SLATE_ADD_MEMBER_ATTRIBUTE_DEFINITION(AttributeInitializer, WrapTextAt, EInvalidateWidgetReason::Layout);
+	SLATE_ADD_MEMBER_ATTRIBUTE_DEFINITION(AttributeInitializer, WrappingPolicy, EInvalidateWidgetReason::Layout);
+	SLATE_ADD_MEMBER_ATTRIBUTE_DEFINITION(AttributeInitializer, TransformPolicy, EInvalidateWidgetReason::Layout);
+	SLATE_ADD_MEMBER_ATTRIBUTE_DEFINITION(AttributeInitializer, Justification, EInvalidateWidgetReason::Layout);
+	SLATE_ADD_MEMBER_ATTRIBUTE_DEFINITION(AttributeInitializer, AutoWrapText, EInvalidateWidgetReason::Layout);
+	SLATE_ADD_MEMBER_ATTRIBUTE_DEFINITION(AttributeInitializer, Margin, EInvalidateWidgetReason::Layout);
+	SLATE_ADD_MEMBER_ATTRIBUTE_DEFINITION(AttributeInitializer, LineHeightPercentage, EInvalidateWidgetReason::Layout);
+	SLATE_ADD_MEMBER_ATTRIBUTE_DEFINITION(AttributeInitializer, MinDesiredWidth, EInvalidateWidgetReason::Layout);
+}
+
 SRichTextBlock::SRichTextBlock()
+	: BoundText(*this)
+	, HighlightText(*this)
+	, WrapTextAt(*this)
+	, WrappingPolicy(*this)
+	, TransformPolicy(*this)
+	, Justification(*this)
+	, AutoWrapText(*this)
+	, Margin(*this)
+	, LineHeightPercentage(*this)
+	, MinDesiredWidth(*this)
 {
 }
 
@@ -20,18 +45,18 @@ SRichTextBlock::~SRichTextBlock()
 
 void SRichTextBlock::Construct( const FArguments& InArgs )
 {
-	BoundText = InArgs._Text;
-	HighlightText = InArgs._HighlightText;
+	SetText(InArgs._Text);
+	SetHighlightText(InArgs._HighlightText);
 
-	TextStyle = *InArgs._TextStyle;
-	WrapTextAt = InArgs._WrapTextAt;
-	AutoWrapText = InArgs._AutoWrapText;
-	WrappingPolicy = InArgs._WrappingPolicy;
-	TransformPolicy = InArgs._TransformPolicy;
-	Margin = InArgs._Margin;
-	LineHeightPercentage = InArgs._LineHeightPercentage;
-	Justification = InArgs._Justification;
-	MinDesiredWidth = InArgs._MinDesiredWidth;
+	SetTextStyle(*InArgs._TextStyle);
+	SetWrapTextAt(InArgs._WrapTextAt);
+	SetAutoWrapText(InArgs._AutoWrapText);
+	SetWrappingPolicy(InArgs._WrappingPolicy);
+	SetTransformPolicy(InArgs._TransformPolicy);
+	SetMargin(InArgs._Margin);
+	SetLineHeightPercentage(InArgs._LineHeightPercentage);
+	SetJustification(InArgs._Justification);
+	SetMinDesiredWidth(InArgs._MinDesiredWidth);
 
 	{
 		TSharedPtr<IRichTextMarkupParser> Parser = InArgs._Parser;
@@ -72,7 +97,7 @@ int32 SRichTextBlock::OnPaint(const FPaintArgs& Args, const FGeometry& AllottedG
 	// HACK: Due to the nature of wrapping and layout, we may have been arranged in a different box than what we were cached with.  Which
 	// might update wrapping, so make sure we always set the desired size to the current size of the text layout, which may have changed
 	// during paint.
-	bool bCanWrap = WrapTextAt.Get() > 0 || AutoWrapText.Get();
+	bool bCanWrap = WrapTextAt.Get() > 0.f || AutoWrapText.Get();
 
 	if (bCanWrap && !NewDesiredSize.Equals(LastDesiredSize))
 	{
@@ -86,7 +111,16 @@ FVector2D SRichTextBlock::ComputeDesiredSize(float LayoutScaleMultiplier) const
 {
 	// ComputeDesiredSize will also update the text layout cache if required
 	const FVector2D TextSize = TextLayoutCache->ComputeDesiredSize(
-		FSlateTextBlockLayout::FWidgetArgs(BoundText, HighlightText, WrapTextAt, AutoWrapText, WrappingPolicy, TransformPolicy, Margin, LineHeightPercentage, Justification),
+		FSlateTextBlockLayout::FWidgetDesiredSizeArgs(
+			BoundText.Get(),
+			HighlightText.Get(),
+			WrapTextAt.Get(),
+			AutoWrapText.Get(),
+			WrappingPolicy.Get(),
+			TransformPolicy.Get(),
+			Margin.Get(),
+			LineHeightPercentage.Get(),
+			Justification.Get()),
 		LayoutScaleMultiplier * TextBlockScale, TextStyle) * TextBlockScale;
 
 	return FVector2D(FMath::Max(TextSize.X, MinDesiredWidth.Get()), TextSize.Y);
@@ -104,17 +138,15 @@ void SRichTextBlock::OnArrangeChildren(const FGeometry& AllottedGeometry, FArran
 	TextLayoutCache->ArrangeChildren(TextBlockScaledGeometry, ArrangedChildren);
 }
 
-void SRichTextBlock::SetText( const TAttribute<FText>& InTextAttr )
+void SRichTextBlock::SetText(TAttribute<FText> InTextAttr)
 {
-	BoundText = InTextAttr;
-	Invalidate(EInvalidateWidget::LayoutAndVolatility);
+	BoundText.Assign(*this, MoveTemp(InTextAttr), FText::GetEmpty());
 	InvalidatePrepass();
 }
 
-void SRichTextBlock::SetHighlightText( const TAttribute<FText>& InHighlightText )
+void SRichTextBlock::SetHighlightText(TAttribute<FText> InHighlightText)
 {
-	HighlightText = InHighlightText;
-	Invalidate(EInvalidateWidget::LayoutAndVolatility);
+	HighlightText.Assign(*this, MoveTemp(InHighlightText), FText::GetEmpty());
 }
 
 void SRichTextBlock::SetTextShapingMethod(const TOptional<ETextShapingMethod>& InTextShapingMethod)
@@ -129,40 +161,40 @@ void SRichTextBlock::SetTextFlowDirection(const TOptional<ETextFlowDirection>& I
 	Invalidate(EInvalidateWidget::Layout);
 }
 
-void SRichTextBlock::SetWrapTextAt(const TAttribute<float>& InWrapTextAt)
+void SRichTextBlock::SetWrapTextAt(TAttribute<float> InWrapTextAt)
 {
-	SetAttribute(WrapTextAt, InWrapTextAt, EInvalidateWidgetReason::Layout);
+	WrapTextAt.Assign(*this, MoveTemp(InWrapTextAt), 0.f);
 }
 
-void SRichTextBlock::SetAutoWrapText(const TAttribute<bool>& InAutoWrapText)
+void SRichTextBlock::SetAutoWrapText(TAttribute<bool> InAutoWrapText)
 {
-	SetAttribute(AutoWrapText, InAutoWrapText, EInvalidateWidgetReason::Layout);
+	AutoWrapText.Assign(*this, MoveTemp(InAutoWrapText), false);
 	InvalidatePrepass();
 }
 
-void SRichTextBlock::SetWrappingPolicy(const TAttribute<ETextWrappingPolicy>& InWrappingPolicy)
+void SRichTextBlock::SetWrappingPolicy(TAttribute<ETextWrappingPolicy> InWrappingPolicy)
 {
-	SetAttribute(WrappingPolicy, InWrappingPolicy, EInvalidateWidgetReason::Layout);
+	WrappingPolicy.Assign(*this, MoveTemp(InWrappingPolicy));
 }
 
-void SRichTextBlock::SetTransformPolicy(const TAttribute<ETextTransformPolicy>& InTransformPolicy)
+void SRichTextBlock::SetTransformPolicy(TAttribute<ETextTransformPolicy> InTransformPolicy)
 {
-	SetAttribute(TransformPolicy, InTransformPolicy, EInvalidateWidgetReason::Layout);
+	TransformPolicy.Assign(*this, MoveTemp(InTransformPolicy));
 }
 
-void SRichTextBlock::SetLineHeightPercentage(const TAttribute<float>& InLineHeightPercentage)
+void SRichTextBlock::SetLineHeightPercentage(TAttribute<float> InLineHeightPercentage)
 {
-	SetAttribute(LineHeightPercentage, InLineHeightPercentage, EInvalidateWidgetReason::Layout);
+	LineHeightPercentage.Assign(*this, MoveTemp(InLineHeightPercentage));
 }
 
-void SRichTextBlock::SetMargin(const TAttribute<FMargin>& InMargin)
+void SRichTextBlock::SetMargin(TAttribute<FMargin> InMargin)
 {
-	SetAttribute(Margin, InMargin, EInvalidateWidgetReason::Layout);
+	Margin.Assign(*this, MoveTemp(InMargin));
 }
 
-void SRichTextBlock::SetJustification(const TAttribute<ETextJustify::Type>& InJustification)
+void SRichTextBlock::SetJustification(TAttribute<ETextJustify::Type> InJustification)
 {
-	SetAttribute(Justification, InJustification, EInvalidateWidgetReason::Layout);
+	Justification.Assign(*this, MoveTemp(InJustification));
 }
 
 void SRichTextBlock::SetTextStyle(const FTextBlockStyle& InTextStyle)
@@ -171,9 +203,9 @@ void SRichTextBlock::SetTextStyle(const FTextBlockStyle& InTextStyle)
 	Invalidate(EInvalidateWidget::Layout);
 }
 
-void SRichTextBlock::SetMinDesiredWidth(const TAttribute<float>& InMinDesiredWidth)
+void SRichTextBlock::SetMinDesiredWidth(TAttribute<float> InMinDesiredWidth)
 {
-	SetAttribute(MinDesiredWidth, InMinDesiredWidth, EInvalidateWidgetReason::Layout);
+	MinDesiredWidth.Assign(*this, MoveTemp(InMinDesiredWidth));
 }
 
 void SRichTextBlock::SetDecoratorStyleSet(const ISlateStyle* NewDecoratorStyleSet)
@@ -198,19 +230,4 @@ void SRichTextBlock::Refresh()
 	Invalidate(EInvalidateWidget::Layout);
 }
 
-bool SRichTextBlock::ComputeVolatility() const
-{
-	return SWidget::ComputeVolatility() 
-		|| BoundText.IsBound()
-		|| HighlightText.IsBound()
-		|| WrapTextAt.IsBound()
-		|| AutoWrapText.IsBound()
-		|| WrappingPolicy.IsBound()
-		|| TransformPolicy.IsBound()
-		|| Margin.IsBound()
-		|| Justification.IsBound()
-		|| LineHeightPercentage.IsBound()
-		|| MinDesiredWidth.IsBound();
-}
-
 #endif //WITH_FANCY_TEXT
diff --git a/Engine/Source/Runtime/Slate/Private/Widgets/Text/STextBlock.cpp b/Engine/Source/Runtime/Slate/Private/Widgets/Text/STextBlock.cpp
--- a/Engine/Source/Runtime/Slate/Private/Widgets/Text/STextBlock.cpp
+++ b/Engine/Source/Runtime/Slate/Private/Widgets/Text/STextBlock.cpp
@@ -15,13 +15,66 @@
 DECLARE_CYCLE_STAT(TEXT("STextBlock::SetText Time"), Stat_SlateTextBlockSetText, STATGROUP_SlateVerbose)
 DECLARE_CYCLE_STAT(TEXT("STextBlock::OnPaint Time"), Stat_SlateTextBlockOnPaint, STATGROUP_SlateVerbose)
 DECLARE_CYCLE_STAT(TEXT("STextBlock::ComputeDesiredSize"), Stat_SlateTextBlockCDS, STATGROUP_SlateVerbose)
-DECLARE_CYCLE_STAT(TEXT("STextBlock::ComputeVolitility"), Stat_SlateTextBlockCV, STATGROUP_SlateVerbose)
+
+
+SLATE_IMPLEMENT_WIDGET(STextBlock)
+void STextBlock::PrivateRegisterAttributes(FSlateAttributeInitializer& AttributeInitializer)
+{
+	//~ If we are in SimpleTextMode the invalidation reason can be different.
+	struct FInvalidation
+	{
+		static EInvalidateWidgetReason GetInvalidationNoneIfSimpleTextMode(const SWidget& OwningWidget)
+		{
+			return static_cast<const STextBlock&>(OwningWidget).bSimpleTextMode ? EInvalidateWidgetReason::None : EInvalidateWidgetReason::Layout;
+		}
+		static EInvalidateWidgetReason GetInvalidationPaintIfSimpleTextMode(const SWidget& OwningWidget)
+		{
+			return static_cast<const STextBlock&>(OwningWidget).bSimpleTextMode ? EInvalidateWidgetReason::Paint : EInvalidateWidgetReason::Layout;
+		}
+	};
+
+	SLATE_ADD_MEMBER_ATTRIBUTE_DEFINITION(AttributeInitializer, BoundText, EInvalidateWidgetReason::Layout);
+	SLATE_ADD_MEMBER_ATTRIBUTE_DEFINITION(AttributeInitializer, Font, EInvalidateWidgetReason::Layout);
+	SLATE_ADD_MEMBER_ATTRIBUTE_DEFINITION(AttributeInitializer, StrikeBrush, FInvalidation::GetInvalidationNoneIfSimpleTextMode);
+	SLATE_ADD_MEMBER_ATTRIBUTE_DEFINITION(AttributeInitializer, ColorAndOpacity, FInvalidation::GetInvalidationPaintIfSimpleTextMode);
+	SLATE_ADD_MEMBER_ATTRIBUTE_DEFINITION(AttributeInitializer, ShadowOffset, EInvalidateWidgetReason::Layout);
+	SLATE_ADD_MEMBER_ATTRIBUTE_DEFINITION(AttributeInitializer, ShadowColorAndOpacity, FInvalidation::GetInvalidationPaintIfSimpleTextMode);
+	SLATE_ADD_MEMBER_ATTRIBUTE_DEFINITION(AttributeInitializer, HighlightColor, FInvalidation::GetInvalidationNoneIfSimpleTextMode);
+	SLATE_ADD_MEMBER_ATTRIBUTE_DEFINITION(AttributeInitializer, HighlightShape, FInvalidation::GetInvalidationNoneIfSimpleTextMode);
+	SLATE_ADD_MEMBER_ATTRIBUTE_DEFINITION(AttributeInitializer, HighlightText, FInvalidation::GetInvalidationNoneIfSimpleTextMode);
+	SLATE_ADD_MEMBER_ATTRIBUTE_DEFINITION(AttributeInitializer, WrapTextAt, FInvalidation::GetInvalidationNoneIfSimpleTextMode);
+	SLATE_ADD_MEMBER_ATTRIBUTE_DEFINITION(AttributeInitializer, AutoWrapText, FInvalidation::GetInvalidationNoneIfSimpleTextMode);
+	SLATE_ADD_MEMBER_ATTRIBUTE_DEFINITION(AttributeInitializer, WrappingPolicy, FInvalidation::GetInvalidationNoneIfSimpleTextMode);
+	SLATE_ADD_MEMBER_ATTRIBUTE_DEFINITION(AttributeInitializer, TransformPolicy, FInvalidation::GetInvalidationNoneIfSimpleTextMode);
+	SLATE_ADD_MEMBER_ATTRIBUTE_DEFINITION(AttributeInitializer, Margin, FInvalidation::GetInvalidationNoneIfSimpleTextMode);
+	SLATE_ADD_MEMBER_ATTRIBUTE_DEFINITION(AttributeInitializer, Justification, FInvalidation::GetInvalidationNoneIfSimpleTextMode);
+	SLATE_ADD_MEMBER_ATTRIBUTE_DEFINITION(AttributeInitializer, LineHeightPercentage, FInvalidation::GetInvalidationNoneIfSimpleTextMode);
+	SLATE_ADD_MEMBER_ATTRIBUTE_DEFINITION(AttributeInitializer, MinDesiredWidth, EInvalidateWidgetReason::Layout);
+}
 
 STextBlock::STextBlock()
+	: BoundText(*this)
+	, Font(*this)
+	, StrikeBrush(*this)
+	, ColorAndOpacity(*this)
+	, ShadowOffset(*this)
+	, ShadowColorAndOpacity(*this)
+	, HighlightColor(*this)
+	, HighlightShape(*this)
+	, HighlightText(*this)
+	, WrapTextAt(*this)
+	, AutoWrapText(*this)
+	, WrappingPolicy(*this)
+	, TransformPolicy(*this)
+	, Margin(*this)
+	, Justification(*this)
+	, LineHeightPercentage(*this)
+	, MinDesiredWidth(*this)
+	, Union_IsAttributeSet(0)
+	, bSimpleTextMode(false)
 {
 	SetCanTick(false);
 	bCanSupportFocus = false;
-	bSimpleTextMode = false;
 
 #if WITH_ACCESSIBILITY
 	AccessibleBehavior = EAccessibleBehavior::Auto;
@@ -38,33 +91,31 @@ void STextBlock::Construct( const FArguments& InArgs )
 {
 	TextStyle = *InArgs._TextStyle;
 
-	HighlightText = InArgs._HighlightText;
-	WrapTextAt = InArgs._WrapTextAt;
-	AutoWrapText = InArgs._AutoWrapText;
-	WrappingPolicy = InArgs._WrappingPolicy;
+	SetHighlightText(InArgs._HighlightText);
+	SetWrapTextAt(InArgs._WrapTextAt);
+	SetAutoWrapText(InArgs._AutoWrapText);
+	SetWrappingPolicy(InArgs._WrappingPolicy);
 
-	if (InArgs._TransformPolicy.IsSet()) 
-	{
-		TransformPolicy = InArgs._TransformPolicy;
-	}
-	Margin = InArgs._Margin;
-	LineHeightPercentage = InArgs._LineHeightPercentage;
-	Justification = InArgs._Justification;
-	MinDesiredWidth = InArgs._MinDesiredWidth;
-
-	Font = InArgs._Font;
-	StrikeBrush = InArgs._StrikeBrush;
-	ColorAndOpacity = InArgs._ColorAndOpacity;
-	ShadowOffset = InArgs._ShadowOffset;
-	ShadowColorAndOpacity = InArgs._ShadowColorAndOpacity;
-	HighlightColor = InArgs._HighlightColor;
-	HighlightShape = InArgs._HighlightShape;
+	SetTransformPolicy(InArgs._TransformPolicy);
+
+	SetMargin(InArgs._Margin);
+	SetLineHeightPercentage(InArgs._LineHeightPercentage);
+	SetJustification(InArgs._Justification);
+	SetMinDesiredWidth(InArgs._MinDesiredWidth);
+
+	SetFont(InArgs._Font);
+	SetStrikeBrush(InArgs._StrikeBrush);
+	SetColorAndOpacity(InArgs._ColorAndOpacity);
+	SetShadowOffset(InArgs._ShadowOffset);
+	SetShadowColorAndOpacity(InArgs._ShadowColorAndOpacity);
+	SetHighlightColor(InArgs._HighlightColor);
+	SetHighlightShape(InArgs._HighlightShape);
 
 	bSimpleTextMode = InArgs._SimpleTextMode;
 
 	SetOnMouseDoubleClick(InArgs._OnDoubleClicked);
 
-	BoundText = InArgs._Text;
+	BoundText.Assign(*this, InArgs._Text);
 
 	//if(!bSimpleTextMode)
 	{
@@ -76,99 +127,63 @@ void STextBlock::Construct( const FArguments& InArgs )
 
 FSlateFontInfo STextBlock::GetFont() const
 {
-	return Font.IsSet() ? Font.Get() : TextStyle.Font;
+	return bIsAttributeFontSet ? Font.Get() : TextStyle.Font;
 }
 
 const FSlateBrush* STextBlock::GetStrikeBrush() const
 {
-	return StrikeBrush.IsSet() ? StrikeBrush.Get() : &TextStyle.StrikeBrush;
+	return bIsAttributeStrikeBrushSet ? StrikeBrush.Get() : &TextStyle.StrikeBrush;
 }
 
 FSlateColor STextBlock::GetColorAndOpacity() const
 {
-	return ColorAndOpacity.IsSet() ? ColorAndOpacity.Get() : TextStyle.ColorAndOpacity;
+	return bIsAttributeColorAndOpacitySet ? ColorAndOpacity.Get() : TextStyle.ColorAndOpacity;
 }
 
 FVector2D STextBlock::GetShadowOffset() const
 {
-	return ShadowOffset.IsSet() ? ShadowOffset.Get() : TextStyle.ShadowOffset;
+	return bIsAttributeShadowOffsetSet ? ShadowOffset.Get() : TextStyle.ShadowOffset;
 }
 
 FLinearColor STextBlock::GetShadowColorAndOpacity() const
 {
-	return ShadowColorAndOpacity.IsSet() ? ShadowColorAndOpacity.Get() : TextStyle.ShadowColorAndOpacity;
+	return bIsAttributeShadowColorAndOpacitySet ? ShadowColorAndOpacity.Get() : TextStyle.ShadowColorAndOpacity;
 }
 
 FSlateColor STextBlock::GetHighlightColor() const
 {
-	return HighlightColor.IsSet() ? HighlightColor.Get() : TextStyle.HighlightColor;
+	return bIsAttributeHighlightColorSet ? HighlightColor.Get() : TextStyle.HighlightColor;
 }
 
 const FSlateBrush* STextBlock::GetHighlightShape() const
 {
-	return HighlightShape.IsSet() ? HighlightShape.Get() : &TextStyle.HighlightShape;
+	return bIsAttributeHighlightShapeSet ? HighlightShape.Get() : &TextStyle.HighlightShape;
 }
 
-void STextBlock::InvalidateText(EInvalidateWidget InvalidateReason)
+ETextTransformPolicy STextBlock::GetTransformPolicyImpl() const
 {
-	if (bSimpleTextMode && EnumHasAnyFlags(InvalidateReason, EInvalidateWidget::Layout))
-	{
-		CachedSimpleDesiredSize.Reset();
-	}
-
-	Invalidate(InvalidateReason);
+	return bIsAttributeTransformPolicySet ? TransformPolicy.Get() : TextStyle.TransformPolicy;
 }
 
-void STextBlock::SetText( const TAttribute< FString >& InText )
+void STextBlock::InvalidateText(EInvalidateWidgetReason InvalidateReason)
 {
-	if (InText.IsSet() && !InText.IsBound())
+	if (bSimpleTextMode && EnumHasAnyFlags(InvalidateReason, EInvalidateWidgetReason::Layout))
 	{
-		SetText(FText::AsCultureInvariant(InText.Get()));
-		return;
-	}
-
-	SCOPE_CYCLE_COUNTER(Stat_SlateTextBlockSetText);
-	BoundText = MakeAttributeLambda([InText]()
-	{
-		return FText::AsCultureInvariant(InText.Get(FString()));
-	});
-	InvalidateText(EInvalidateWidget::LayoutAndVolatility);
-}
-
-void STextBlock::SetText( const FString& InText )
-{
-	SetText(FText::AsCultureInvariant(InText));
-}
-
-void STextBlock::SetText( const TAttribute< FText >& InText )
-{
-	if (InText.IsSet() && !InText.IsBound())
-	{
-		SetText(InText.Get());
-		return;
+		CachedSimpleDesiredSize.Reset();
 	}
 
-	SCOPE_CYCLE_COUNTER(Stat_SlateTextBlockSetText);
-	BoundText = InText;
-	InvalidateText(EInvalidateWidget::LayoutAndVolatility);
+	Invalidate(InvalidateReason);
 }
 
-void STextBlock::SetText( const FText& InText )
+void STextBlock::SetText(TAttribute<FText> InText)
 {
 	SCOPE_CYCLE_COUNTER(Stat_SlateTextBlockSetText);
-
-	if (!BoundText.IsBound() && InText.IdenticalTo(BoundText.Get(), ETextIdenticalModeFlags::DeepCompare | ETextIdenticalModeFlags::LexicalCompareInvariants))
-	{
-		return;
-	}
-
-	BoundText = InText;
-	InvalidateText(EInvalidateWidget::LayoutAndVolatility);
+	BoundText.Assign(*this, MoveTemp(InText));
 }
 
 void STextBlock::SetHighlightText(TAttribute<FText> InText)
 {
-	HighlightText = InText;
+	HighlightText.Assign(*this, MoveTemp(InText));
 }
 
 int32 STextBlock::OnPaint( const FPaintArgs& Args, const FGeometry& AllottedGeometry, const FSlateRect& MyCullingRect, FSlateWindowElementList& OutDrawElements, int32 LayerId, const FWidgetStyle& InWidgetStyle, bool bParentEnabled ) const
@@ -235,11 +250,11 @@ int32 STextBlock::OnPaint( const FPaintArgs& Args, const FGeometry& AllottedGeom
 		// HACK: Due to the nature of wrapping and layout, we may have been arranged in a different box than what we were cached with.  Which
 		// might update wrapping, so make sure we always set the desired size to the current size of the text layout, which may have changed
 		// during paint.
-		bool bCanWrap = WrapTextAt.Get() > 0 || AutoWrapText.Get();
+		const bool bCanWrap = WrapTextAt.Get() > 0 || AutoWrapText.Get();
 
 		if (bCanWrap && !NewDesiredSize.Equals(LastDesiredSize))
 		{
-			const_cast<STextBlock*>(this)->Invalidate(EInvalidateWidget::Layout);
+			const_cast<STextBlock*>(this)->Invalidate(EInvalidateWidgetReason::Layout);
 		}
 	}
 
@@ -254,76 +269,52 @@ FVector2D STextBlock::ComputeDesiredSize(float LayoutScaleMultiplier) const
 	{
 		const FVector2D LocalShadowOffset = GetShadowOffset();
 
-		const float LocalOutlineSize = GetFont().OutlineSettings.OutlineSize;
+		const float LocalOutlineSize = (float)(GetFont().OutlineSettings.OutlineSize);
 
 		// Account for the outline width impacting both size of the text by multiplying by 2
 		// Outline size in Y is accounted for in MaxHeight calculation in Measure()
-		const FVector2D ComputedOutlineSize(LocalOutlineSize * 2, LocalOutlineSize);
+		const FVector2D ComputedOutlineSize(LocalOutlineSize * 2.f, LocalOutlineSize);
 		const FVector2D TextSize = FSlateApplication::Get().GetRenderer()->GetFontMeasureService()->Measure(GetText(), GetFont()) + ComputedOutlineSize + LocalShadowOffset;
 
-		CachedSimpleDesiredSize = FVector2D(FMath::Max(MinDesiredWidth.Get(0.0f), TextSize.X), TextSize.Y);
+		CachedSimpleDesiredSize = FVector2D(FMath::Max(MinDesiredWidth.Get(), TextSize.X), TextSize.Y);
 		return CachedSimpleDesiredSize.GetValue();
 	}
 	else
 	{
 		// ComputeDesiredSize will also update the text layout cache if required
 		const FVector2D TextSize = TextLayoutCache->ComputeDesiredSize(
-			FSlateTextBlockLayout::FWidgetArgs(BoundText, HighlightText, WrapTextAt, AutoWrapText, WrappingPolicy, GetTransformPolicy(), Margin, LineHeightPercentage, Justification),
-			LayoutScaleMultiplier, GetComputedTextStyle()
-		);
-
-		return FVector2D(FMath::Max(MinDesiredWidth.Get(0.0f), TextSize.X), TextSize.Y);
+			FSlateTextBlockLayout::FWidgetDesiredSizeArgs(
+				BoundText.Get(),
+				HighlightText.Get(),
+				WrapTextAt.Get(),
+				AutoWrapText.Get(),
+				WrappingPolicy.Get(),
+				GetTransformPolicyImpl(),
+				Margin.Get(),
+				LineHeightPercentage.Get(),
+				Justification.Get()),
+			LayoutScaleMultiplier, GetComputedTextStyle());
+
+		return FVector2D(FMath::Max(MinDesiredWidth.Get(), TextSize.X), TextSize.Y);
 	}
 }
 
-bool STextBlock::ComputeVolatility() const
+void STextBlock::SetFont(TAttribute<FSlateFontInfo> InFont)
 {
-	SCOPE_CYCLE_COUNTER(Stat_SlateTextBlockCV);
-	return SLeafWidget::ComputeVolatility() 
-		|| BoundText.IsBound()
-		|| Font.IsBound()
-		|| ColorAndOpacity.IsBound()
-		|| ShadowOffset.IsBound()
-		|| ShadowColorAndOpacity.IsBound()
-		|| HighlightColor.IsBound()
-		|| HighlightShape.IsBound()
-		|| HighlightText.IsBound()
-		|| WrapTextAt.IsBound()
-		|| AutoWrapText.IsBound()
-		|| WrappingPolicy.IsBound()
-		|| TransformPolicy.IsBound()
-		|| Margin.IsBound()
-		|| Justification.IsBound()
-		|| LineHeightPercentage.IsBound()
-		|| MinDesiredWidth.IsBound();
+	bIsAttributeFontSet = InFont.IsSet();
+	Font.Assign(*this, MoveTemp(InFont));
 }
 
-void STextBlock::SetFont(const TAttribute< FSlateFontInfo >& InFont)
+void STextBlock::SetStrikeBrush(TAttribute<const FSlateBrush*> InStrikeBrush)
 {
-	if(!Font.IsSet() || !Font.IdenticalTo(InFont))
-	{
-		Font = InFont;
-		InvalidateText(EInvalidateWidget::LayoutAndVolatility);
-	}
+	bIsAttributeStrikeBrushSet = InStrikeBrush.IsSet();
+	StrikeBrush.Assign(*this, MoveTemp(InStrikeBrush));
 }
 
-void STextBlock::SetStrikeBrush(const TAttribute<const FSlateBrush*>& InStrikeBrush)
+void STextBlock::SetColorAndOpacity(TAttribute<FSlateColor> InColorAndOpacity)
 {
-	if (!StrikeBrush.IsSet() || !StrikeBrush.IdenticalTo(InStrikeBrush))
-	{
-		StrikeBrush = InStrikeBrush;
-		InvalidateText(EInvalidateWidget::LayoutAndVolatility);
-	}
-}
-
-void STextBlock::SetColorAndOpacity(const TAttribute<FSlateColor>& InColorAndOpacity)
-{
-	if ( !ColorAndOpacity.IsSet() || !ColorAndOpacity.IdenticalTo(InColorAndOpacity) )
-	{
-		ColorAndOpacity = InColorAndOpacity;
-		// HACK: Normally this would be Paint only, but textblocks need to recache layout.
-		InvalidateText(EInvalidateWidget::LayoutAndVolatility);
-	}
+	bIsAttributeColorAndOpacitySet = InColorAndOpacity.IsSet();
+	ColorAndOpacity.Assign(*this, MoveTemp(InColorAndOpacity));
 }
 
 void STextBlock::SetTextStyle(const FTextBlockStyle* InTextStyle)
@@ -338,7 +329,7 @@ void STextBlock::SetTextStyle(const FTextBlockStyle* InTextStyle)
 		TextStyle = *Defaults._TextStyle;
 	}
 
-	InvalidateText(EInvalidateWidget::Layout);
+	InvalidateText(EInvalidateWidgetReason::Layout);
 }
 
 void STextBlock::SetTextShapingMethod(const TOptional<ETextShapingMethod>& InTextShapingMethod)
@@ -346,7 +337,7 @@ void STextBlock::SetTextShapingMethod(const TOptional<ETextShapingMethod>& InTex
 	if (!bSimpleTextMode)
 	{
 		TextLayoutCache->SetTextShapingMethod(InTextShapingMethod);
-		InvalidateText(EInvalidateWidget::Layout);
+		InvalidateText(EInvalidateWidgetReason::Layout);
 	}
 }
 
@@ -355,121 +346,82 @@ void STextBlock::SetTextFlowDirection(const TOptional<ETextFlowDirection>& InTex
 	if(!bSimpleTextMode)
 	{
 		TextLayoutCache->SetTextFlowDirection(InTextFlowDirection);
-		InvalidateText(EInvalidateWidget::Layout);
+		InvalidateText(EInvalidateWidgetReason::Layout);
 	}
 }
 
-void STextBlock::SetWrapTextAt(const TAttribute<float>& InWrapTextAt)
+void STextBlock::SetWrapTextAt(TAttribute<float> InWrapTextAt)
 {
-	if(!WrapTextAt.IdenticalTo(InWrapTextAt))
-	{
-		WrapTextAt = InWrapTextAt;
-		InvalidateText(EInvalidateWidget::LayoutAndVolatility);
-	}
+	WrapTextAt.Assign(*this, MoveTemp(InWrapTextAt), 0.f);
 }
 
-void STextBlock::SetAutoWrapText(const TAttribute<bool>& InAutoWrapText)
+void STextBlock::SetAutoWrapText(TAttribute<bool> InAutoWrapText)
 {
-	if(!AutoWrapText.IdenticalTo(InAutoWrapText))
-	{
-		AutoWrapText = InAutoWrapText;
-		InvalidateText(EInvalidateWidget::LayoutAndVolatility);
-	}
+	AutoWrapText.Assign(*this, MoveTemp(InAutoWrapText), 0.f);
 }
 
-void STextBlock::SetWrappingPolicy(const TAttribute<ETextWrappingPolicy>& InWrappingPolicy)
+void STextBlock::SetWrappingPolicy(TAttribute<ETextWrappingPolicy> InWrappingPolicy)
 {
-	if(!WrappingPolicy.IdenticalTo(InWrappingPolicy))
-	{
-		WrappingPolicy = InWrappingPolicy;
-		InvalidateText(EInvalidateWidget::LayoutAndVolatility);
-	}
+	WrappingPolicy.Assign(*this, MoveTemp(InWrappingPolicy));
 }
 
-void STextBlock::SetTransformPolicy(const TAttribute<ETextTransformPolicy>& InTransformPolicy)
+void STextBlock::SetTransformPolicy(TAttribute<ETextTransformPolicy> InTransformPolicy)
 {
-	if(!TransformPolicy.IdenticalTo(InTransformPolicy))
-	{
-		TransformPolicy = InTransformPolicy;
-		InvalidateText(EInvalidateWidget::LayoutAndVolatility);
-	}
+	bIsAttributeTransformPolicySet = InTransformPolicy.IsSet();
+	TransformPolicy.Assign(*this, MoveTemp(InTransformPolicy));
 }
 
-ETextTransformPolicy STextBlock::GetTransformPolicy() const
+void STextBlock::SetShadowOffset(TAttribute<FVector2D> InShadowOffset)
 {
-	if (TransformPolicy.IsSet())
-	{
-		return TransformPolicy.Get();
-	}
-	return TextStyle.TransformPolicy;
+	bIsAttributeShadowOffsetSet = InShadowOffset.IsSet();
+	ShadowOffset.Assign(*this, MoveTemp(InShadowOffset));
 }
 
-void STextBlock::SetShadowOffset(const TAttribute<FVector2D>& InShadowOffset)
+void STextBlock::SetShadowColorAndOpacity(TAttribute<FLinearColor> InShadowColorAndOpacity)
 {
-	if(!ShadowOffset.IdenticalTo(InShadowOffset))
-	{
-		ShadowOffset = InShadowOffset;
-		InvalidateText(EInvalidateWidget::LayoutAndVolatility);
-	}
+	bIsAttributeShadowColorAndOpacitySet = InShadowColorAndOpacity.IsSet();
+	ShadowColorAndOpacity.Assign(*this, MoveTemp(InShadowColorAndOpacity));
 }
 
-void STextBlock::SetShadowColorAndOpacity(const TAttribute<FLinearColor>& InShadowColorAndOpacity)
+void STextBlock::SetHighlightColor(TAttribute<FLinearColor> InHighlightColor)
 {
-	if(!ShadowColorAndOpacity.IdenticalTo(InShadowColorAndOpacity))
-	{
-		ShadowColorAndOpacity = InShadowColorAndOpacity;
-		// HACK: Normally this would be Paint only, but textblocks need to recache layout.
-		InvalidateText(EInvalidateWidget::LayoutAndVolatility);
-	}
+	bIsAttributeHighlightColorSet = InHighlightColor.IsSet();
+	HighlightColor.Assign(*this, MoveTemp(InHighlightColor));
 }
 
-void STextBlock::SetMinDesiredWidth(const TAttribute<float>& InMinDesiredWidth)
+void STextBlock::SetHighlightShape(TAttribute<const FSlateBrush*> InHighlightShape)
 {
-	if(!MinDesiredWidth.IdenticalTo(InMinDesiredWidth))
-	{
-		MinDesiredWidth = InMinDesiredWidth;
-		InvalidateText(EInvalidateWidget::LayoutAndVolatility);
-	}
+	bIsAttributeHighlightShapeSet = InHighlightShape.IsSet();
+	HighlightShape.Assign(*this, MoveTemp(InHighlightShape));
 }
 
-void STextBlock::SetLineHeightPercentage(const TAttribute<float>& InLineHeightPercentage)
+void STextBlock::SetMinDesiredWidth(TAttribute<float> InMinDesiredWidth)
 {
-	if(!LineHeightPercentage.IdenticalTo(InLineHeightPercentage))
-	{
-		LineHeightPercentage = InLineHeightPercentage;
-		InvalidateText(EInvalidateWidget::LayoutAndVolatility);
-	}
+	MinDesiredWidth.Assign(*this, MoveTemp(InMinDesiredWidth), 0.f);
 }
 
-void STextBlock::SetMargin(const TAttribute<FMargin>& InMargin)
+void STextBlock::SetLineHeightPercentage(TAttribute<float> InLineHeightPercentage)
 {
-	if(!Margin.IdenticalTo(InMargin))
-	{
-		Margin = InMargin;
-		InvalidateText(EInvalidateWidget::LayoutAndVolatility);
-	}
+	LineHeightPercentage.Assign(*this, MoveTemp(InLineHeightPercentage));
 }
 
-void STextBlock::SetJustification(const TAttribute<ETextJustify::Type>& InJustification)
+void STextBlock::SetMargin(TAttribute<FMargin> InMargin)
 {
-	if(!Justification.IdenticalTo(InJustification))
-	{
-		Justification = InJustification;
-		InvalidateText(EInvalidateWidget::LayoutAndVolatility);
-	}
+	Margin.Assign(*this, MoveTemp(InMargin));
+}
+
+void STextBlock::SetJustification(TAttribute<ETextJustify::Type> InJustification)
+{
+	Justification.Assign(*this, MoveTemp(InJustification));
 }
 
 FTextBlockStyle STextBlock::GetComputedTextStyle() const
 {
 	FTextBlockStyle ComputedStyle = TextStyle;
 	ComputedStyle.SetFont( GetFont() );
-	if (StrikeBrush.IsSet())
+	if (const FSlateBrush* const ComputedStrikeBrush = GetStrikeBrush())
 	{
-		const FSlateBrush* const ComputedStrikeBrush = StrikeBrush.Get();
-		if (ComputedStrikeBrush)
-		{
-			ComputedStyle.SetStrikeBrush(*ComputedStrikeBrush);
-		}
+		ComputedStyle.SetStrikeBrush(*ComputedStrikeBrush);
 	}
 	ComputedStyle.SetColorAndOpacity( GetColorAndOpacity() );
 	ComputedStyle.SetShadowOffset( GetShadowOffset() );
diff --git a/Engine/Source/Runtime/Slate/Private/Widgets/Text/SlateEditableTextLayout.cpp b/Engine/Source/Runtime/Slate/Private/Widgets/Text/SlateEditableTextLayout.cpp
--- a/Engine/Source/Runtime/Slate/Private/Widgets/Text/SlateEditableTextLayout.cpp
+++ b/Engine/Source/Runtime/Slate/Private/Widgets/Text/SlateEditableTextLayout.cpp
@@ -3428,9 +3428,17 @@ FVector2D FSlateEditableTextLayout::ComputeDesiredSize(float LayoutScaleMultipli
 		MarginValue.Right += CaretWidth;
 
 		const FVector2D HintTextSize = HintTextLayout->ComputeDesiredSize(
-			FSlateTextBlockLayout::FWidgetArgs(HintText, FText::GetEmpty(), WrapTextAt, AutoWrapText, WrappingPolicy, ETextTransformPolicy::None, MarginValue, LineHeightPercentage, Justification),
-			LayoutScaleMultiplier, HintTextStyle
-			);
+			FSlateTextBlockLayout::FWidgetDesiredSizeArgs(
+				HintText.Get(FText::GetEmpty()),
+				FText::GetEmpty(),
+				WrapTextAt.Get(0.f),
+				AutoWrapText.Get(false),
+				WrappingPolicy.Get(),
+				ETextTransformPolicy::None,
+				MarginValue,
+				LineHeightPercentage.Get(),
+				Justification.Get()),
+			LayoutScaleMultiplier, HintTextStyle);
 
 		// If a wrapping width has been provided, then we need to report that as the desired width
 		DesiredWidth = WrappingWidth > 0 ? WrappingWidth : HintTextSize.X;
diff --git a/Engine/Source/Runtime/Slate/Private/Widgets/Text/SlateTextBlockLayout.cpp b/Engine/Source/Runtime/Slate/Private/Widgets/Text/SlateTextBlockLayout.cpp
--- a/Engine/Source/Runtime/Slate/Private/Widgets/Text/SlateTextBlockLayout.cpp
+++ b/Engine/Source/Runtime/Slate/Private/Widgets/Text/SlateTextBlockLayout.cpp
@@ -25,6 +25,80 @@ FSlateTextBlockLayout::FSlateTextBlockLayout(SWidget* InOwner, FTextBlockStyle I
 	TextLayout->SetLineBreakIterator(MoveTemp(InLineBreakPolicy));
 }
 
+FVector2D FSlateTextBlockLayout::ComputeDesiredSize(const FWidgetDesiredSizeArgs& InWidgetArgs, const float InScale, const FTextBlockStyle& InTextStyle)
+{
+	// Cache the wrapping rules so that we can recompute the wrap at width in paint.
+	CachedWrapTextAt = InWidgetArgs.WrapTextAt;
+	bCachedAutoWrapText = InWidgetArgs.AutoWrapText;
+
+	const ETextTransformPolicy PreviousTransformPolicy = TextLayout->GetTransformPolicy();
+
+	// Set the text layout information
+	TextLayout->SetScale(InScale);
+	TextLayout->SetWrappingWidth(CalculateWrappingWidth());
+	TextLayout->SetWrappingPolicy(InWidgetArgs.WrappingPolicy);
+	TextLayout->SetTransformPolicy(InWidgetArgs.TransformPolicy);
+	TextLayout->SetMargin(InWidgetArgs.Margin);
+	TextLayout->SetJustification(InWidgetArgs.Justification);
+	TextLayout->SetLineHeightPercentage(InWidgetArgs.LineHeightPercentage);
+
+	// Has the transform policy changed? If so we need a full refresh as that is destructive to the model text
+	if (PreviousTransformPolicy != TextLayout->GetTransformPolicy())
+	{
+		Marshaller->MakeDirty();
+	}
+
+	// Has the style used for this text block changed?
+	if (!IsStyleUpToDate(InTextStyle))
+	{
+		TextLayout->SetDefaultTextStyle(InTextStyle);
+		Marshaller->MakeDirty(); // will regenerate the text using the new default style
+	}
+
+	{
+		bool bRequiresTextUpdate = false;
+		const FText& TextToSet = InWidgetArgs.Text;
+		if (!TextLastUpdate.IdenticalTo(TextToSet))
+		{
+			// The pointer used by the bound text has changed, however the text may still be the same - check that now
+			if (!TextLastUpdate.IsDisplayStringEqualTo(TextToSet))
+			{
+				// The source text has changed, so update the internal text
+				bRequiresTextUpdate = true;
+			}
+
+			// Update this even if the text is lexically identical, as it will update the pointer compared by IdenticalTo for the next Tick
+			TextLastUpdate = FTextSnapshot(TextToSet);
+		}
+
+		if (bRequiresTextUpdate || Marshaller->IsDirty())
+		{
+			UpdateTextLayout(TextToSet);
+		}
+	}
+
+	{
+		const FText& HighlightTextToSet = InWidgetArgs.HighlightText;
+		if (!HighlightTextLastUpdate.IdenticalTo(HighlightTextToSet))
+		{
+			// The pointer used by the bound text has changed, however the text may still be the same - check that now
+			if (!HighlightTextLastUpdate.IsDisplayStringEqualTo(HighlightTextToSet))
+			{
+				UpdateTextHighlights(HighlightTextToSet);
+			}
+
+			// Update this even if the text is lexically identical, as it will update the pointer compared by IdenticalTo for the next Tick
+			HighlightTextLastUpdate = FTextSnapshot(HighlightTextToSet);
+		}
+	}
+
+	// We need to update our size if the text layout has become dirty
+	TextLayout->UpdateIfNeeded();
+
+	return TextLayout->GetSize();
+}
+
+PRAGMA_DISABLE_DEPRECATION_WARNINGS
 FVector2D FSlateTextBlockLayout::ComputeDesiredSize(const FWidgetArgs& InWidgetArgs, const float InScale, const FTextBlockStyle& InTextStyle)
 {
 	// Cache the wrapping rules so that we can recompute the wrap at width in paint.
@@ -97,6 +171,7 @@ FVector2D FSlateTextBlockLayout::ComputeDesiredSize(const FWidgetArgs& InWidgetA
 
 	return TextLayout->GetSize();
 }
+PRAGMA_DISABLE_DEPRECATION_WARNINGS
 
 FVector2D FSlateTextBlockLayout::GetDesiredSize() const
 {
diff --git a/Engine/Source/Runtime/Slate/Public/Widgets/Text/SInlineEditableTextBlock.h b/Engine/Source/Runtime/Slate/Public/Widgets/Text/SInlineEditableTextBlock.h
--- a/Engine/Source/Runtime/Slate/Public/Widgets/Text/SInlineEditableTextBlock.h
+++ b/Engine/Source/Runtime/Slate/Public/Widgets/Text/SInlineEditableTextBlock.h
@@ -30,6 +30,9 @@ DECLARE_DELEGATE_OneParam(FOnBeginTextEdit, const FText&)
  */
 class SLATE_API SInlineEditableTextBlock: public SCompoundWidget
 {
+	SLATE_DECLARE_WIDGET(SInlineEditableTextBlock, SCompoundWidget)
+
+public:
 	SLATE_BEGIN_ARGS( SInlineEditableTextBlock )
 		: _Text()
 		, _Style( &FCoreStyle::Get().GetWidgetStyle<FInlineEditableTextBlockStyle>("InlineEditableTextBlockStyle") )
@@ -105,6 +108,7 @@ class SLATE_API SInlineEditableTextBlock: public SCompoundWidget
 		SLATE_EVENT( FOnVerifyTextChanged, OnVerifyTextChanged )
 	SLATE_END_ARGS()
 
+	SInlineEditableTextBlock();
 	~SInlineEditableTextBlock();
 
 	/**
@@ -143,11 +147,13 @@ class SLATE_API SInlineEditableTextBlock: public SCompoundWidget
 
 	void SetReadOnly(bool bInIsReadOnly);
 
-	void SetText( const TAttribute< FText >& InText );
+	void SetText(TAttribute<FText> InText);
+
+	UE_DEPRECATED(5.0, "SetText taking FString is deprecated. Use the FText version instead")
 	void SetText( const FString& InText );
 
 	/** Sets the wrap text at attribute.  See WrapTextAt attribute */
-	void SetWrapTextAt(const TAttribute<float>& InWrapTextAt);
+	void SetWrapTextAt(TAttribute<float> InWrapTextAt);
 
 protected:
 	/** Callback for the text box's OnTextChanged event */
@@ -159,6 +165,12 @@ protected:
 	/** Cancels the edit mode and switches back to label mode */
 	void CancelEditMode();
 
+	/** @return an attribute reference of Text */
+	TSlateAttributeRef<FText> GetTextAttribute() const { return TSlateAttributeRef<FText>(*this, TextAttribute); }
+
+	/** @return an attribute reference of ColorAndOpacity */
+	TSlateAttributeRef<bool> GetIsReadOnlyAttribute() const { return TSlateAttributeRef<bool>(*this, bIsReadOnlyAttribute); }
+
 protected:
 	/** The widget used when in label mode */
 	TSharedPtr< STextBlock > TextBlock;
@@ -191,38 +203,45 @@ protected:
 	/** Callback to verify text when changed. Will return an error message to denote problems. */
 	FOnVerifyTextChanged OnVerifyTextChanged;
 
-	/** Attribute for the text to use for the widget */
-	TAttribute<FText> Text;
-
-	/** Attribute to look up if the widget is read-only */
-	TAttribute< bool > bIsReadOnly;
-
 	/** Widget to focus when we finish editing */
 	TWeakPtr<SWidget> WidgetToFocus;
 
+	/**
+	 * When selection of widget is managed by another widget, this delays the "double select" from
+	 * occurring immediately, offering a chance for double clicking to take action.
+	 */
+	float DoubleSelectDelay;
+
+	/** Attribute to look up if the widget is multiline */
+	bool bIsMultiLine;
+
+//#if WITH_EDITORONLY_DATA
+//	UE_DEPRECATED(5.0, "Direct access to Text is now deprecated. Use the setter or getter.")
+//	FSlateDeprecatedTAttribute<FText> Text;
+//	UE_DEPRECATED(5.0, "Direct access to bIsReadOnly is now deprecated. Use the setter or getter.")
+//	FSlateDeprecatedTAttribute<bool> bIsReadOnly;
+//#endif
+
 private:
 
 	/** Get editable text widget */
 	TSharedPtr<SWidget> GetEditableTextWidget() const;
 
 	/** Set editable text */
-	void SetEditableText( const TAttribute< FText >& InNewText );
+	void SetEditableText(TAttribute<FText> InNewText);
 
 	/** Set textbox error */
-	void SetTextBoxError( const FText& ErrorText );
+	void SetTextBoxError(const FText& ErrorText);
 
 	/** Active timer to trigger entry into edit mode after a delay */
 	EActiveTimerReturnType TriggerEditMode(double InCurrentTime, float InDeltaTime);
 
+	/** Attribute for the text to use for the widget */
+	TSlateAttribute<FText> TextAttribute;
+
 	/** The handle to the active timer */
 	TWeakPtr<FActiveTimerHandle> ActiveTimerHandle;
 
-protected:
-	/** When selection of widget is managed by another widget, this delays the "double select" from 
-	occurring immediately, offering a chance for double clicking to take action. */
-	float DoubleSelectDelay;
-
-	/** Attribute to look up if the widget is multiline */
-	uint8 bIsMultiLine:1;
-
+	/** Attribute to look up if the widget is read-only */
+	TSlateAttribute<bool> bIsReadOnlyAttribute;
 };
diff --git a/Engine/Source/Runtime/Slate/Public/Widgets/Text/SRichTextBlock.h b/Engine/Source/Runtime/Slate/Public/Widgets/Text/SRichTextBlock.h
--- a/Engine/Source/Runtime/Slate/Public/Widgets/Text/SRichTextBlock.h
+++ b/Engine/Source/Runtime/Slate/Public/Widgets/Text/SRichTextBlock.h
@@ -7,7 +7,7 @@
 #include "Layout/Margin.h"
 #include "Styling/SlateWidgetStyleAsset.h"
 #include "Widgets/DeclarativeSyntaxSupport.h"
-#include "Widgets/SWidget.h"
+#include "Widgets/SLeafWidget.h"
 #include "Styling/SlateTypes.h"
 #include "Styling/CoreStyle.h"
 #include "Framework/Text/TextLayout.h"
@@ -30,8 +30,9 @@ enum class ETextShapingMethod : uint8;
  * A rich static text widget. 
  * Through the use of markup and text decorators, text with different styles, embedded image and widgets can be achieved.
  */
-class SLATE_API SRichTextBlock : public SWidget
+class SLATE_API SRichTextBlock : public SLeafWidget
 {
+	SLATE_DECLARE_WIDGET(SRichTextBlock, SLeafWidget)
 public:
 
 	SLATE_BEGIN_ARGS( SRichTextBlock )
@@ -171,16 +172,16 @@ public:
 	 */
 	const FText& GetText() const
 	{
-		return BoundText.Get(FText::GetEmpty());
+		return BoundText.Get();
 	}
 
 	/**
 	 * Sets the text for this text block
 	 */
-	void SetText( const TAttribute<FText>& InTextAttr );
+	void SetText(TAttribute<FText> InTextAttr );
 
 	/** See HighlightText attribute */
-	void SetHighlightText( const TAttribute<FText>& InHighlightText );
+	void SetHighlightText(TAttribute<FText> InHighlightText);
 
 	/** See TextShapingMethod attribute */
 	void SetTextShapingMethod(const TOptional<ETextShapingMethod>& InTextShapingMethod);
@@ -189,31 +190,31 @@ public:
 	void SetTextFlowDirection(const TOptional<ETextFlowDirection>& InTextFlowDirection);
 
 	/** See WrapTextAt attribute */
-	void SetWrapTextAt(const TAttribute<float>& InWrapTextAt);
+	void SetWrapTextAt(TAttribute<float> InWrapTextAt);
 
 	/** See AutoWrapText attribute */
-	void SetAutoWrapText(const TAttribute<bool>& InAutoWrapText);
+	void SetAutoWrapText(TAttribute<bool> InAutoWrapText);
 
 	/** Set WrappingPolicy attribute */
-	void SetWrappingPolicy(const TAttribute<ETextWrappingPolicy>& InWrappingPolicy);
+	void SetWrappingPolicy(TAttribute<ETextWrappingPolicy> InWrappingPolicy);
 
 	/** Set TransformPolicy attribute */
-	void SetTransformPolicy(const TAttribute<ETextTransformPolicy>& InTransformPolicy);
+	void SetTransformPolicy(TAttribute<ETextTransformPolicy> InTransformPolicy);
 
 	/** See LineHeightPercentage attribute */
-	void SetLineHeightPercentage(const TAttribute<float>& InLineHeightPercentage);
+	void SetLineHeightPercentage(TAttribute<float> InLineHeightPercentage);
 
 	/** See Margin attribute */
-	void SetMargin(const TAttribute<FMargin>& InMargin);
+	void SetMargin(TAttribute<FMargin> InMargin);
 
 	/** See Justification attribute */
-	void SetJustification(const TAttribute<ETextJustify::Type>& InJustification);
+	void SetJustification(TAttribute<ETextJustify::Type> InJustification);
 
 	/** See TextStyle argument */
 	void SetTextStyle(const FTextBlockStyle& InTextStyle);
 
 	/** See MinDesiredWidth attribute */
-	void SetMinDesiredWidth(const TAttribute<float>& InMinDesiredWidth);
+	void SetMinDesiredWidth(TAttribute<float> InMinDesiredWidth);
 
 	/**  */
 	void SetDecoratorStyleSet(const ISlateStyle* NewDecoratorStyleSet);
@@ -229,48 +230,47 @@ public:
 protected:
 	//~ SWidget interface
 	virtual FVector2D ComputeDesiredSize(float) const override;
-	virtual bool ComputeVolatility() const override;
 	//~ End of SWidget interface
 
 private:
 
 	/** The text displayed in this text block */
-	TAttribute< FText > BoundText;
+	TSlateAttribute<FText> BoundText;
 
 	/** The wrapped layout for this text block */
-	TUniquePtr< FSlateTextBlockLayout > TextLayoutCache;
+	TUniquePtr<FSlateTextBlockLayout> TextLayoutCache;
 
 	/** Default style used by the TextLayout */
 	FTextBlockStyle TextStyle;
 
-	/** Highlight this text in the textblock */
-	TAttribute<FText> HighlightText;
+	/** Highlight this text in the text block */
+	TSlateAttribute<FText> HighlightText;
 
 	/** Whether text wraps onto a new line when it's length exceeds this width; if this value is zero or negative, no wrapping occurs. */
-	TAttribute<float> WrapTextAt;
-	
-	/** True if we're wrapping text automatically based on the computed horizontal space for this widget */
-	TAttribute<bool> AutoWrapText;
+	TSlateAttribute<float> WrapTextAt;
 
 	/** The wrapping policy we're using */
-	TAttribute<ETextWrappingPolicy> WrappingPolicy;
+	TSlateAttribute<ETextWrappingPolicy> WrappingPolicy;
 
 	/** The transform policy we're using */
-	TAttribute<ETextTransformPolicy> TransformPolicy;
-
-	/** The amount of blank space left around the edges of text area. */
-	TAttribute< FMargin > Margin;
+	TSlateAttribute<ETextTransformPolicy> TransformPolicy;
 
 	/** The amount to scale each lines height by. */
-	TAttribute< ETextJustify::Type > Justification; 
+	TSlateAttribute<ETextJustify::Type> Justification;
+
+	/** True if we're wrapping text automatically based on the computed horizontal space for this widget */
+	TSlateAttribute<bool> AutoWrapText;
+
+	/** The amount of blank space left around the edges of text area. */
+	TSlateAttribute<FMargin> Margin;
 
 	/** How the text should be aligned with the margin. */
-	TAttribute< float > LineHeightPercentage;
+	TSlateAttribute<float> LineHeightPercentage;
 
 	/** Prevents the text block from being smaller than desired in certain cases (e.g. when it is empty) */
-	TAttribute<float> MinDesiredWidth;
+	TSlateAttribute<float> MinDesiredWidth;
 
-	/** Use to Scale the entire Text Block*/
+	/** Use to Scale the entire text block */
 	float TextBlockScale = 1.0f;
 
 	/**  */
diff --git a/Engine/Source/Runtime/Slate/Public/Widgets/Text/STextBlock.h b/Engine/Source/Runtime/Slate/Public/Widgets/Text/STextBlock.h
--- a/Engine/Source/Runtime/Slate/Public/Widgets/Text/STextBlock.h
+++ b/Engine/Source/Runtime/Slate/Public/Widgets/Text/STextBlock.h
@@ -38,6 +38,7 @@ namespace ETextRole
  */
 class SLATE_API STextBlock : public SLeafWidget
 {
+	SLATE_DECLARE_WIDGET(STextBlock, SLeafWidget)
 
 public:
 
@@ -92,7 +93,7 @@ public:
 		/** The color used to highlight the specified text */
 		SLATE_ATTRIBUTE( FLinearColor, HighlightColor )
 
-		/** The brush used to highlight the specified text*/
+		/** The brush used to highlight the specified text */
 		SLATE_ATTRIBUTE( const FSlateBrush*, HighlightShape )
 
 		/** Highlight this text in the text block */
@@ -169,38 +170,20 @@ public:
 	}
 	
 public:
-	UE_DEPRECATED(4.24, "SetText taking FString is deprecated. Use the FText version instead")
-	void SetText( const TAttribute< FString >& InText );
-	UE_DEPRECATED(4.24, "SetText taking FString is deprecated. Use the FText version instead")
-	void SetText( const FString& InText );
+	/** Sets the text for this text block */
+	void SetText(TAttribute<FText> InText);
 
-	/**
-	 * Sets the text for this text block
-	 *
-	 * @param	InText	The new text to display
-	 */
-	void SetText( const TAttribute< FText >& InText );
-	void SetText( const FText& InText );
-
-	/**
-	* Sets the highlight text for this text block 
-	*
-	* @param	InText	The new text to highlight
-	*/
+	/** Sets the highlight text for this text block */
 	void SetHighlightText(TAttribute<FText> InText);
 
-	/**
-	 * Sets the font used to draw the text
-	 *
-	 * @param	InFont	The new font to use
-	 */
-	void SetFont(const TAttribute< FSlateFontInfo >& InFont);
+	/** Sets the font used to draw the text	*/
+	void SetFont(TAttribute<FSlateFontInfo> InFont);
 
 	/** Sets the brush used to strike through the text */
-	void SetStrikeBrush(const TAttribute<const FSlateBrush*>& InStrikeBrush);
+	void SetStrikeBrush(TAttribute<const FSlateBrush*> InStrikeBrush);
 
 	/** See ColorAndOpacity attribute */
-	void SetColorAndOpacity(const TAttribute<FSlateColor>& InColorAndOpacity);
+	void SetColorAndOpacity(TAttribute<FSlateColor> InColorAndOpacity);
 
 	/** See TextStyle argument */
 	void SetTextStyle(const FTextBlockStyle* InTextStyle);
@@ -212,35 +195,44 @@ public:
 	void SetTextFlowDirection(const TOptional<ETextFlowDirection>& InTextFlowDirection);
 
 	/** See WrapTextAt attribute */
-	void SetWrapTextAt(const TAttribute<float>& InWrapTextAt);
+	void SetWrapTextAt(TAttribute<float> InWrapTextAt);
 
 	/** See AutoWrapText attribute */
-	void SetAutoWrapText(const TAttribute<bool>& InAutoWrapText);
+	void SetAutoWrapText(TAttribute<bool> InAutoWrapText);
 
 	/** Set WrappingPolicy attribute */
-	void SetWrappingPolicy(const TAttribute<ETextWrappingPolicy>& InWrappingPolicy);
+	void SetWrappingPolicy(TAttribute<ETextWrappingPolicy> InWrappingPolicy);
 
 	/** Set TransformPolicy attribute */
-	void SetTransformPolicy(const TAttribute<ETextTransformPolicy>& InTransformPolicy);
-	ETextTransformPolicy GetTransformPolicy() const;
+	void SetTransformPolicy(TAttribute<ETextTransformPolicy> InTransformPolicy);
+
+	/** Get TransformPolicy attribute */
+	UE_DEPRECATED(5.0, "GetTransformPolicy is not accessible anymore since it's attribute value may not have been updated yet.")
+	ETextTransformPolicy GetTransformPolicy() const { return GetTransformPolicyImpl(); }
 
 	/** See ShadowOffset attribute */
-	void SetShadowOffset(const TAttribute<FVector2D>& InShadowOffset);
+	void SetShadowOffset(TAttribute<FVector2D> InShadowOffset);
 
 	/** See ShadowColorAndOpacity attribute */
-	void SetShadowColorAndOpacity(const TAttribute<FLinearColor>& InShadowColorAndOpacity);
+	void SetShadowColorAndOpacity(TAttribute<FLinearColor> InShadowColorAndOpacity);
+
+	/** See HighlightColor attribute */
+	void SetHighlightColor(TAttribute<FLinearColor> InHighlightColor);
+	
+	/** See HighlightShape attribute */
+	void SetHighlightShape(TAttribute<const FSlateBrush*> InHighlightShape);
 
 	/** See MinDesiredWidth attribute */
-	void SetMinDesiredWidth(const TAttribute<float>& InMinDesiredWidth);
+	void SetMinDesiredWidth(TAttribute<float> InMinDesiredWidth);
 
 	/** See LineHeightPercentage attribute */
-	void SetLineHeightPercentage(const TAttribute<float>& InLineHeightPercentage);
+	void SetLineHeightPercentage(TAttribute<float> InLineHeightPercentage);
 
 	/** See Margin attribute */
-	void SetMargin(const TAttribute<FMargin>& InMargin);
+	void SetMargin(TAttribute<FMargin> InMargin);
 
 	/** See Justification attribute */
-	void SetJustification(const TAttribute<ETextJustify::Type>& InJustification);
+	void SetJustification(TAttribute<ETextJustify::Type> InJustification);
 
 	// SWidget interface
 	virtual int32 OnPaint( const FPaintArgs& Args, const FGeometry& AllottedGeometry, const FSlateRect& MyCullingRect, FSlateWindowElementList& OutDrawElements, int32 LayerId, const FWidgetStyle& InWidgetStyle, bool bParentEnabled ) const override;
@@ -251,11 +243,6 @@ public:
 #endif
 	// End of SWidget interface
 
-protected:
-	// Begin SWidget overrides.
-	virtual bool ComputeVolatility() const override;
-	// End of SWidget interface
-
 private:
 	/** Get the computed text style to use with the text marshaller */
 	FTextBlockStyle GetComputedTextStyle() const;
@@ -268,6 +255,9 @@ private:
 
 	/** Gets the current strike brush */
 	const FSlateBrush* GetStrikeBrush() const;
+
+	/** Get TransformPolicy attribute */
+	ETextTransformPolicy GetTransformPolicyImpl() const;
 	
 	/** Gets the current shadow offset */
 	FVector2D GetShadowOffset() const;
@@ -282,11 +272,14 @@ private:
 	const FSlateBrush* GetHighlightShape() const;
 
 	/** Call to invalidate this text block */
-	void InvalidateText(EInvalidateWidget InvalidateReason);
+	void InvalidateText(EInvalidateWidgetReason InvalidateReason);
+
+private:
+
 
 private:
 	/** The text displayed in this text block */
-	TAttribute< FText > BoundText;
+	TSlateAttribute<FText> BoundText;
 
 	/** The wrapped layout for this text block */
 	TUniquePtr< FSlateTextBlockLayout > TextLayoutCache;
@@ -295,52 +288,73 @@ private:
 	FTextBlockStyle TextStyle;
 
 	/** Sets the font used to draw the text */
-	TAttribute< FSlateFontInfo > Font;
+	TSlateAttribute<FSlateFontInfo> Font;
 
 	/** Sets the brush used to strike through the text */
-	TAttribute< const FSlateBrush* > StrikeBrush;
+	TSlateAttribute<const FSlateBrush*> StrikeBrush;
 
 	/** Text color and opacity */
-	TAttribute<FSlateColor> ColorAndOpacity;
+	TSlateAttribute<FSlateColor> ColorAndOpacity;
 
 	/** Drop shadow offset in pixels */
-	TAttribute< FVector2D > ShadowOffset;
+	TSlateAttribute<FVector2D> ShadowOffset;
 
 	/** Shadow color and opacity */
-	TAttribute<FLinearColor> ShadowColorAndOpacity;
+	TSlateAttribute<FLinearColor> ShadowColorAndOpacity;
 
-	TAttribute<FLinearColor> HighlightColor;
+	/** The color used to highlight the specified text */
+	TSlateAttribute<FLinearColor> HighlightColor;
 
-	TAttribute< const FSlateBrush* > HighlightShape;
+	/** The brush used to highlight the specified text */
+	TSlateAttribute<const FSlateBrush*> HighlightShape;
 
-	/** Highlight this text in the textblock */
-	TAttribute<FText> HighlightText;
+	/** Highlight this text in the TextBlock */
+	TSlateAttribute<FText> HighlightText;
 
 	/** Whether text wraps onto a new line when it's length exceeds this width; if this value is zero or negative, no wrapping occurs. */
-	TAttribute<float> WrapTextAt;
+	TSlateAttribute<float> WrapTextAt;
 
 	/** True if we're wrapping text automatically based on the computed horizontal space for this widget */
-	TAttribute<bool> AutoWrapText;
+	TSlateAttribute<bool> AutoWrapText;
 
 	/** The wrapping policy we're using */
-	TAttribute<ETextWrappingPolicy> WrappingPolicy;
+	TSlateAttribute<ETextWrappingPolicy> WrappingPolicy;
 
 	/** The transform policy we're using */
-	TAttribute<ETextTransformPolicy> TransformPolicy;
+	TSlateAttribute<ETextTransformPolicy> TransformPolicy;
 
 	/** The amount of blank space left around the edges of text area. */
-	TAttribute< FMargin > Margin;
+	TSlateAttribute<FMargin> Margin;
 
 	/** The amount to scale each lines height by. */
-	TAttribute< ETextJustify::Type > Justification; 
+	TSlateAttribute<ETextJustify::Type> Justification;
 
 	/** How the text should be aligned with the margin. */
-	TAttribute< float > LineHeightPercentage;
+	TSlateAttribute<float> LineHeightPercentage;
 
 	/** Prevents the text block from being smaller than desired in certain cases (e.g. when it is empty) */
-	TAttribute<float> MinDesiredWidth;
+	TSlateAttribute<float> MinDesiredWidth;
 
 	/** If this is enabled, text shaping, wrapping, justification are disabled in favor of much faster text layout and measurement. */
 	mutable TOptional<FVector2D> CachedSimpleDesiredSize;
+
+	/** Flags used to check if the SlateAttribute is set. */
+	union
+	{
+		struct 
+		{
+			uint16 bIsAttributeFontSet : 1;
+			uint16 bIsAttributeStrikeBrushSet : 1;
+			uint16 bIsAttributeColorAndOpacitySet : 1;
+			uint16 bIsAttributeShadowOffsetSet : 1;
+			uint16 bIsAttributeShadowColorAndOpacitySet : 1;
+			uint16 bIsAttributeHighlightColorSet : 1;
+			uint16 bIsAttributeHighlightShapeSet : 1;
+			uint16 bIsAttributeWrapTextAtSet : 1;
+			uint16 bIsAttributeTransformPolicySet : 1;
+		};
+		uint16 Union_IsAttributeSet;
+	};
+
 	bool bSimpleTextMode;
 };
diff --git a/Engine/Source/Runtime/Slate/Public/Widgets/Text/SlateTextBlockLayout.h b/Engine/Source/Runtime/Slate/Public/Widgets/Text/SlateTextBlockLayout.h
--- a/Engine/Source/Runtime/Slate/Public/Widgets/Text/SlateTextBlockLayout.h
+++ b/Engine/Source/Runtime/Slate/Public/Widgets/Text/SlateTextBlockLayout.h
@@ -20,7 +20,45 @@ enum class ETextShapingMethod : uint8;
 class SLATE_API FSlateTextBlockLayout
 {
 public:
-	struct FWidgetArgs
+	struct FWidgetDesiredSizeArgs
+	{
+		FORCEINLINE FWidgetDesiredSizeArgs(
+			const FText& InText,
+			const FText& InHighlightText,
+			const float InWrapTextAt,
+			const bool InAutoWrapText,
+			const ETextWrappingPolicy InWrappingPolicy,
+			const ETextTransformPolicy InTransformPolicy,
+			const FMargin& InMargin,
+			const float InLineHeightPercentage,
+			const ETextJustify::Type InJustification
+		)
+			: Text(InText)
+			, HighlightText(InHighlightText)
+			, Margin(InMargin)
+			, WrapTextAt(InWrapTextAt)
+			, LineHeightPercentage(InLineHeightPercentage)
+			, WrappingPolicy(InWrappingPolicy)
+			, TransformPolicy(InTransformPolicy)
+			, Justification(InJustification)
+			, AutoWrapText(InAutoWrapText)
+		{
+		}
+
+		const FText Text = FText::GetEmpty();
+		const FText HighlightText = FText::GetEmpty();
+		const FMargin Margin;
+		const float WrapTextAt = 0.f;
+		const float LineHeightPercentage;
+		const ETextWrappingPolicy WrappingPolicy;
+		const ETextTransformPolicy TransformPolicy;
+		const ETextJustify::Type Justification;
+		const bool AutoWrapText = false;
+	};
+
+	struct
+		UE_DEPRECATED(5.0, "FWidgetArgs is deprecated. Upgrade to FWidgetDesiredSizeArgs instead.")
+		FWidgetArgs
 	{
 		FORCEINLINE FWidgetArgs(
 			const TAttribute<FText>& InText, 
@@ -64,7 +102,15 @@ public:
 	/**
 	 * Get the computed desired size for this layout, updating the internal cache as required
 	 */
+	FVector2D ComputeDesiredSize(const FWidgetDesiredSizeArgs& InWidgetArgs, const float InScale, const FTextBlockStyle& InTextStyle);
+
+	PRAGMA_DISABLE_DEPRECATION_WARNINGS
+	/**
+	 * Get the computed desired size for this layout, updating the internal cache as required
+	 */
+	UE_DEPRECATED(5.0, "FWidgetArgs is deprecated. Upgrade to FWidgetDesiredSizeArgs instead.")
 	FVector2D ComputeDesiredSize(const FWidgetArgs& InWidgetArgs, const float InScale, const FTextBlockStyle& InTextStyle);
+	PRAGMA_DISABLE_DEPRECATION_WARNINGS
 
 	/**
 	 * Gets the last computed desired size.
