Back out changelist 15710618
"Slate: Convert STextBlock Attributes into SlateAttributes." Cause performance issue.
#rnx

diff --git a/Engine/Plugins/VirtualProduction/DMX/DMXEngine/Source/DMXEditor/Private/Widgets/OutputConsole/SDMXFader.cpp b/Engine/Plugins/VirtualProduction/DMX/DMXEngine/Source/DMXEditor/Private/Widgets/OutputConsole/SDMXFader.cpp
--- a/Engine/Plugins/VirtualProduction/DMX/DMXEngine/Source/DMXEditor/Private/Widgets/OutputConsole/SDMXFader.cpp
+++ b/Engine/Plugins/VirtualProduction/DMX/DMXEngine/Source/DMXEditor/Private/Widgets/OutputConsole/SDMXFader.cpp
@@ -33,7 +33,7 @@ void SDMXFader::Construct(const FArguments& InArgs)
 {
 	OnRequestDelete = InArgs._OnRequestDelete;
 	OnRequestSelect = InArgs._OnRequestSelect;
-	FaderName = InArgs._FaderName;
+	FaderName = InArgs._FaderName.ToString();
 
 	// Init args and test them for general DMX validity, no protocol specifics
 	check(InArgs._UniverseID >= 0 && InArgs._UniverseID < DMX_MAX_UNIVERSE);
@@ -494,7 +494,8 @@ void SDMXFader::OnFaderNameCommitted(const FText& NewFaderName, ETextCommit::Typ
 	check(FaderNameTextBox.IsValid());
 
 	FaderNameTextBox->SetText(FaderName);
-	FaderName = NewFaderName;
+
+	FaderName = NewFaderName.ToString();
 }
 
 FReply SDMXFader::OnUniverseIDBorderDoubleClicked(const FGeometry& InMyGeometry, const FPointerEvent& InMouseEvent)
diff --git a/Engine/Plugins/VirtualProduction/DMX/DMXEngine/Source/DMXEditor/Private/Widgets/OutputConsole/SDMXFader.h b/Engine/Plugins/VirtualProduction/DMX/DMXEngine/Source/DMXEditor/Private/Widgets/OutputConsole/SDMXFader.h
--- a/Engine/Plugins/VirtualProduction/DMX/DMXEngine/Source/DMXEditor/Private/Widgets/OutputConsole/SDMXFader.h
+++ b/Engine/Plugins/VirtualProduction/DMX/DMXEngine/Source/DMXEditor/Private/Widgets/OutputConsole/SDMXFader.h
@@ -126,7 +126,7 @@ private:
 
 public:
 	/** Returns the name of the fader */
-	FString GetFaderName() const { return FaderName.ToString(); };
+	FString GetFaderName() const { return FaderName; };
 
 	/** Returns the universe ID to which to should send DMX to */
 	int32 GetUniverseID() const { return UniverseID; }
@@ -151,7 +151,7 @@ public:
 
 private:
 	/** Cached Name of the Fader */
-	FText FaderName;
+	FString FaderName;
 
 	/** The universe the should send to fader */
 	int32 UniverseID;
diff --git a/Engine/Source/Developer/SlateFileDialogs/Private/SlateFileDlgWindow.cpp b/Engine/Source/Developer/SlateFileDialogs/Private/SlateFileDlgWindow.cpp
--- a/Engine/Source/Developer/SlateFileDialogs/Private/SlateFileDlgWindow.cpp
+++ b/Engine/Source/Developer/SlateFileDialogs/Private/SlateFileDlgWindow.cpp
@@ -1456,7 +1456,7 @@ void SSlateFileOpenDlg::SetDefaultFile(FString DefaultFile)
 		SaveFilename = FileList;
 	}
 
-	SaveFilenameEditBox->SetText(FText::FromString(SaveFilename));
+	SaveFilenameEditBox->SetText(SaveFilename);
 }
 
 
@@ -1630,7 +1630,7 @@ FReply SSlateFileOpenDlg::OnNewDirectoryClick()
 {
 	NewDirectorySizeBox->SetVisibility(EVisibility::Visible);
 	NewDirCancelButton->SetVisibility(EVisibility::Visible);
-	NewDirectoryEditBox->SetText(FText::GetEmpty());
+	NewDirectoryEditBox->SetText(FString(""));
 	
 	FSlateApplication::Get().SetKeyboardFocus(NewDirectoryEditBox);	
 	NewDirectoryEditBox->EnterEditingMode();
@@ -1673,7 +1673,7 @@ FReply SSlateFileOpenDlg::OnNewDirectoryAcceptCancelClick(FSlateFileDlgWindow::E
 	NewDirCancelButton->SetVisibility(EVisibility::Hidden);
 	DirErrorMsg->SetVisibility(EVisibility::Collapsed);
 
-	NewDirectoryEditBox->SetText(FText::GetEmpty());
+	NewDirectoryEditBox->SetText(FString(""));
 
 	return FReply::Handled();
 }
diff --git a/Engine/Source/Runtime/Slate/Private/Widgets/Text/SInlineEditableTextBlock.cpp b/Engine/Source/Runtime/Slate/Private/Widgets/Text/SInlineEditableTextBlock.cpp
--- a/Engine/Source/Runtime/Slate/Private/Widgets/Text/SInlineEditableTextBlock.cpp
+++ b/Engine/Source/Runtime/Slate/Private/Widgets/Text/SInlineEditableTextBlock.cpp
@@ -7,17 +7,6 @@
 #include "Widgets/Input/SMultiLineEditableTextBox.h"
 #include "Widgets/Input/SEditableTextBox.h"
 
-SLATE_IMPLEMENT_WIDGET(SInlineEditableTextBlock)
-void SInlineEditableTextBlock::PrivateRegisterAttributes(FSlateAttributeInitializer& AttributeInitializer)
-{
-	SLATE_ADD_MEMBER_ATTRIBUTE_DEFINITION_WITH_NAME(AttributeInitializer, "Text", TextAttribute, EInvalidateWidgetReason::Layout);
-	SLATE_ADD_MEMBER_ATTRIBUTE_DEFINITION_WITH_NAME(AttributeInitializer, "bIsReadOnly", bIsReadOnlyAttribute, EInvalidateWidgetReason::None);
-}
-
-SInlineEditableTextBlock::SInlineEditableTextBlock()
-	: TextAttribute(*this)
-	, bIsReadOnlyAttribute(*this)
-{}
 
 void SInlineEditableTextBlock::Construct( const FArguments& InArgs )
 {
@@ -27,8 +16,8 @@ void SInlineEditableTextBlock::Construct( const FArguments& InArgs )
 	OnTextCommittedDelegate = InArgs._OnTextCommitted;
 	IsSelected = InArgs._IsSelected;
 	OnVerifyTextChanged= InArgs._OnVerifyTextChanged;
-	TextAttribute.Assign(*this, InArgs._Text);
-	bIsReadOnlyAttribute.Assign(*this, InArgs._IsReadOnly);
+	Text = InArgs._Text;
+	bIsReadOnly = InArgs._IsReadOnly;
 	bIsMultiLine = InArgs._MultiLine;
 	DoubleSelectDelay = 0.0f;
 
@@ -43,7 +32,7 @@ void SInlineEditableTextBlock::Construct( const FArguments& InArgs )
 		.VAlign(VAlign_Center)
 		[
 			SAssignNew(TextBlock, STextBlock)
-			.Text(InArgs._Text)
+			.Text(Text)
 			.TextStyle( &InArgs._Style->TextStyle )
 			.Font(InArgs._Font)
 			.ColorAndOpacity( InArgs._ColorAndOpacity )
@@ -103,7 +92,7 @@ void SInlineEditableTextBlock::CancelEditMode()
 	ExitEditingMode();
 
 	// Get the text from source again.
-	SetEditableText(TextAttribute.Get());
+	SetEditableText(Text);
 }
 
 bool SInlineEditableTextBlock::SupportsKeyboardFocus() const
@@ -114,7 +103,7 @@ bool SInlineEditableTextBlock::SupportsKeyboardFocus() const
 
 void SInlineEditableTextBlock::EnterEditingMode()
 {
-	if(!bIsReadOnlyAttribute.Get() && FSlateApplication::Get().HasAnyMouseCaptor() == false)
+	if(!bIsReadOnly.Get() && FSlateApplication::Get().HasAnyMouseCaptor() == false)
 	{
 		if(TextBlock->GetVisibility() == EVisibility::Visible)
 		{
@@ -167,26 +156,26 @@ bool SInlineEditableTextBlock::IsInEditMode() const
 
 void SInlineEditableTextBlock::SetReadOnly(bool bInIsReadOnly)
 {
-	bIsReadOnlyAttribute.Set(*this, bInIsReadOnly);
+	bIsReadOnly = bInIsReadOnly;
 }
 
-void SInlineEditableTextBlock::SetText( TAttribute<FText> InText )
+void SInlineEditableTextBlock::SetText( const TAttribute< FText >& InText )
 {
-	TextBlock->SetText( InText );
-	SetEditableText( InText );
-	TextAttribute.Assign(*this, MoveTemp(InText));
+	Text = InText;
+	TextBlock->SetText( Text );
+	SetEditableText( Text );
 }
 
 void SInlineEditableTextBlock::SetText( const FString& InText )
 {
-	TextAttribute.Set(*this, FText::FromString( InText ));
-	TextBlock->SetText( TextAttribute.Get() );
-	SetEditableText( TextAttribute.Get() );
+	Text = FText::FromString( InText );
+	TextBlock->SetText( Text );
+	SetEditableText( Text );
 }
 
-void SInlineEditableTextBlock::SetWrapTextAt( TAttribute<float> InWrapTextAt )
+void SInlineEditableTextBlock::SetWrapTextAt( const TAttribute<float>& InWrapTextAt )
 {
-	TextBlock->SetWrapTextAt( MoveTemp(InWrapTextAt) );
+	TextBlock->SetWrapTextAt( InWrapTextAt );
 }
 
 FReply SInlineEditableTextBlock::OnMouseButtonDown( const FGeometry& MyGeometry, const FPointerEvent& MouseEvent )
@@ -198,7 +187,7 @@ FReply SInlineEditableTextBlock::OnMouseButtonDown( const FGeometry& MyGeometry,
 
 	if (IsSelected.IsBound())
 	{
-		if (IsSelected.Execute() && !bIsReadOnlyAttribute.Get() && !ActiveTimerHandle.IsValid())
+		if (IsSelected.Execute() && !bIsReadOnly.Get() && !ActiveTimerHandle.IsValid())
 		{
 			ActiveTimerHandle = RegisterActiveTimer(0.5f, FWidgetActiveTimerDelegate::CreateSP(this, &SInlineEditableTextBlock::TriggerEditMode));
 		}
@@ -206,7 +195,7 @@ FReply SInlineEditableTextBlock::OnMouseButtonDown( const FGeometry& MyGeometry,
 	else
 	{
 		// The widget is not managed by another widget, so handle the mouse input and enter edit mode if ready.
-		if (HasKeyboardFocus() && !bIsReadOnlyAttribute.Get())
+		if (HasKeyboardFocus() && !bIsReadOnly.Get())
 		{
 			EnterEditingMode();
 			return FReply::Handled();
@@ -285,7 +274,7 @@ void SInlineEditableTextBlock::OnTextBoxCommitted(const FText& InText, ETextComm
 {
 	if(InCommitType == ETextCommit::OnCleared)
 	{
-		FText SourceTextInEditMode = TextAttribute.Get();
+		FText SourceTextInEditMode = Text.Get();
 		CancelEditMode();
 		// Commit the name, certain actions might need to be taken by the bound function
 		OnTextCommittedDelegate.ExecuteIfBound(SourceTextInEditMode, InCommitType);
@@ -312,7 +301,7 @@ void SInlineEditableTextBlock::OnTextBoxCommitted(const FText& InText, ETextComm
 					CancelEditMode();
 
 					// Commit the name, certain actions might need to be taken by the bound function
-					OnTextCommittedDelegate.ExecuteIfBound(TextAttribute.Get(), InCommitType);
+					OnTextCommittedDelegate.ExecuteIfBound(Text.Get(), InCommitType);
 				
 					return;
 				}
@@ -322,7 +311,7 @@ void SInlineEditableTextBlock::OnTextBoxCommitted(const FText& InText, ETextComm
 				CancelEditMode();
 
 				// Commit the name, certain actions might need to be taken by the bound function
-				OnTextCommittedDelegate.ExecuteIfBound(TextAttribute.Get(), InCommitType);
+				OnTextCommittedDelegate.ExecuteIfBound(Text.Get(), InCommitType);
 				return;
 			}
 		}
@@ -331,9 +320,9 @@ void SInlineEditableTextBlock::OnTextBoxCommitted(const FText& InText, ETextComm
 
 		OnTextCommittedDelegate.ExecuteIfBound(InText, InCommitType);
 
-		if ( !TextAttribute.IsBound(*this) )
+		if ( !Text.IsBound() )
 		{
-			TextBlock->SetText( TextAttribute.Get() );
+			TextBlock->SetText( Text );
 		}
 	}
 }
@@ -352,12 +341,12 @@ TSharedPtr<SWidget> SInlineEditableTextBlock::GetEditableTextWidget() const
 	}
 }
 
-void SInlineEditableTextBlock::SetEditableText( TAttribute<FText> InNewText )
+void SInlineEditableTextBlock::SetEditableText( const TAttribute< FText >& InNewText )
 {
 #if WITH_FANCY_TEXT
-	bIsMultiLine ? MultiLineTextBox->SetText( MoveTemp(InNewText) ) :
+	bIsMultiLine ? MultiLineTextBox->SetText( Text ) :
 #endif //WITH_FANCY_TEXT
-				 TextBox->SetText( MoveTemp(InNewText) );
+				 TextBox->SetText( Text );
 }
 
 void SInlineEditableTextBlock::SetTextBoxError( const FText& ErrorText )
diff --git a/Engine/Source/Runtime/Slate/Private/Widgets/Text/SRichTextBlock.cpp b/Engine/Source/Runtime/Slate/Private/Widgets/Text/SRichTextBlock.cpp
--- a/Engine/Source/Runtime/Slate/Private/Widgets/Text/SRichTextBlock.cpp
+++ b/Engine/Source/Runtime/Slate/Private/Widgets/Text/SRichTextBlock.cpp
@@ -9,32 +9,7 @@
 #include "Framework/Text/RichTextLayoutMarshaller.h"
 #include "Types/ReflectionMetadata.h"
 
-SLATE_IMPLEMENT_WIDGET(SRichTextBlock)
-void SRichTextBlock::PrivateRegisterAttributes(FSlateAttributeInitializer& AttributeInitializer)
-{
-	SLATE_ADD_MEMBER_ATTRIBUTE_DEFINITION(AttributeInitializer, BoundText, EInvalidateWidgetReason::Layout);
-	SLATE_ADD_MEMBER_ATTRIBUTE_DEFINITION(AttributeInitializer, HighlightText, EInvalidateWidgetReason::Layout);
-	SLATE_ADD_MEMBER_ATTRIBUTE_DEFINITION(AttributeInitializer, WrapTextAt, EInvalidateWidgetReason::Layout);
-	SLATE_ADD_MEMBER_ATTRIBUTE_DEFINITION(AttributeInitializer, WrappingPolicy, EInvalidateWidgetReason::Layout);
-	SLATE_ADD_MEMBER_ATTRIBUTE_DEFINITION(AttributeInitializer, TransformPolicy, EInvalidateWidgetReason::Layout);
-	SLATE_ADD_MEMBER_ATTRIBUTE_DEFINITION(AttributeInitializer, Justification, EInvalidateWidgetReason::Layout);
-	SLATE_ADD_MEMBER_ATTRIBUTE_DEFINITION(AttributeInitializer, AutoWrapText, EInvalidateWidgetReason::Layout);
-	SLATE_ADD_MEMBER_ATTRIBUTE_DEFINITION(AttributeInitializer, Margin, EInvalidateWidgetReason::Layout);
-	SLATE_ADD_MEMBER_ATTRIBUTE_DEFINITION(AttributeInitializer, LineHeightPercentage, EInvalidateWidgetReason::Layout);
-	SLATE_ADD_MEMBER_ATTRIBUTE_DEFINITION(AttributeInitializer, MinDesiredWidth, EInvalidateWidgetReason::Layout);
-}
-
 SRichTextBlock::SRichTextBlock()
-	: BoundText(*this)
-	, HighlightText(*this)
-	, WrapTextAt(*this)
-	, WrappingPolicy(*this)
-	, TransformPolicy(*this)
-	, Justification(*this)
-	, AutoWrapText(*this)
-	, Margin(*this)
-	, LineHeightPercentage(*this)
-	, MinDesiredWidth(*this)
 {
 }
 
@@ -45,18 +20,18 @@ SRichTextBlock::~SRichTextBlock()
 
 void SRichTextBlock::Construct( const FArguments& InArgs )
 {
-	SetText(InArgs._Text);
-	SetHighlightText(InArgs._HighlightText);
+	BoundText = InArgs._Text;
+	HighlightText = InArgs._HighlightText;
 
-	SetTextStyle(*InArgs._TextStyle);
-	SetWrapTextAt(InArgs._WrapTextAt);
-	SetAutoWrapText(InArgs._AutoWrapText);
-	SetWrappingPolicy(InArgs._WrappingPolicy);
-	SetTransformPolicy(InArgs._TransformPolicy);
-	SetMargin(InArgs._Margin);
-	SetLineHeightPercentage(InArgs._LineHeightPercentage);
-	SetJustification(InArgs._Justification);
-	SetMinDesiredWidth(InArgs._MinDesiredWidth);
+	TextStyle = *InArgs._TextStyle;
+	WrapTextAt = InArgs._WrapTextAt;
+	AutoWrapText = InArgs._AutoWrapText;
+	WrappingPolicy = InArgs._WrappingPolicy;
+	TransformPolicy = InArgs._TransformPolicy;
+	Margin = InArgs._Margin;
+	LineHeightPercentage = InArgs._LineHeightPercentage;
+	Justification = InArgs._Justification;
+	MinDesiredWidth = InArgs._MinDesiredWidth;
 
 	{
 		TSharedPtr<IRichTextMarkupParser> Parser = InArgs._Parser;
@@ -97,7 +72,7 @@ int32 SRichTextBlock::OnPaint(const FPaintArgs& Args, const FGeometry& AllottedG
 	// HACK: Due to the nature of wrapping and layout, we may have been arranged in a different box than what we were cached with.  Which
 	// might update wrapping, so make sure we always set the desired size to the current size of the text layout, which may have changed
 	// during paint.
-	bool bCanWrap = WrapTextAt.Get() > 0.f || AutoWrapText.Get();
+	bool bCanWrap = WrapTextAt.Get() > 0 || AutoWrapText.Get();
 
 	if (bCanWrap && !NewDesiredSize.Equals(LastDesiredSize))
 	{
@@ -111,16 +86,7 @@ FVector2D SRichTextBlock::ComputeDesiredSize(float LayoutScaleMultiplier) const
 {
 	// ComputeDesiredSize will also update the text layout cache if required
 	const FVector2D TextSize = TextLayoutCache->ComputeDesiredSize(
-		FSlateTextBlockLayout::FWidgetDesiredSizeArgs(
-			BoundText.Get(),
-			HighlightText.Get(),
-			WrapTextAt.Get(),
-			AutoWrapText.Get(),
-			WrappingPolicy.Get(),
-			TransformPolicy.Get(),
-			Margin.Get(),
-			LineHeightPercentage.Get(),
-			Justification.Get()),
+		FSlateTextBlockLayout::FWidgetArgs(BoundText, HighlightText, WrapTextAt, AutoWrapText, WrappingPolicy, TransformPolicy, Margin, LineHeightPercentage, Justification),
 		LayoutScaleMultiplier * TextBlockScale, TextStyle) * TextBlockScale;
 
 	return FVector2D(FMath::Max(TextSize.X, MinDesiredWidth.Get()), TextSize.Y);
@@ -138,15 +104,17 @@ void SRichTextBlock::OnArrangeChildren(const FGeometry& AllottedGeometry, FArran
 	TextLayoutCache->ArrangeChildren(TextBlockScaledGeometry, ArrangedChildren);
 }
 
-void SRichTextBlock::SetText(TAttribute<FText> InTextAttr)
+void SRichTextBlock::SetText( const TAttribute<FText>& InTextAttr )
 {
-	BoundText.Assign(*this, MoveTemp(InTextAttr), FText::GetEmpty());
+	BoundText = InTextAttr;
+	Invalidate(EInvalidateWidget::LayoutAndVolatility);
 	InvalidatePrepass();
 }
 
-void SRichTextBlock::SetHighlightText(TAttribute<FText> InHighlightText)
+void SRichTextBlock::SetHighlightText( const TAttribute<FText>& InHighlightText )
 {
-	HighlightText.Assign(*this, MoveTemp(InHighlightText), FText::GetEmpty());
+	HighlightText = InHighlightText;
+	Invalidate(EInvalidateWidget::LayoutAndVolatility);
 }
 
 void SRichTextBlock::SetTextShapingMethod(const TOptional<ETextShapingMethod>& InTextShapingMethod)
@@ -161,40 +129,40 @@ void SRichTextBlock::SetTextFlowDirection(const TOptional<ETextFlowDirection>& I
 	Invalidate(EInvalidateWidget::Layout);
 }
 
-void SRichTextBlock::SetWrapTextAt(TAttribute<float> InWrapTextAt)
+void SRichTextBlock::SetWrapTextAt(const TAttribute<float>& InWrapTextAt)
 {
-	WrapTextAt.Assign(*this, MoveTemp(InWrapTextAt), 0.f);
+	SetAttribute(WrapTextAt, InWrapTextAt, EInvalidateWidgetReason::Layout);
 }
 
-void SRichTextBlock::SetAutoWrapText(TAttribute<bool> InAutoWrapText)
+void SRichTextBlock::SetAutoWrapText(const TAttribute<bool>& InAutoWrapText)
 {
-	AutoWrapText.Assign(*this, MoveTemp(InAutoWrapText), false);
+	SetAttribute(AutoWrapText, InAutoWrapText, EInvalidateWidgetReason::Layout);
 	InvalidatePrepass();
 }
 
-void SRichTextBlock::SetWrappingPolicy(TAttribute<ETextWrappingPolicy> InWrappingPolicy)
+void SRichTextBlock::SetWrappingPolicy(const TAttribute<ETextWrappingPolicy>& InWrappingPolicy)
 {
-	WrappingPolicy.Assign(*this, MoveTemp(InWrappingPolicy));
+	SetAttribute(WrappingPolicy, InWrappingPolicy, EInvalidateWidgetReason::Layout);
 }
 
-void SRichTextBlock::SetTransformPolicy(TAttribute<ETextTransformPolicy> InTransformPolicy)
+void SRichTextBlock::SetTransformPolicy(const TAttribute<ETextTransformPolicy>& InTransformPolicy)
 {
-	TransformPolicy.Assign(*this, MoveTemp(InTransformPolicy));
+	SetAttribute(TransformPolicy, InTransformPolicy, EInvalidateWidgetReason::Layout);
 }
 
-void SRichTextBlock::SetLineHeightPercentage(TAttribute<float> InLineHeightPercentage)
+void SRichTextBlock::SetLineHeightPercentage(const TAttribute<float>& InLineHeightPercentage)
 {
-	LineHeightPercentage.Assign(*this, MoveTemp(InLineHeightPercentage));
+	SetAttribute(LineHeightPercentage, InLineHeightPercentage, EInvalidateWidgetReason::Layout);
 }
 
-void SRichTextBlock::SetMargin(TAttribute<FMargin> InMargin)
+void SRichTextBlock::SetMargin(const TAttribute<FMargin>& InMargin)
 {
-	Margin.Assign(*this, MoveTemp(InMargin));
+	SetAttribute(Margin, InMargin, EInvalidateWidgetReason::Layout);
 }
 
-void SRichTextBlock::SetJustification(TAttribute<ETextJustify::Type> InJustification)
+void SRichTextBlock::SetJustification(const TAttribute<ETextJustify::Type>& InJustification)
 {
-	Justification.Assign(*this, MoveTemp(InJustification));
+	SetAttribute(Justification, InJustification, EInvalidateWidgetReason::Layout);
 }
 
 void SRichTextBlock::SetTextStyle(const FTextBlockStyle& InTextStyle)
@@ -203,9 +171,9 @@ void SRichTextBlock::SetTextStyle(const FTextBlockStyle& InTextStyle)
 	Invalidate(EInvalidateWidget::Layout);
 }
 
-void SRichTextBlock::SetMinDesiredWidth(TAttribute<float> InMinDesiredWidth)
+void SRichTextBlock::SetMinDesiredWidth(const TAttribute<float>& InMinDesiredWidth)
 {
-	MinDesiredWidth.Assign(*this, MoveTemp(InMinDesiredWidth));
+	SetAttribute(MinDesiredWidth, InMinDesiredWidth, EInvalidateWidgetReason::Layout);
 }
 
 void SRichTextBlock::SetDecoratorStyleSet(const ISlateStyle* NewDecoratorStyleSet)
@@ -230,4 +198,19 @@ void SRichTextBlock::Refresh()
 	Invalidate(EInvalidateWidget::Layout);
 }
 
+bool SRichTextBlock::ComputeVolatility() const
+{
+	return SWidget::ComputeVolatility() 
+		|| BoundText.IsBound()
+		|| HighlightText.IsBound()
+		|| WrapTextAt.IsBound()
+		|| AutoWrapText.IsBound()
+		|| WrappingPolicy.IsBound()
+		|| TransformPolicy.IsBound()
+		|| Margin.IsBound()
+		|| Justification.IsBound()
+		|| LineHeightPercentage.IsBound()
+		|| MinDesiredWidth.IsBound();
+}
+
 #endif //WITH_FANCY_TEXT
diff --git a/Engine/Source/Runtime/Slate/Private/Widgets/Text/STextBlock.cpp b/Engine/Source/Runtime/Slate/Private/Widgets/Text/STextBlock.cpp
--- a/Engine/Source/Runtime/Slate/Private/Widgets/Text/STextBlock.cpp
+++ b/Engine/Source/Runtime/Slate/Private/Widgets/Text/STextBlock.cpp
@@ -15,66 +15,13 @@
 DECLARE_CYCLE_STAT(TEXT("STextBlock::SetText Time"), Stat_SlateTextBlockSetText, STATGROUP_SlateVerbose)
 DECLARE_CYCLE_STAT(TEXT("STextBlock::OnPaint Time"), Stat_SlateTextBlockOnPaint, STATGROUP_SlateVerbose)
 DECLARE_CYCLE_STAT(TEXT("STextBlock::ComputeDesiredSize"), Stat_SlateTextBlockCDS, STATGROUP_SlateVerbose)
-
-
-SLATE_IMPLEMENT_WIDGET(STextBlock)
-void STextBlock::PrivateRegisterAttributes(FSlateAttributeInitializer& AttributeInitializer)
-{
-	//~ If we are in SimpleTextMode the invalidation reason can be different.
-	struct FInvalidation
-	{
-		static EInvalidateWidgetReason GetInvalidationNoneIfSimpleTextMode(const SWidget& OwningWidget)
-		{
-			return static_cast<const STextBlock&>(OwningWidget).bSimpleTextMode ? EInvalidateWidgetReason::None : EInvalidateWidgetReason::Layout;
-		}
-		static EInvalidateWidgetReason GetInvalidationPaintIfSimpleTextMode(const SWidget& OwningWidget)
-		{
-			return static_cast<const STextBlock&>(OwningWidget).bSimpleTextMode ? EInvalidateWidgetReason::Paint : EInvalidateWidgetReason::Layout;
-		}
-	};
-
-	SLATE_ADD_MEMBER_ATTRIBUTE_DEFINITION(AttributeInitializer, BoundText, EInvalidateWidgetReason::Layout);
-	SLATE_ADD_MEMBER_ATTRIBUTE_DEFINITION(AttributeInitializer, Font, EInvalidateWidgetReason::Layout);
-	SLATE_ADD_MEMBER_ATTRIBUTE_DEFINITION(AttributeInitializer, StrikeBrush, FInvalidation::GetInvalidationNoneIfSimpleTextMode);
-	SLATE_ADD_MEMBER_ATTRIBUTE_DEFINITION(AttributeInitializer, ColorAndOpacity, FInvalidation::GetInvalidationPaintIfSimpleTextMode);
-	SLATE_ADD_MEMBER_ATTRIBUTE_DEFINITION(AttributeInitializer, ShadowOffset, EInvalidateWidgetReason::Layout);
-	SLATE_ADD_MEMBER_ATTRIBUTE_DEFINITION(AttributeInitializer, ShadowColorAndOpacity, FInvalidation::GetInvalidationPaintIfSimpleTextMode);
-	SLATE_ADD_MEMBER_ATTRIBUTE_DEFINITION(AttributeInitializer, HighlightColor, FInvalidation::GetInvalidationNoneIfSimpleTextMode);
-	SLATE_ADD_MEMBER_ATTRIBUTE_DEFINITION(AttributeInitializer, HighlightShape, FInvalidation::GetInvalidationNoneIfSimpleTextMode);
-	SLATE_ADD_MEMBER_ATTRIBUTE_DEFINITION(AttributeInitializer, HighlightText, FInvalidation::GetInvalidationNoneIfSimpleTextMode);
-	SLATE_ADD_MEMBER_ATTRIBUTE_DEFINITION(AttributeInitializer, WrapTextAt, FInvalidation::GetInvalidationNoneIfSimpleTextMode);
-	SLATE_ADD_MEMBER_ATTRIBUTE_DEFINITION(AttributeInitializer, AutoWrapText, FInvalidation::GetInvalidationNoneIfSimpleTextMode);
-	SLATE_ADD_MEMBER_ATTRIBUTE_DEFINITION(AttributeInitializer, WrappingPolicy, FInvalidation::GetInvalidationNoneIfSimpleTextMode);
-	SLATE_ADD_MEMBER_ATTRIBUTE_DEFINITION(AttributeInitializer, TransformPolicy, FInvalidation::GetInvalidationNoneIfSimpleTextMode);
-	SLATE_ADD_MEMBER_ATTRIBUTE_DEFINITION(AttributeInitializer, Margin, FInvalidation::GetInvalidationNoneIfSimpleTextMode);
-	SLATE_ADD_MEMBER_ATTRIBUTE_DEFINITION(AttributeInitializer, Justification, FInvalidation::GetInvalidationNoneIfSimpleTextMode);
-	SLATE_ADD_MEMBER_ATTRIBUTE_DEFINITION(AttributeInitializer, LineHeightPercentage, FInvalidation::GetInvalidationNoneIfSimpleTextMode);
-	SLATE_ADD_MEMBER_ATTRIBUTE_DEFINITION(AttributeInitializer, MinDesiredWidth, EInvalidateWidgetReason::Layout);
-}
+DECLARE_CYCLE_STAT(TEXT("STextBlock::ComputeVolitility"), Stat_SlateTextBlockCV, STATGROUP_SlateVerbose)
 
 STextBlock::STextBlock()
-	: BoundText(*this)
-	, Font(*this)
-	, StrikeBrush(*this)
-	, ColorAndOpacity(*this)
-	, ShadowOffset(*this)
-	, ShadowColorAndOpacity(*this)
-	, HighlightColor(*this)
-	, HighlightShape(*this)
-	, HighlightText(*this)
-	, WrapTextAt(*this)
-	, AutoWrapText(*this)
-	, WrappingPolicy(*this)
-	, TransformPolicy(*this)
-	, Margin(*this)
-	, Justification(*this)
-	, LineHeightPercentage(*this)
-	, MinDesiredWidth(*this)
-	, Union_IsAttributeSet(0)
-	, bSimpleTextMode(false)
 {
 	SetCanTick(false);
 	bCanSupportFocus = false;
+	bSimpleTextMode = false;
 
 #if WITH_ACCESSIBILITY
 	AccessibleBehavior = EAccessibleBehavior::Auto;
@@ -91,31 +38,33 @@ void STextBlock::Construct( const FArguments& InArgs )
 {
 	TextStyle = *InArgs._TextStyle;
 
-	SetHighlightText(InArgs._HighlightText);
-	SetWrapTextAt(InArgs._WrapTextAt);
-	SetAutoWrapText(InArgs._AutoWrapText);
-	SetWrappingPolicy(InArgs._WrappingPolicy);
-
-	SetTransformPolicy(InArgs._TransformPolicy);
+	HighlightText = InArgs._HighlightText;
+	WrapTextAt = InArgs._WrapTextAt;
+	AutoWrapText = InArgs._AutoWrapText;
+	WrappingPolicy = InArgs._WrappingPolicy;
 
-	SetMargin(InArgs._Margin);
-	SetLineHeightPercentage(InArgs._LineHeightPercentage);
-	SetJustification(InArgs._Justification);
-	SetMinDesiredWidth(InArgs._MinDesiredWidth);
-
-	SetFont(InArgs._Font);
-	SetStrikeBrush(InArgs._StrikeBrush);
-	SetColorAndOpacity(InArgs._ColorAndOpacity);
-	SetShadowOffset(InArgs._ShadowOffset);
-	SetShadowColorAndOpacity(InArgs._ShadowColorAndOpacity);
-	SetHighlightColor(InArgs._HighlightColor);
-	SetHighlightShape(InArgs._HighlightShape);
+	if (InArgs._TransformPolicy.IsSet()) 
+	{
+		TransformPolicy = InArgs._TransformPolicy;
+	}
+	Margin = InArgs._Margin;
+	LineHeightPercentage = InArgs._LineHeightPercentage;
+	Justification = InArgs._Justification;
+	MinDesiredWidth = InArgs._MinDesiredWidth;
+
+	Font = InArgs._Font;
+	StrikeBrush = InArgs._StrikeBrush;
+	ColorAndOpacity = InArgs._ColorAndOpacity;
+	ShadowOffset = InArgs._ShadowOffset;
+	ShadowColorAndOpacity = InArgs._ShadowColorAndOpacity;
+	HighlightColor = InArgs._HighlightColor;
+	HighlightShape = InArgs._HighlightShape;
 
 	bSimpleTextMode = InArgs._SimpleTextMode;
 
 	SetOnMouseDoubleClick(InArgs._OnDoubleClicked);
 
-	BoundText.Assign(*this, InArgs._Text);
+	BoundText = InArgs._Text;
 
 	//if(!bSimpleTextMode)
 	{
@@ -127,63 +76,99 @@ void STextBlock::Construct( const FArguments& InArgs )
 
 FSlateFontInfo STextBlock::GetFont() const
 {
-	return bIsAttributeFontSet ? Font.Get() : TextStyle.Font;
+	return Font.IsSet() ? Font.Get() : TextStyle.Font;
 }
 
 const FSlateBrush* STextBlock::GetStrikeBrush() const
 {
-	return bIsAttributeStrikeBrushSet ? StrikeBrush.Get() : &TextStyle.StrikeBrush;
+	return StrikeBrush.IsSet() ? StrikeBrush.Get() : &TextStyle.StrikeBrush;
 }
 
 FSlateColor STextBlock::GetColorAndOpacity() const
 {
-	return bIsAttributeColorAndOpacitySet ? ColorAndOpacity.Get() : TextStyle.ColorAndOpacity;
+	return ColorAndOpacity.IsSet() ? ColorAndOpacity.Get() : TextStyle.ColorAndOpacity;
 }
 
 FVector2D STextBlock::GetShadowOffset() const
 {
-	return bIsAttributeShadowOffsetSet ? ShadowOffset.Get() : TextStyle.ShadowOffset;
+	return ShadowOffset.IsSet() ? ShadowOffset.Get() : TextStyle.ShadowOffset;
 }
 
 FLinearColor STextBlock::GetShadowColorAndOpacity() const
 {
-	return bIsAttributeShadowColorAndOpacitySet ? ShadowColorAndOpacity.Get() : TextStyle.ShadowColorAndOpacity;
+	return ShadowColorAndOpacity.IsSet() ? ShadowColorAndOpacity.Get() : TextStyle.ShadowColorAndOpacity;
 }
 
 FSlateColor STextBlock::GetHighlightColor() const
 {
-	return bIsAttributeHighlightColorSet ? HighlightColor.Get() : TextStyle.HighlightColor;
+	return HighlightColor.IsSet() ? HighlightColor.Get() : TextStyle.HighlightColor;
 }
 
 const FSlateBrush* STextBlock::GetHighlightShape() const
 {
-	return bIsAttributeHighlightShapeSet ? HighlightShape.Get() : &TextStyle.HighlightShape;
+	return HighlightShape.IsSet() ? HighlightShape.Get() : &TextStyle.HighlightShape;
 }
 
-ETextTransformPolicy STextBlock::GetTransformPolicyImpl() const
+void STextBlock::InvalidateText(EInvalidateWidget InvalidateReason)
 {
-	return bIsAttributeTransformPolicySet ? TransformPolicy.Get() : TextStyle.TransformPolicy;
+	if (bSimpleTextMode && EnumHasAnyFlags(InvalidateReason, EInvalidateWidget::Layout))
+	{
+		CachedSimpleDesiredSize.Reset();
+	}
+
+	Invalidate(InvalidateReason);
 }
 
-void STextBlock::InvalidateText(EInvalidateWidgetReason InvalidateReason)
+void STextBlock::SetText( const TAttribute< FString >& InText )
 {
-	if (bSimpleTextMode && EnumHasAnyFlags(InvalidateReason, EInvalidateWidgetReason::Layout))
+	if (InText.IsSet() && !InText.IsBound())
 	{
-		CachedSimpleDesiredSize.Reset();
+		SetText(FText::AsCultureInvariant(InText.Get()));
+		return;
 	}
 
-	Invalidate(InvalidateReason);
+	SCOPE_CYCLE_COUNTER(Stat_SlateTextBlockSetText);
+	BoundText = MakeAttributeLambda([InText]()
+	{
+		return FText::AsCultureInvariant(InText.Get(FString()));
+	});
+	InvalidateText(EInvalidateWidget::LayoutAndVolatility);
+}
+
+void STextBlock::SetText( const FString& InText )
+{
+	SetText(FText::AsCultureInvariant(InText));
+}
+
+void STextBlock::SetText( const TAttribute< FText >& InText )
+{
+	if (InText.IsSet() && !InText.IsBound())
+	{
+		SetText(InText.Get());
+		return;
+	}
+
+	SCOPE_CYCLE_COUNTER(Stat_SlateTextBlockSetText);
+	BoundText = InText;
+	InvalidateText(EInvalidateWidget::LayoutAndVolatility);
 }
 
-void STextBlock::SetText(TAttribute<FText> InText)
+void STextBlock::SetText( const FText& InText )
 {
 	SCOPE_CYCLE_COUNTER(Stat_SlateTextBlockSetText);
-	BoundText.Assign(*this, MoveTemp(InText));
+
+	if (!BoundText.IsBound() && InText.IdenticalTo(BoundText.Get(), ETextIdenticalModeFlags::DeepCompare | ETextIdenticalModeFlags::LexicalCompareInvariants))
+	{
+		return;
+	}
+
+	BoundText = InText;
+	InvalidateText(EInvalidateWidget::LayoutAndVolatility);
 }
 
 void STextBlock::SetHighlightText(TAttribute<FText> InText)
 {
-	HighlightText.Assign(*this, MoveTemp(InText));
+	HighlightText = InText;
 }
 
 int32 STextBlock::OnPaint( const FPaintArgs& Args, const FGeometry& AllottedGeometry, const FSlateRect& MyCullingRect, FSlateWindowElementList& OutDrawElements, int32 LayerId, const FWidgetStyle& InWidgetStyle, bool bParentEnabled ) const
@@ -250,11 +235,11 @@ int32 STextBlock::OnPaint( const FPaintArgs& Args, const FGeometry& AllottedGeom
 		// HACK: Due to the nature of wrapping and layout, we may have been arranged in a different box than what we were cached with.  Which
 		// might update wrapping, so make sure we always set the desired size to the current size of the text layout, which may have changed
 		// during paint.
-		const bool bCanWrap = WrapTextAt.Get() > 0 || AutoWrapText.Get();
+		bool bCanWrap = WrapTextAt.Get() > 0 || AutoWrapText.Get();
 
 		if (bCanWrap && !NewDesiredSize.Equals(LastDesiredSize))
 		{
-			const_cast<STextBlock*>(this)->Invalidate(EInvalidateWidgetReason::Layout);
+			const_cast<STextBlock*>(this)->Invalidate(EInvalidateWidget::Layout);
 		}
 	}
 
@@ -269,52 +254,76 @@ FVector2D STextBlock::ComputeDesiredSize(float LayoutScaleMultiplier) const
 	{
 		const FVector2D LocalShadowOffset = GetShadowOffset();
 
-		const float LocalOutlineSize = (float)(GetFont().OutlineSettings.OutlineSize);
+		const float LocalOutlineSize = GetFont().OutlineSettings.OutlineSize;
 
 		// Account for the outline width impacting both size of the text by multiplying by 2
 		// Outline size in Y is accounted for in MaxHeight calculation in Measure()
-		const FVector2D ComputedOutlineSize(LocalOutlineSize * 2.f, LocalOutlineSize);
+		const FVector2D ComputedOutlineSize(LocalOutlineSize * 2, LocalOutlineSize);
 		const FVector2D TextSize = FSlateApplication::Get().GetRenderer()->GetFontMeasureService()->Measure(GetText(), GetFont()) + ComputedOutlineSize + LocalShadowOffset;
 
-		CachedSimpleDesiredSize = FVector2D(FMath::Max(MinDesiredWidth.Get(), TextSize.X), TextSize.Y);
+		CachedSimpleDesiredSize = FVector2D(FMath::Max(MinDesiredWidth.Get(0.0f), TextSize.X), TextSize.Y);
 		return CachedSimpleDesiredSize.GetValue();
 	}
 	else
 	{
 		// ComputeDesiredSize will also update the text layout cache if required
 		const FVector2D TextSize = TextLayoutCache->ComputeDesiredSize(
-			FSlateTextBlockLayout::FWidgetDesiredSizeArgs(
-				BoundText.Get(),
-				HighlightText.Get(),
-				WrapTextAt.Get(),
-				AutoWrapText.Get(),
-				WrappingPolicy.Get(),
-				GetTransformPolicyImpl(),
-				Margin.Get(),
-				LineHeightPercentage.Get(),
-				Justification.Get()),
-			LayoutScaleMultiplier, GetComputedTextStyle());
-
-		return FVector2D(FMath::Max(MinDesiredWidth.Get(), TextSize.X), TextSize.Y);
+			FSlateTextBlockLayout::FWidgetArgs(BoundText, HighlightText, WrapTextAt, AutoWrapText, WrappingPolicy, GetTransformPolicy(), Margin, LineHeightPercentage, Justification),
+			LayoutScaleMultiplier, GetComputedTextStyle()
+		);
+
+		return FVector2D(FMath::Max(MinDesiredWidth.Get(0.0f), TextSize.X), TextSize.Y);
 	}
 }
 
-void STextBlock::SetFont(TAttribute<FSlateFontInfo> InFont)
+bool STextBlock::ComputeVolatility() const
 {
-	bIsAttributeFontSet = InFont.IsSet();
-	Font.Assign(*this, MoveTemp(InFont));
+	SCOPE_CYCLE_COUNTER(Stat_SlateTextBlockCV);
+	return SLeafWidget::ComputeVolatility() 
+		|| BoundText.IsBound()
+		|| Font.IsBound()
+		|| ColorAndOpacity.IsBound()
+		|| ShadowOffset.IsBound()
+		|| ShadowColorAndOpacity.IsBound()
+		|| HighlightColor.IsBound()
+		|| HighlightShape.IsBound()
+		|| HighlightText.IsBound()
+		|| WrapTextAt.IsBound()
+		|| AutoWrapText.IsBound()
+		|| WrappingPolicy.IsBound()
+		|| TransformPolicy.IsBound()
+		|| Margin.IsBound()
+		|| Justification.IsBound()
+		|| LineHeightPercentage.IsBound()
+		|| MinDesiredWidth.IsBound();
 }
 
-void STextBlock::SetStrikeBrush(TAttribute<const FSlateBrush*> InStrikeBrush)
+void STextBlock::SetFont(const TAttribute< FSlateFontInfo >& InFont)
 {
-	bIsAttributeStrikeBrushSet = InStrikeBrush.IsSet();
-	StrikeBrush.Assign(*this, MoveTemp(InStrikeBrush));
+	if(!Font.IsSet() || !Font.IdenticalTo(InFont))
+	{
+		Font = InFont;
+		InvalidateText(EInvalidateWidget::LayoutAndVolatility);
+	}
 }
 
-void STextBlock::SetColorAndOpacity(TAttribute<FSlateColor> InColorAndOpacity)
+void STextBlock::SetStrikeBrush(const TAttribute<const FSlateBrush*>& InStrikeBrush)
 {
-	bIsAttributeColorAndOpacitySet = InColorAndOpacity.IsSet();
-	ColorAndOpacity.Assign(*this, MoveTemp(InColorAndOpacity));
+	if (!StrikeBrush.IsSet() || !StrikeBrush.IdenticalTo(InStrikeBrush))
+	{
+		StrikeBrush = InStrikeBrush;
+		InvalidateText(EInvalidateWidget::LayoutAndVolatility);
+	}
+}
+
+void STextBlock::SetColorAndOpacity(const TAttribute<FSlateColor>& InColorAndOpacity)
+{
+	if ( !ColorAndOpacity.IsSet() || !ColorAndOpacity.IdenticalTo(InColorAndOpacity) )
+	{
+		ColorAndOpacity = InColorAndOpacity;
+		// HACK: Normally this would be Paint only, but textblocks need to recache layout.
+		InvalidateText(EInvalidateWidget::LayoutAndVolatility);
+	}
 }
 
 void STextBlock::SetTextStyle(const FTextBlockStyle* InTextStyle)
@@ -329,7 +338,7 @@ void STextBlock::SetTextStyle(const FTextBlockStyle* InTextStyle)
 		TextStyle = *Defaults._TextStyle;
 	}
 
-	InvalidateText(EInvalidateWidgetReason::Layout);
+	InvalidateText(EInvalidateWidget::Layout);
 }
 
 void STextBlock::SetTextShapingMethod(const TOptional<ETextShapingMethod>& InTextShapingMethod)
@@ -337,7 +346,7 @@ void STextBlock::SetTextShapingMethod(const TOptional<ETextShapingMethod>& InTex
 	if (!bSimpleTextMode)
 	{
 		TextLayoutCache->SetTextShapingMethod(InTextShapingMethod);
-		InvalidateText(EInvalidateWidgetReason::Layout);
+		InvalidateText(EInvalidateWidget::Layout);
 	}
 }
 
@@ -346,82 +355,121 @@ void STextBlock::SetTextFlowDirection(const TOptional<ETextFlowDirection>& InTex
 	if(!bSimpleTextMode)
 	{
 		TextLayoutCache->SetTextFlowDirection(InTextFlowDirection);
-		InvalidateText(EInvalidateWidgetReason::Layout);
+		InvalidateText(EInvalidateWidget::Layout);
 	}
 }
 
-void STextBlock::SetWrapTextAt(TAttribute<float> InWrapTextAt)
+void STextBlock::SetWrapTextAt(const TAttribute<float>& InWrapTextAt)
 {
-	WrapTextAt.Assign(*this, MoveTemp(InWrapTextAt), 0.f);
-}
-
-void STextBlock::SetAutoWrapText(TAttribute<bool> InAutoWrapText)
-{
-	AutoWrapText.Assign(*this, MoveTemp(InAutoWrapText), 0.f);
+	if(!WrapTextAt.IdenticalTo(InWrapTextAt))
+	{
+		WrapTextAt = InWrapTextAt;
+		InvalidateText(EInvalidateWidget::LayoutAndVolatility);
+	}
 }
 
-void STextBlock::SetWrappingPolicy(TAttribute<ETextWrappingPolicy> InWrappingPolicy)
+void STextBlock::SetAutoWrapText(const TAttribute<bool>& InAutoWrapText)
 {
-	WrappingPolicy.Assign(*this, MoveTemp(InWrappingPolicy));
+	if(!AutoWrapText.IdenticalTo(InAutoWrapText))
+	{
+		AutoWrapText = InAutoWrapText;
+		InvalidateText(EInvalidateWidget::LayoutAndVolatility);
+	}
 }
 
-void STextBlock::SetTransformPolicy(TAttribute<ETextTransformPolicy> InTransformPolicy)
+void STextBlock::SetWrappingPolicy(const TAttribute<ETextWrappingPolicy>& InWrappingPolicy)
 {
-	bIsAttributeTransformPolicySet = InTransformPolicy.IsSet();
-	TransformPolicy.Assign(*this, MoveTemp(InTransformPolicy));
+	if(!WrappingPolicy.IdenticalTo(InWrappingPolicy))
+	{
+		WrappingPolicy = InWrappingPolicy;
+		InvalidateText(EInvalidateWidget::LayoutAndVolatility);
+	}
 }
 
-void STextBlock::SetShadowOffset(TAttribute<FVector2D> InShadowOffset)
+void STextBlock::SetTransformPolicy(const TAttribute<ETextTransformPolicy>& InTransformPolicy)
 {
-	bIsAttributeShadowOffsetSet = InShadowOffset.IsSet();
-	ShadowOffset.Assign(*this, MoveTemp(InShadowOffset));
+	if(!TransformPolicy.IdenticalTo(InTransformPolicy))
+	{
+		TransformPolicy = InTransformPolicy;
+		InvalidateText(EInvalidateWidget::LayoutAndVolatility);
+	}
 }
 
-void STextBlock::SetShadowColorAndOpacity(TAttribute<FLinearColor> InShadowColorAndOpacity)
+ETextTransformPolicy STextBlock::GetTransformPolicy() const
 {
-	bIsAttributeShadowColorAndOpacitySet = InShadowColorAndOpacity.IsSet();
-	ShadowColorAndOpacity.Assign(*this, MoveTemp(InShadowColorAndOpacity));
+	if (TransformPolicy.IsSet())
+	{
+		return TransformPolicy.Get();
+	}
+	return TextStyle.TransformPolicy;
 }
 
-void STextBlock::SetHighlightColor(TAttribute<FLinearColor> InHighlightColor)
+void STextBlock::SetShadowOffset(const TAttribute<FVector2D>& InShadowOffset)
 {
-	bIsAttributeHighlightColorSet = InHighlightColor.IsSet();
-	HighlightColor.Assign(*this, MoveTemp(InHighlightColor));
+	if(!ShadowOffset.IdenticalTo(InShadowOffset))
+	{
+		ShadowOffset = InShadowOffset;
+		InvalidateText(EInvalidateWidget::LayoutAndVolatility);
+	}
 }
 
-void STextBlock::SetHighlightShape(TAttribute<const FSlateBrush*> InHighlightShape)
+void STextBlock::SetShadowColorAndOpacity(const TAttribute<FLinearColor>& InShadowColorAndOpacity)
 {
-	bIsAttributeHighlightShapeSet = InHighlightShape.IsSet();
-	HighlightShape.Assign(*this, MoveTemp(InHighlightShape));
+	if(!ShadowColorAndOpacity.IdenticalTo(InShadowColorAndOpacity))
+	{
+		ShadowColorAndOpacity = InShadowColorAndOpacity;
+		// HACK: Normally this would be Paint only, but textblocks need to recache layout.
+		InvalidateText(EInvalidateWidget::LayoutAndVolatility);
+	}
 }
 
-void STextBlock::SetMinDesiredWidth(TAttribute<float> InMinDesiredWidth)
+void STextBlock::SetMinDesiredWidth(const TAttribute<float>& InMinDesiredWidth)
 {
-	MinDesiredWidth.Assign(*this, MoveTemp(InMinDesiredWidth), 0.f);
+	if(!MinDesiredWidth.IdenticalTo(InMinDesiredWidth))
+	{
+		MinDesiredWidth = InMinDesiredWidth;
+		InvalidateText(EInvalidateWidget::LayoutAndVolatility);
+	}
 }
 
-void STextBlock::SetLineHeightPercentage(TAttribute<float> InLineHeightPercentage)
+void STextBlock::SetLineHeightPercentage(const TAttribute<float>& InLineHeightPercentage)
 {
-	LineHeightPercentage.Assign(*this, MoveTemp(InLineHeightPercentage));
+	if(!LineHeightPercentage.IdenticalTo(InLineHeightPercentage))
+	{
+		LineHeightPercentage = InLineHeightPercentage;
+		InvalidateText(EInvalidateWidget::LayoutAndVolatility);
+	}
 }
 
-void STextBlock::SetMargin(TAttribute<FMargin> InMargin)
+void STextBlock::SetMargin(const TAttribute<FMargin>& InMargin)
 {
-	Margin.Assign(*this, MoveTemp(InMargin));
+	if(!Margin.IdenticalTo(InMargin))
+	{
+		Margin = InMargin;
+		InvalidateText(EInvalidateWidget::LayoutAndVolatility);
+	}
 }
 
-void STextBlock::SetJustification(TAttribute<ETextJustify::Type> InJustification)
+void STextBlock::SetJustification(const TAttribute<ETextJustify::Type>& InJustification)
 {
-	Justification.Assign(*this, MoveTemp(InJustification));
+	if(!Justification.IdenticalTo(InJustification))
+	{
+		Justification = InJustification;
+		InvalidateText(EInvalidateWidget::LayoutAndVolatility);
+	}
 }
 
 FTextBlockStyle STextBlock::GetComputedTextStyle() const
 {
 	FTextBlockStyle ComputedStyle = TextStyle;
 	ComputedStyle.SetFont( GetFont() );
-	if (const FSlateBrush* const ComputedStrikeBrush = GetStrikeBrush())
+	if (StrikeBrush.IsSet())
 	{
-		ComputedStyle.SetStrikeBrush(*ComputedStrikeBrush);
+		const FSlateBrush* const ComputedStrikeBrush = StrikeBrush.Get();
+		if (ComputedStrikeBrush)
+		{
+			ComputedStyle.SetStrikeBrush(*ComputedStrikeBrush);
+		}
 	}
 	ComputedStyle.SetColorAndOpacity( GetColorAndOpacity() );
 	ComputedStyle.SetShadowOffset( GetShadowOffset() );
diff --git a/Engine/Source/Runtime/Slate/Private/Widgets/Text/SlateEditableTextLayout.cpp b/Engine/Source/Runtime/Slate/Private/Widgets/Text/SlateEditableTextLayout.cpp
--- a/Engine/Source/Runtime/Slate/Private/Widgets/Text/SlateEditableTextLayout.cpp
+++ b/Engine/Source/Runtime/Slate/Private/Widgets/Text/SlateEditableTextLayout.cpp
@@ -3428,17 +3428,9 @@ FVector2D FSlateEditableTextLayout::ComputeDesiredSize(float LayoutScaleMultipli
 		MarginValue.Right += CaretWidth;
 
 		const FVector2D HintTextSize = HintTextLayout->ComputeDesiredSize(
-			FSlateTextBlockLayout::FWidgetDesiredSizeArgs(
-				HintText.Get(FText::GetEmpty()),
-				FText::GetEmpty(),
-				WrapTextAt.Get(0.f),
-				AutoWrapText.Get(false),
-				WrappingPolicy.Get(),
-				ETextTransformPolicy::None,
-				MarginValue,
-				LineHeightPercentage.Get(),
-				Justification.Get()),
-			LayoutScaleMultiplier, HintTextStyle);
+			FSlateTextBlockLayout::FWidgetArgs(HintText, FText::GetEmpty(), WrapTextAt, AutoWrapText, WrappingPolicy, ETextTransformPolicy::None, MarginValue, LineHeightPercentage, Justification),
+			LayoutScaleMultiplier, HintTextStyle
+			);
 
 		// If a wrapping width has been provided, then we need to report that as the desired width
 		DesiredWidth = WrappingWidth > 0 ? WrappingWidth : HintTextSize.X;
diff --git a/Engine/Source/Runtime/Slate/Private/Widgets/Text/SlateTextBlockLayout.cpp b/Engine/Source/Runtime/Slate/Private/Widgets/Text/SlateTextBlockLayout.cpp
--- a/Engine/Source/Runtime/Slate/Private/Widgets/Text/SlateTextBlockLayout.cpp
+++ b/Engine/Source/Runtime/Slate/Private/Widgets/Text/SlateTextBlockLayout.cpp
@@ -25,80 +25,6 @@ FSlateTextBlockLayout::FSlateTextBlockLayout(SWidget* InOwner, FTextBlockStyle I
 	TextLayout->SetLineBreakIterator(MoveTemp(InLineBreakPolicy));
 }
 
-FVector2D FSlateTextBlockLayout::ComputeDesiredSize(const FWidgetDesiredSizeArgs& InWidgetArgs, const float InScale, const FTextBlockStyle& InTextStyle)
-{
-	// Cache the wrapping rules so that we can recompute the wrap at width in paint.
-	CachedWrapTextAt = InWidgetArgs.WrapTextAt;
-	bCachedAutoWrapText = InWidgetArgs.AutoWrapText;
-
-	const ETextTransformPolicy PreviousTransformPolicy = TextLayout->GetTransformPolicy();
-
-	// Set the text layout information
-	TextLayout->SetScale(InScale);
-	TextLayout->SetWrappingWidth(CalculateWrappingWidth());
-	TextLayout->SetWrappingPolicy(InWidgetArgs.WrappingPolicy);
-	TextLayout->SetTransformPolicy(InWidgetArgs.TransformPolicy);
-	TextLayout->SetMargin(InWidgetArgs.Margin);
-	TextLayout->SetJustification(InWidgetArgs.Justification);
-	TextLayout->SetLineHeightPercentage(InWidgetArgs.LineHeightPercentage);
-
-	// Has the transform policy changed? If so we need a full refresh as that is destructive to the model text
-	if (PreviousTransformPolicy != TextLayout->GetTransformPolicy())
-	{
-		Marshaller->MakeDirty();
-	}
-
-	// Has the style used for this text block changed?
-	if (!IsStyleUpToDate(InTextStyle))
-	{
-		TextLayout->SetDefaultTextStyle(InTextStyle);
-		Marshaller->MakeDirty(); // will regenerate the text using the new default style
-	}
-
-	{
-		bool bRequiresTextUpdate = false;
-		const FText& TextToSet = InWidgetArgs.Text;
-		if (!TextLastUpdate.IdenticalTo(TextToSet))
-		{
-			// The pointer used by the bound text has changed, however the text may still be the same - check that now
-			if (!TextLastUpdate.IsDisplayStringEqualTo(TextToSet))
-			{
-				// The source text has changed, so update the internal text
-				bRequiresTextUpdate = true;
-			}
-
-			// Update this even if the text is lexically identical, as it will update the pointer compared by IdenticalTo for the next Tick
-			TextLastUpdate = FTextSnapshot(TextToSet);
-		}
-
-		if (bRequiresTextUpdate || Marshaller->IsDirty())
-		{
-			UpdateTextLayout(TextToSet);
-		}
-	}
-
-	{
-		const FText& HighlightTextToSet = InWidgetArgs.HighlightText;
-		if (!HighlightTextLastUpdate.IdenticalTo(HighlightTextToSet))
-		{
-			// The pointer used by the bound text has changed, however the text may still be the same - check that now
-			if (!HighlightTextLastUpdate.IsDisplayStringEqualTo(HighlightTextToSet))
-			{
-				UpdateTextHighlights(HighlightTextToSet);
-			}
-
-			// Update this even if the text is lexically identical, as it will update the pointer compared by IdenticalTo for the next Tick
-			HighlightTextLastUpdate = FTextSnapshot(HighlightTextToSet);
-		}
-	}
-
-	// We need to update our size if the text layout has become dirty
-	TextLayout->UpdateIfNeeded();
-
-	return TextLayout->GetSize();
-}
-
-PRAGMA_DISABLE_DEPRECATION_WARNINGS
 FVector2D FSlateTextBlockLayout::ComputeDesiredSize(const FWidgetArgs& InWidgetArgs, const float InScale, const FTextBlockStyle& InTextStyle)
 {
 	// Cache the wrapping rules so that we can recompute the wrap at width in paint.
@@ -171,7 +97,6 @@ FVector2D FSlateTextBlockLayout::ComputeDesiredSize(const FWidgetArgs& InWidgetA
 
 	return TextLayout->GetSize();
 }
-PRAGMA_DISABLE_DEPRECATION_WARNINGS
 
 FVector2D FSlateTextBlockLayout::GetDesiredSize() const
 {
diff --git a/Engine/Source/Runtime/Slate/Public/Widgets/Text/SInlineEditableTextBlock.h b/Engine/Source/Runtime/Slate/Public/Widgets/Text/SInlineEditableTextBlock.h
--- a/Engine/Source/Runtime/Slate/Public/Widgets/Text/SInlineEditableTextBlock.h
+++ b/Engine/Source/Runtime/Slate/Public/Widgets/Text/SInlineEditableTextBlock.h
@@ -30,9 +30,6 @@ DECLARE_DELEGATE_OneParam(FOnBeginTextEdit, const FText&)
  */
 class SLATE_API SInlineEditableTextBlock: public SCompoundWidget
 {
-	SLATE_DECLARE_WIDGET(SInlineEditableTextBlock, SCompoundWidget)
-
-public:
 	SLATE_BEGIN_ARGS( SInlineEditableTextBlock )
 		: _Text()
 		, _Style( &FCoreStyle::Get().GetWidgetStyle<FInlineEditableTextBlockStyle>("InlineEditableTextBlockStyle") )
@@ -108,7 +105,6 @@ public:
 		SLATE_EVENT( FOnVerifyTextChanged, OnVerifyTextChanged )
 	SLATE_END_ARGS()
 
-	SInlineEditableTextBlock();
 	~SInlineEditableTextBlock();
 
 	/**
@@ -147,13 +143,11 @@ public:
 
 	void SetReadOnly(bool bInIsReadOnly);
 
-	void SetText(TAttribute<FText> InText);
-
-	UE_DEPRECATED(5.0, "SetText taking FString is deprecated. Use the FText version instead")
+	void SetText( const TAttribute< FText >& InText );
 	void SetText( const FString& InText );
 
 	/** Sets the wrap text at attribute.  See WrapTextAt attribute */
-	void SetWrapTextAt(TAttribute<float> InWrapTextAt);
+	void SetWrapTextAt(const TAttribute<float>& InWrapTextAt);
 
 protected:
 	/** Callback for the text box's OnTextChanged event */
@@ -165,12 +159,6 @@ protected:
 	/** Cancels the edit mode and switches back to label mode */
 	void CancelEditMode();
 
-	/** @return an attribute reference of Text */
-	TSlateAttributeRef<FText> GetTextAttribute() const { return TSlateAttributeRef<FText>(*this, TextAttribute); }
-
-	/** @return an attribute reference of ColorAndOpacity */
-	TSlateAttributeRef<bool> GetIsReadOnlyAttribute() const { return TSlateAttributeRef<bool>(*this, bIsReadOnlyAttribute); }
-
 protected:
 	/** The widget used when in label mode */
 	TSharedPtr< STextBlock > TextBlock;
@@ -203,24 +191,14 @@ protected:
 	/** Callback to verify text when changed. Will return an error message to denote problems. */
 	FOnVerifyTextChanged OnVerifyTextChanged;
 
-	/** Widget to focus when we finish editing */
-	TWeakPtr<SWidget> WidgetToFocus;
-
-	/**
-	 * When selection of widget is managed by another widget, this delays the "double select" from
-	 * occurring immediately, offering a chance for double clicking to take action.
-	 */
-	float DoubleSelectDelay;
+	/** Attribute for the text to use for the widget */
+	TAttribute<FText> Text;
 
-	/** Attribute to look up if the widget is multiline */
-	bool bIsMultiLine;
+	/** Attribute to look up if the widget is read-only */
+	TAttribute< bool > bIsReadOnly;
 
-//#if WITH_EDITORONLY_DATA
-//	UE_DEPRECATED(5.0, "Direct access to Text is now deprecated. Use the setter or getter.")
-//	FSlateDeprecatedTAttribute<FText> Text;
-//	UE_DEPRECATED(5.0, "Direct access to bIsReadOnly is now deprecated. Use the setter or getter.")
-//	FSlateDeprecatedTAttribute<bool> bIsReadOnly;
-//#endif
+	/** Widget to focus when we finish editing */
+	TWeakPtr<SWidget> WidgetToFocus;
 
 private:
 
@@ -228,20 +206,23 @@ private:
 	TSharedPtr<SWidget> GetEditableTextWidget() const;
 
 	/** Set editable text */
-	void SetEditableText(TAttribute<FText> InNewText);
+	void SetEditableText( const TAttribute< FText >& InNewText );
 
 	/** Set textbox error */
-	void SetTextBoxError(const FText& ErrorText);
+	void SetTextBoxError( const FText& ErrorText );
 
 	/** Active timer to trigger entry into edit mode after a delay */
 	EActiveTimerReturnType TriggerEditMode(double InCurrentTime, float InDeltaTime);
 
-	/** Attribute for the text to use for the widget */
-	TSlateAttribute<FText> TextAttribute;
-
 	/** The handle to the active timer */
 	TWeakPtr<FActiveTimerHandle> ActiveTimerHandle;
 
-	/** Attribute to look up if the widget is read-only */
-	TSlateAttribute<bool> bIsReadOnlyAttribute;
+protected:
+	/** When selection of widget is managed by another widget, this delays the "double select" from 
+	occurring immediately, offering a chance for double clicking to take action. */
+	float DoubleSelectDelay;
+
+	/** Attribute to look up if the widget is multiline */
+	uint8 bIsMultiLine:1;
+
 };
diff --git a/Engine/Source/Runtime/Slate/Public/Widgets/Text/SRichTextBlock.h b/Engine/Source/Runtime/Slate/Public/Widgets/Text/SRichTextBlock.h
--- a/Engine/Source/Runtime/Slate/Public/Widgets/Text/SRichTextBlock.h
+++ b/Engine/Source/Runtime/Slate/Public/Widgets/Text/SRichTextBlock.h
@@ -7,7 +7,7 @@
 #include "Layout/Margin.h"
 #include "Styling/SlateWidgetStyleAsset.h"
 #include "Widgets/DeclarativeSyntaxSupport.h"
-#include "Widgets/SLeafWidget.h"
+#include "Widgets/SWidget.h"
 #include "Styling/SlateTypes.h"
 #include "Styling/CoreStyle.h"
 #include "Framework/Text/TextLayout.h"
@@ -30,9 +30,8 @@ enum class ETextShapingMethod : uint8;
  * A rich static text widget. 
  * Through the use of markup and text decorators, text with different styles, embedded image and widgets can be achieved.
  */
-class SLATE_API SRichTextBlock : public SLeafWidget
+class SLATE_API SRichTextBlock : public SWidget
 {
-	SLATE_DECLARE_WIDGET(SRichTextBlock, SLeafWidget)
 public:
 
 	SLATE_BEGIN_ARGS( SRichTextBlock )
@@ -172,16 +171,16 @@ public:
 	 */
 	const FText& GetText() const
 	{
-		return BoundText.Get();
+		return BoundText.Get(FText::GetEmpty());
 	}
 
 	/**
 	 * Sets the text for this text block
 	 */
-	void SetText(TAttribute<FText> InTextAttr );
+	void SetText( const TAttribute<FText>& InTextAttr );
 
 	/** See HighlightText attribute */
-	void SetHighlightText(TAttribute<FText> InHighlightText);
+	void SetHighlightText( const TAttribute<FText>& InHighlightText );
 
 	/** See TextShapingMethod attribute */
 	void SetTextShapingMethod(const TOptional<ETextShapingMethod>& InTextShapingMethod);
@@ -190,31 +189,31 @@ public:
 	void SetTextFlowDirection(const TOptional<ETextFlowDirection>& InTextFlowDirection);
 
 	/** See WrapTextAt attribute */
-	void SetWrapTextAt(TAttribute<float> InWrapTextAt);
+	void SetWrapTextAt(const TAttribute<float>& InWrapTextAt);
 
 	/** See AutoWrapText attribute */
-	void SetAutoWrapText(TAttribute<bool> InAutoWrapText);
+	void SetAutoWrapText(const TAttribute<bool>& InAutoWrapText);
 
 	/** Set WrappingPolicy attribute */
-	void SetWrappingPolicy(TAttribute<ETextWrappingPolicy> InWrappingPolicy);
+	void SetWrappingPolicy(const TAttribute<ETextWrappingPolicy>& InWrappingPolicy);
 
 	/** Set TransformPolicy attribute */
-	void SetTransformPolicy(TAttribute<ETextTransformPolicy> InTransformPolicy);
+	void SetTransformPolicy(const TAttribute<ETextTransformPolicy>& InTransformPolicy);
 
 	/** See LineHeightPercentage attribute */
-	void SetLineHeightPercentage(TAttribute<float> InLineHeightPercentage);
+	void SetLineHeightPercentage(const TAttribute<float>& InLineHeightPercentage);
 
 	/** See Margin attribute */
-	void SetMargin(TAttribute<FMargin> InMargin);
+	void SetMargin(const TAttribute<FMargin>& InMargin);
 
 	/** See Justification attribute */
-	void SetJustification(TAttribute<ETextJustify::Type> InJustification);
+	void SetJustification(const TAttribute<ETextJustify::Type>& InJustification);
 
 	/** See TextStyle argument */
 	void SetTextStyle(const FTextBlockStyle& InTextStyle);
 
 	/** See MinDesiredWidth attribute */
-	void SetMinDesiredWidth(TAttribute<float> InMinDesiredWidth);
+	void SetMinDesiredWidth(const TAttribute<float>& InMinDesiredWidth);
 
 	/**  */
 	void SetDecoratorStyleSet(const ISlateStyle* NewDecoratorStyleSet);
@@ -230,47 +229,48 @@ public:
 protected:
 	//~ SWidget interface
 	virtual FVector2D ComputeDesiredSize(float) const override;
+	virtual bool ComputeVolatility() const override;
 	//~ End of SWidget interface
 
 private:
 
 	/** The text displayed in this text block */
-	TSlateAttribute<FText> BoundText;
+	TAttribute< FText > BoundText;
 
 	/** The wrapped layout for this text block */
-	TUniquePtr<FSlateTextBlockLayout> TextLayoutCache;
+	TUniquePtr< FSlateTextBlockLayout > TextLayoutCache;
 
 	/** Default style used by the TextLayout */
 	FTextBlockStyle TextStyle;
 
-	/** Highlight this text in the text block */
-	TSlateAttribute<FText> HighlightText;
+	/** Highlight this text in the textblock */
+	TAttribute<FText> HighlightText;
 
 	/** Whether text wraps onto a new line when it's length exceeds this width; if this value is zero or negative, no wrapping occurs. */
-	TSlateAttribute<float> WrapTextAt;
+	TAttribute<float> WrapTextAt;
+	
+	/** True if we're wrapping text automatically based on the computed horizontal space for this widget */
+	TAttribute<bool> AutoWrapText;
 
 	/** The wrapping policy we're using */
-	TSlateAttribute<ETextWrappingPolicy> WrappingPolicy;
+	TAttribute<ETextWrappingPolicy> WrappingPolicy;
 
 	/** The transform policy we're using */
-	TSlateAttribute<ETextTransformPolicy> TransformPolicy;
-
-	/** The amount to scale each lines height by. */
-	TSlateAttribute<ETextJustify::Type> Justification;
-
-	/** True if we're wrapping text automatically based on the computed horizontal space for this widget */
-	TSlateAttribute<bool> AutoWrapText;
+	TAttribute<ETextTransformPolicy> TransformPolicy;
 
 	/** The amount of blank space left around the edges of text area. */
-	TSlateAttribute<FMargin> Margin;
+	TAttribute< FMargin > Margin;
+
+	/** The amount to scale each lines height by. */
+	TAttribute< ETextJustify::Type > Justification; 
 
 	/** How the text should be aligned with the margin. */
-	TSlateAttribute<float> LineHeightPercentage;
+	TAttribute< float > LineHeightPercentage;
 
 	/** Prevents the text block from being smaller than desired in certain cases (e.g. when it is empty) */
-	TSlateAttribute<float> MinDesiredWidth;
+	TAttribute<float> MinDesiredWidth;
 
-	/** Use to Scale the entire text block */
+	/** Use to Scale the entire Text Block*/
 	float TextBlockScale = 1.0f;
 
 	/**  */
diff --git a/Engine/Source/Runtime/Slate/Public/Widgets/Text/STextBlock.h b/Engine/Source/Runtime/Slate/Public/Widgets/Text/STextBlock.h
--- a/Engine/Source/Runtime/Slate/Public/Widgets/Text/STextBlock.h
+++ b/Engine/Source/Runtime/Slate/Public/Widgets/Text/STextBlock.h
@@ -38,7 +38,6 @@ namespace ETextRole
  */
 class SLATE_API STextBlock : public SLeafWidget
 {
-	SLATE_DECLARE_WIDGET(STextBlock, SLeafWidget)
 
 public:
 
@@ -93,7 +92,7 @@ public:
 		/** The color used to highlight the specified text */
 		SLATE_ATTRIBUTE( FLinearColor, HighlightColor )
 
-		/** The brush used to highlight the specified text */
+		/** The brush used to highlight the specified text*/
 		SLATE_ATTRIBUTE( const FSlateBrush*, HighlightShape )
 
 		/** Highlight this text in the text block */
@@ -170,20 +169,38 @@ public:
 	}
 	
 public:
-	/** Sets the text for this text block */
-	void SetText(TAttribute<FText> InText);
+	UE_DEPRECATED(4.24, "SetText taking FString is deprecated. Use the FText version instead")
+	void SetText( const TAttribute< FString >& InText );
+	UE_DEPRECATED(4.24, "SetText taking FString is deprecated. Use the FText version instead")
+	void SetText( const FString& InText );
 
-	/** Sets the highlight text for this text block */
+	/**
+	 * Sets the text for this text block
+	 *
+	 * @param	InText	The new text to display
+	 */
+	void SetText( const TAttribute< FText >& InText );
+	void SetText( const FText& InText );
+
+	/**
+	* Sets the highlight text for this text block 
+	*
+	* @param	InText	The new text to highlight
+	*/
 	void SetHighlightText(TAttribute<FText> InText);
 
-	/** Sets the font used to draw the text	*/
-	void SetFont(TAttribute<FSlateFontInfo> InFont);
+	/**
+	 * Sets the font used to draw the text
+	 *
+	 * @param	InFont	The new font to use
+	 */
+	void SetFont(const TAttribute< FSlateFontInfo >& InFont);
 
 	/** Sets the brush used to strike through the text */
-	void SetStrikeBrush(TAttribute<const FSlateBrush*> InStrikeBrush);
+	void SetStrikeBrush(const TAttribute<const FSlateBrush*>& InStrikeBrush);
 
 	/** See ColorAndOpacity attribute */
-	void SetColorAndOpacity(TAttribute<FSlateColor> InColorAndOpacity);
+	void SetColorAndOpacity(const TAttribute<FSlateColor>& InColorAndOpacity);
 
 	/** See TextStyle argument */
 	void SetTextStyle(const FTextBlockStyle* InTextStyle);
@@ -195,44 +212,35 @@ public:
 	void SetTextFlowDirection(const TOptional<ETextFlowDirection>& InTextFlowDirection);
 
 	/** See WrapTextAt attribute */
-	void SetWrapTextAt(TAttribute<float> InWrapTextAt);
+	void SetWrapTextAt(const TAttribute<float>& InWrapTextAt);
 
 	/** See AutoWrapText attribute */
-	void SetAutoWrapText(TAttribute<bool> InAutoWrapText);
+	void SetAutoWrapText(const TAttribute<bool>& InAutoWrapText);
 
 	/** Set WrappingPolicy attribute */
-	void SetWrappingPolicy(TAttribute<ETextWrappingPolicy> InWrappingPolicy);
+	void SetWrappingPolicy(const TAttribute<ETextWrappingPolicy>& InWrappingPolicy);
 
 	/** Set TransformPolicy attribute */
-	void SetTransformPolicy(TAttribute<ETextTransformPolicy> InTransformPolicy);
-
-	/** Get TransformPolicy attribute */
-	UE_DEPRECATED(5.0, "GetTransformPolicy is not accessible anymore since it's attribute value may not have been updated yet.")
-	ETextTransformPolicy GetTransformPolicy() const { return GetTransformPolicyImpl(); }
+	void SetTransformPolicy(const TAttribute<ETextTransformPolicy>& InTransformPolicy);
+	ETextTransformPolicy GetTransformPolicy() const;
 
 	/** See ShadowOffset attribute */
-	void SetShadowOffset(TAttribute<FVector2D> InShadowOffset);
+	void SetShadowOffset(const TAttribute<FVector2D>& InShadowOffset);
 
 	/** See ShadowColorAndOpacity attribute */
-	void SetShadowColorAndOpacity(TAttribute<FLinearColor> InShadowColorAndOpacity);
-
-	/** See HighlightColor attribute */
-	void SetHighlightColor(TAttribute<FLinearColor> InHighlightColor);
-	
-	/** See HighlightShape attribute */
-	void SetHighlightShape(TAttribute<const FSlateBrush*> InHighlightShape);
+	void SetShadowColorAndOpacity(const TAttribute<FLinearColor>& InShadowColorAndOpacity);
 
 	/** See MinDesiredWidth attribute */
-	void SetMinDesiredWidth(TAttribute<float> InMinDesiredWidth);
+	void SetMinDesiredWidth(const TAttribute<float>& InMinDesiredWidth);
 
 	/** See LineHeightPercentage attribute */
-	void SetLineHeightPercentage(TAttribute<float> InLineHeightPercentage);
+	void SetLineHeightPercentage(const TAttribute<float>& InLineHeightPercentage);
 
 	/** See Margin attribute */
-	void SetMargin(TAttribute<FMargin> InMargin);
+	void SetMargin(const TAttribute<FMargin>& InMargin);
 
 	/** See Justification attribute */
-	void SetJustification(TAttribute<ETextJustify::Type> InJustification);
+	void SetJustification(const TAttribute<ETextJustify::Type>& InJustification);
 
 	// SWidget interface
 	virtual int32 OnPaint( const FPaintArgs& Args, const FGeometry& AllottedGeometry, const FSlateRect& MyCullingRect, FSlateWindowElementList& OutDrawElements, int32 LayerId, const FWidgetStyle& InWidgetStyle, bool bParentEnabled ) const override;
@@ -243,6 +251,11 @@ public:
 #endif
 	// End of SWidget interface
 
+protected:
+	// Begin SWidget overrides.
+	virtual bool ComputeVolatility() const override;
+	// End of SWidget interface
+
 private:
 	/** Get the computed text style to use with the text marshaller */
 	FTextBlockStyle GetComputedTextStyle() const;
@@ -255,9 +268,6 @@ private:
 
 	/** Gets the current strike brush */
 	const FSlateBrush* GetStrikeBrush() const;
-
-	/** Get TransformPolicy attribute */
-	ETextTransformPolicy GetTransformPolicyImpl() const;
 	
 	/** Gets the current shadow offset */
 	FVector2D GetShadowOffset() const;
@@ -272,14 +282,11 @@ private:
 	const FSlateBrush* GetHighlightShape() const;
 
 	/** Call to invalidate this text block */
-	void InvalidateText(EInvalidateWidgetReason InvalidateReason);
-
-private:
-
+	void InvalidateText(EInvalidateWidget InvalidateReason);
 
 private:
 	/** The text displayed in this text block */
-	TSlateAttribute<FText> BoundText;
+	TAttribute< FText > BoundText;
 
 	/** The wrapped layout for this text block */
 	TUniquePtr< FSlateTextBlockLayout > TextLayoutCache;
@@ -288,73 +295,52 @@ private:
 	FTextBlockStyle TextStyle;
 
 	/** Sets the font used to draw the text */
-	TSlateAttribute<FSlateFontInfo> Font;
+	TAttribute< FSlateFontInfo > Font;
 
 	/** Sets the brush used to strike through the text */
-	TSlateAttribute<const FSlateBrush*> StrikeBrush;
+	TAttribute< const FSlateBrush* > StrikeBrush;
 
 	/** Text color and opacity */
-	TSlateAttribute<FSlateColor> ColorAndOpacity;
+	TAttribute<FSlateColor> ColorAndOpacity;
 
 	/** Drop shadow offset in pixels */
-	TSlateAttribute<FVector2D> ShadowOffset;
+	TAttribute< FVector2D > ShadowOffset;
 
 	/** Shadow color and opacity */
-	TSlateAttribute<FLinearColor> ShadowColorAndOpacity;
+	TAttribute<FLinearColor> ShadowColorAndOpacity;
 
-	/** The color used to highlight the specified text */
-	TSlateAttribute<FLinearColor> HighlightColor;
+	TAttribute<FLinearColor> HighlightColor;
 
-	/** The brush used to highlight the specified text */
-	TSlateAttribute<const FSlateBrush*> HighlightShape;
+	TAttribute< const FSlateBrush* > HighlightShape;
 
-	/** Highlight this text in the TextBlock */
-	TSlateAttribute<FText> HighlightText;
+	/** Highlight this text in the textblock */
+	TAttribute<FText> HighlightText;
 
 	/** Whether text wraps onto a new line when it's length exceeds this width; if this value is zero or negative, no wrapping occurs. */
-	TSlateAttribute<float> WrapTextAt;
+	TAttribute<float> WrapTextAt;
 
 	/** True if we're wrapping text automatically based on the computed horizontal space for this widget */
-	TSlateAttribute<bool> AutoWrapText;
+	TAttribute<bool> AutoWrapText;
 
 	/** The wrapping policy we're using */
-	TSlateAttribute<ETextWrappingPolicy> WrappingPolicy;
+	TAttribute<ETextWrappingPolicy> WrappingPolicy;
 
 	/** The transform policy we're using */
-	TSlateAttribute<ETextTransformPolicy> TransformPolicy;
+	TAttribute<ETextTransformPolicy> TransformPolicy;
 
 	/** The amount of blank space left around the edges of text area. */
-	TSlateAttribute<FMargin> Margin;
+	TAttribute< FMargin > Margin;
 
 	/** The amount to scale each lines height by. */
-	TSlateAttribute<ETextJustify::Type> Justification;
+	TAttribute< ETextJustify::Type > Justification; 
 
 	/** How the text should be aligned with the margin. */
-	TSlateAttribute<float> LineHeightPercentage;
+	TAttribute< float > LineHeightPercentage;
 
 	/** Prevents the text block from being smaller than desired in certain cases (e.g. when it is empty) */
-	TSlateAttribute<float> MinDesiredWidth;
+	TAttribute<float> MinDesiredWidth;
 
 	/** If this is enabled, text shaping, wrapping, justification are disabled in favor of much faster text layout and measurement. */
 	mutable TOptional<FVector2D> CachedSimpleDesiredSize;
-
-	/** Flags used to check if the SlateAttribute is set. */
-	union
-	{
-		struct 
-		{
-			uint16 bIsAttributeFontSet : 1;
-			uint16 bIsAttributeStrikeBrushSet : 1;
-			uint16 bIsAttributeColorAndOpacitySet : 1;
-			uint16 bIsAttributeShadowOffsetSet : 1;
-			uint16 bIsAttributeShadowColorAndOpacitySet : 1;
-			uint16 bIsAttributeHighlightColorSet : 1;
-			uint16 bIsAttributeHighlightShapeSet : 1;
-			uint16 bIsAttributeWrapTextAtSet : 1;
-			uint16 bIsAttributeTransformPolicySet : 1;
-		};
-		uint16 Union_IsAttributeSet;
-	};
-
 	bool bSimpleTextMode;
 };
diff --git a/Engine/Source/Runtime/Slate/Public/Widgets/Text/SlateTextBlockLayout.h b/Engine/Source/Runtime/Slate/Public/Widgets/Text/SlateTextBlockLayout.h
--- a/Engine/Source/Runtime/Slate/Public/Widgets/Text/SlateTextBlockLayout.h
+++ b/Engine/Source/Runtime/Slate/Public/Widgets/Text/SlateTextBlockLayout.h
@@ -20,45 +20,7 @@ enum class ETextShapingMethod : uint8;
 class SLATE_API FSlateTextBlockLayout
 {
 public:
-	struct FWidgetDesiredSizeArgs
-	{
-		FORCEINLINE FWidgetDesiredSizeArgs(
-			const FText& InText,
-			const FText& InHighlightText,
-			const float InWrapTextAt,
-			const bool InAutoWrapText,
-			const ETextWrappingPolicy InWrappingPolicy,
-			const ETextTransformPolicy InTransformPolicy,
-			const FMargin& InMargin,
-			const float InLineHeightPercentage,
-			const ETextJustify::Type InJustification
-		)
-			: Text(InText)
-			, HighlightText(InHighlightText)
-			, Margin(InMargin)
-			, WrapTextAt(InWrapTextAt)
-			, LineHeightPercentage(InLineHeightPercentage)
-			, WrappingPolicy(InWrappingPolicy)
-			, TransformPolicy(InTransformPolicy)
-			, Justification(InJustification)
-			, AutoWrapText(InAutoWrapText)
-		{
-		}
-
-		const FText Text = FText::GetEmpty();
-		const FText HighlightText = FText::GetEmpty();
-		const FMargin Margin;
-		const float WrapTextAt = 0.f;
-		const float LineHeightPercentage;
-		const ETextWrappingPolicy WrappingPolicy;
-		const ETextTransformPolicy TransformPolicy;
-		const ETextJustify::Type Justification;
-		const bool AutoWrapText = false;
-	};
-
-	struct
-		UE_DEPRECATED(5.0, "FWidgetArgs is deprecated. Upgrade to FWidgetDesiredSizeArgs instead.")
-		FWidgetArgs
+	struct FWidgetArgs
 	{
 		FORCEINLINE FWidgetArgs(
 			const TAttribute<FText>& InText, 
@@ -102,15 +64,7 @@ public:
 	/**
 	 * Get the computed desired size for this layout, updating the internal cache as required
 	 */
-	FVector2D ComputeDesiredSize(const FWidgetDesiredSizeArgs& InWidgetArgs, const float InScale, const FTextBlockStyle& InTextStyle);
-
-	PRAGMA_DISABLE_DEPRECATION_WARNINGS
-	/**
-	 * Get the computed desired size for this layout, updating the internal cache as required
-	 */
-	UE_DEPRECATED(5.0, "FWidgetArgs is deprecated. Upgrade to FWidgetDesiredSizeArgs instead.")
 	FVector2D ComputeDesiredSize(const FWidgetArgs& InWidgetArgs, const float InScale, const FTextBlockStyle& InTextStyle);
-	PRAGMA_DISABLE_DEPRECATION_WARNINGS
 
 	/**
 	 * Gets the last computed desired size.
