When closing the app, only ask major tabs whether they want to block closing.

diff --git a/Engine/Source/Runtime/Slate/Private/Framework/Docking/TabManager.cpp b/Engine/Source/Runtime/Slate/Private/Framework/Docking/TabManager.cpp
--- a/Engine/Source/Runtime/Slate/Private/Framework/Docking/TabManager.cpp
+++ b/Engine/Source/Runtime/Slate/Private/Framework/Docking/TabManager.cpp
@@ -2217,7 +2217,10 @@ bool FTabManager::CanCloseManager( const TSet< TSharedRef<SDockTab> >& TabsToIgn
 		
 		for (int32 TabIndex=0; bCanCloseManager && TabIndex < AreasTabs.Num(); ++TabIndex)	
 		{
-			bCanCloseManager = TabsToIgnore.Contains( AreasTabs[TabIndex] ) || AreasTabs[TabIndex]->CanCloseTab();
+			bCanCloseManager =
+				TabsToIgnore.Contains( AreasTabs[TabIndex] ) ||
+				AreasTabs[TabIndex]->GetTabRole() != ETabRole::MajorTab ||
+				AreasTabs[TabIndex]->CanCloseTab();
 		}		
 	}
 
