Slate: Fix DPI Scale in PIE. The first frame had a  DPI scaler of 0 causing a divide by zero exception.
#rnx

diff --git a/Engine/Source/Runtime/Slate/Private/Widgets/Layout/SDPIScaler.cpp b/Engine/Source/Runtime/Slate/Private/Widgets/Layout/SDPIScaler.cpp
--- a/Engine/Source/Runtime/Slate/Private/Widgets/Layout/SDPIScaler.cpp
+++ b/Engine/Source/Runtime/Slate/Private/Widgets/Layout/SDPIScaler.cpp
@@ -11,7 +11,7 @@ void SDPIScaler::PrivateRegisterAttributes(FSlateAttributeInitializer& Attribute
 
 SDPIScaler::SDPIScaler()
 	: ChildSlot(*this)
-	, DPIScaleAttribute(*this)
+	, DPIScaleAttribute(*this, 1.f)
 {
 	SetCanTick(false);
 	bCanSupportFocus = false;
@@ -21,6 +21,7 @@ SDPIScaler::SDPIScaler()
 void SDPIScaler::Construct( const FArguments& InArgs )
 {
 	SetDPIScale(InArgs._DPIScale);
+	DPIScaleAttribute.UpdateNow(*this);
 
 	ChildSlot
 	[
@@ -46,7 +47,12 @@ void SDPIScaler::OnArrangeChildren( const FGeometry& AllottedGeometry, FArranged
 	
 FVector2D SDPIScaler::ComputeDesiredSize( float ) const
 {
-	return DPIScaleAttribute.Get() * ChildSlot.GetWidget()->GetDesiredSize();
+	float DPIScaleValue = DPIScaleAttribute.Get();
+	if (ensure(DPIScaleValue > 0.f))
+	{
+		return DPIScaleValue * ChildSlot.GetWidget()->GetDesiredSize();
+	}
+	return ChildSlot.GetWidget()->GetDesiredSize();
 }
 
 FChildren* SDPIScaler::GetChildren()
@@ -64,7 +70,7 @@ void SDPIScaler::SetContent(TSharedRef<SWidget> InContent)
 
 void SDPIScaler::SetDPIScale(TAttribute<float> InDPIScale)
 {
-	if (DPIScaleAttribute.Assign(*this, MoveTemp(InDPIScale)))
+	if (DPIScaleAttribute.Assign(*this, MoveTemp(InDPIScale), 1.f))
 	{
 		InvalidatePrepass();
 	}
