Slate: Convert SBorder attributes to SlateAttributes

diff --git a/Engine/Plugins/AI/EnvironmentQueryEditor/Source/EnvironmentQueryEditor/Private/SEnvQueryProfiler.cpp b/Engine/Plugins/AI/EnvironmentQueryEditor/Source/EnvironmentQueryEditor/Private/SEnvQueryProfiler.cpp
--- a/Engine/Plugins/AI/EnvironmentQueryEditor/Source/EnvironmentQueryEditor/Private/SEnvQueryProfiler.cpp
+++ b/Engine/Plugins/AI/EnvironmentQueryEditor/Source/EnvironmentQueryEditor/Private/SEnvQueryProfiler.cpp
@@ -261,7 +261,7 @@ void SEnvQueryProfilerTableRow::Construct(const FArguments& InArgs, const TShare
 	StatInfo = InStatInfo;
 
 	SMultiColumnTableRow<FEnvQueryProfilerStatDataPtr>::Construct(FSuperRowType::FArguments(), InOwnerTableView);
-	BorderImage = TAttribute<const FSlateBrush*>(this, &SEnvQueryProfilerTableRow::GetBorder);
+	SetBorderImage(TAttribute<const FSlateBrush*>(this, &SEnvQueryProfilerTableRow::GetBorder));
 }
 
 const FSlateBrush* SEnvQueryProfilerTableRow::GetBorder() const
diff --git a/Engine/Plugins/NetcodeUnitTest/NetcodeUnitTest/Source/NetcodeUnitTest/Private/UI/SLogWidget.cpp b/Engine/Plugins/NetcodeUnitTest/NetcodeUnitTest/Source/NetcodeUnitTest/Private/UI/SLogWidget.cpp
--- a/Engine/Plugins/NetcodeUnitTest/NetcodeUnitTest/Source/NetcodeUnitTest/Private/UI/SLogWidget.cpp
+++ b/Engine/Plugins/NetcodeUnitTest/NetcodeUnitTest/Source/NetcodeUnitTest/Private/UI/SLogWidget.cpp
@@ -45,8 +45,8 @@ IMPLEMENT_GET_PRIVATE_VAR(SEditableTextBox, EditableText, TSharedPtr<SEditableTe
 // Enable access to SButton.Style
 IMPLEMENT_GET_PRIVATE_VAR(SButton, Style, const FButtonStyle*);
 
-// Enable reading SBorder.BorderBackgroundColor
-IMPLEMENT_GET_PRIVATE_VAR(SBorder, BorderBackgroundColor, TAttribute<FSlateColor>);
+// Enable reading SBorder::GetBorderBackgroundColor
+IMPLEMENT_GET_PROTECTED_FUNC_CONST(SBorder, GetBorderBackgroundColor, const FSlateColor&, void,, const);
 
 // Enable access to SDockTab::GetCurrentStyle, using the CALL_PROTECTED macro
 IMPLEMENT_GET_PROTECTED_FUNC_CONST(SDockTab, GetCurrentStyle, const FDockTabStyle&, void,, const);
@@ -1440,7 +1440,7 @@ void SLogWidget::Tick(const FGeometry& AllottedGeometry, const double InCurrentT
 		if (ActiveFindLabel.IsValid() && ActiveFindLabel->GetVisibility() != EVisibility::Collapsed)
 		{
 			const float FadeDuration = 2.f;
-			FLinearColor CurColor = (GET_PRIVATE(SBorder, ActiveFindLabel, BorderBackgroundColor)).Get().GetSpecifiedColor();
+			FLinearColor CurColor = CALL_PROTECTED(SBorder, ActiveFindLabel, GetBorderBackgroundColor)().GetSpecifiedColor();
 
 			CurColor.A -= (InDeltaTime / FadeDuration);
 
diff --git a/Engine/Plugins/Runtime/DataRegistry/Source/DataRegistryEditor/Private/SDataRegistryListViewRow.cpp b/Engine/Plugins/Runtime/DataRegistry/Source/DataRegistryEditor/Private/SDataRegistryListViewRow.cpp
--- a/Engine/Plugins/Runtime/DataRegistry/Source/DataRegistryEditor/Private/SDataRegistryListViewRow.cpp
+++ b/Engine/Plugins/Runtime/DataRegistry/Source/DataRegistryEditor/Private/SDataRegistryListViewRow.cpp
@@ -23,8 +23,7 @@ void SDataRegistryListViewRow::Construct(const FArguments& InArgs, const TShared
 		InOwnerTableView
 	);
 
-	BorderImage = TAttribute<const FSlateBrush*>(this, &SDataRegistryListViewRow::GetBorder);
-
+	SetBorderImage(TAttribute<const FSlateBrush*>(this, &SDataRegistryListViewRow::GetBorder));
 }
 
 FReply SDataRegistryListViewRow::OnMouseButtonUp(const FGeometry& MyGeometry, const FPointerEvent& MouseEvent)
diff --git a/Engine/Source/Developer/CollisionAnalyzer/Private/SCAQueryTableRow.cpp b/Engine/Source/Developer/CollisionAnalyzer/Private/SCAQueryTableRow.cpp
--- a/Engine/Source/Developer/CollisionAnalyzer/Private/SCAQueryTableRow.cpp
+++ b/Engine/Source/Developer/CollisionAnalyzer/Private/SCAQueryTableRow.cpp
@@ -13,7 +13,7 @@ void SCAQueryTableRow::Construct(const FArguments& InArgs, const TSharedRef<STab
 
 	if(Item->bIsGroup)
 	{
-		BorderImage = FInvalidatableBrushAttribute(FCollisionAnalyzerStyle::Get()->GetBrush("CollisionAnalyzer.GroupBackground"));
+		SetBorderImage(FCollisionAnalyzerStyle::Get()->GetBrush("CollisionAnalyzer.GroupBackground"));
 	}
 }
 
diff --git a/Engine/Source/Editor/DataTableEditor/Private/SDataTableListViewRow.cpp b/Engine/Source/Editor/DataTableEditor/Private/SDataTableListViewRow.cpp
--- a/Engine/Source/Editor/DataTableEditor/Private/SDataTableListViewRow.cpp
+++ b/Engine/Source/Editor/DataTableEditor/Private/SDataTableListViewRow.cpp
@@ -28,7 +28,7 @@ void SDataTableListViewRow::Construct(const FArguments& InArgs, const TSharedRef
 		InOwnerTableView
 	);
 
-	BorderImage = TAttribute<const FSlateBrush*>(this, &SDataTableListViewRow::GetBorder);
+	SetBorderImage(TAttribute<const FSlateBrush*>(this, &SDataTableListViewRow::GetBorder));
 
 }
 
diff --git a/Engine/Source/Editor/UnrealEd/Private/ConsolidateWindow.cpp b/Engine/Source/Editor/UnrealEd/Private/ConsolidateWindow.cpp
--- a/Engine/Source/Editor/UnrealEd/Private/ConsolidateWindow.cpp
+++ b/Engine/Source/Editor/UnrealEd/Private/ConsolidateWindow.cpp
@@ -365,7 +365,7 @@ void SConsolidateToolWidget::Construct( const FArguments& InArgs )
 	SelectedListItem = NULL;
 	bSavePackagesChecked = ISourceControlModule::Get().IsEnabled();
 
-	this->BorderImage = FInvalidatableBrushAttribute(FEditorStyle::GetBrush("NoBorder"));
+	SetBorderImage(FEditorStyle::GetBrush("NoBorder"));
 
 	ChildSlot
 	[
diff --git a/Engine/Source/Runtime/Slate/Private/Widgets/Input/SButton.cpp b/Engine/Source/Runtime/Slate/Private/Widgets/Input/SButton.cpp
--- a/Engine/Source/Runtime/Slate/Private/Widgets/Input/SButton.cpp
+++ b/Engine/Source/Runtime/Slate/Private/Widgets/Input/SButton.cpp
@@ -108,7 +108,7 @@ int32 SButton::OnPaint(const FPaintArgs& Args, const FGeometry& AllottedGeometry
 			AllottedGeometry.ToPaintGeometry(),
 			BrushResource,
 			DrawEffects,
-			BrushResource->GetTint(InWidgetStyle) * InWidgetStyle.GetColorAndOpacityTint() * BorderBackgroundColor.Get().GetColor(InWidgetStyle)
+			BrushResource->GetTint(InWidgetStyle) * InWidgetStyle.GetColorAndOpacityTint() * GetBorderBackgroundColor().GetColor(InWidgetStyle)
 		);
 	}
 
@@ -421,7 +421,7 @@ bool SButton::IsInteractable() const
 bool SButton::ComputeVolatility() const
 {
 	// Note: we need to be careful with button volatility.  The parent SBorder class always has bound delegates to button but the following are the only thing that actually would be bound that would not be caught by an Invalidate call alone
-	return DesiredSizeScale.IsBound() || BorderBackgroundColor.IsBound() || ContentPadding.IsBound();
+	return ContentPadding.IsBound();
 }
 
 TEnumAsByte<EButtonClickMethod::Type> SButton::GetClickMethodFromInputType(const FPointerEvent& MouseEvent) const
diff --git a/Engine/Source/Runtime/Slate/Private/Widgets/Layout/SBorder.cpp b/Engine/Source/Runtime/Slate/Private/Widgets/Layout/SBorder.cpp
--- a/Engine/Source/Runtime/Slate/Private/Widgets/Layout/SBorder.cpp
+++ b/Engine/Source/Runtime/Slate/Private/Widgets/Layout/SBorder.cpp
@@ -5,10 +5,21 @@
 
 static FName SBorderTypeName("SBorder");
 
+SLATE_IMPLEMENT_WIDGET(SBorder)
+void SBorder::PrivateRegisterAttributes(FSlateAttributeInitializer& AttributeInitializer)
+{
+	SLATE_ADD_MEMBER_ATTRIBUTE_DEFINITION_WITH_NAME(AttributeInitializer, "BorderImage", BorderImageAttribute, EInvalidateWidgetReason::Paint);
+	SLATE_ADD_MEMBER_ATTRIBUTE_DEFINITION_WITH_NAME(AttributeInitializer, "BorderBackgroundColor", BorderBackgroundColorAttribute, EInvalidateWidgetReason::Paint);
+	SLATE_ADD_MEMBER_ATTRIBUTE_DEFINITION_WITH_NAME(AttributeInitializer, "DesiredSizeScale", DesiredSizeScaleAttribute, EInvalidateWidgetReason::Layout);
+	SLATE_ADD_MEMBER_ATTRIBUTE_DEFINITION_WITH_NAME(AttributeInitializer, "ShowDisabledEffect", ShowDisabledEffectAttribute, EInvalidateWidgetReason::Paint);
+}
+
 SBorder::SBorder()
-	: BorderImage( FCoreStyle::Get().GetBrush( "Border" ) )
-	, BorderBackgroundColor( FLinearColor::White )
-	, DesiredSizeScale(FVector2D(1,1))
+	: BorderImageAttribute(*this, FCoreStyle::Get().GetBrush( "Border" ))
+	, BorderBackgroundColorAttribute(*this, FLinearColor::White)
+	, DesiredSizeScaleAttribute(*this, FVector2D(1,1))
+	, ShowDisabledEffectAttribute(*this, true)
+	, bFlipForRightToLeftFlowDirection(false)
 {
 }
 
@@ -23,14 +34,13 @@ void SBorder::Construct( const SBorder::FArguments& InArgs )
 
 	SetContentScale(InArgs._ContentScale);
 	SetColorAndOpacity(InArgs._ColorAndOpacity);
-	DesiredSizeScale = InArgs._DesiredSizeScale;
-
-	ShowDisabledEffect = InArgs._ShowEffectWhenDisabled;
+	SetDesiredSizeScale(InArgs._DesiredSizeScale);
+	SetShowEffectWhenDisabled(InArgs._ShowEffectWhenDisabled);
 
 	bFlipForRightToLeftFlowDirection = InArgs._FlipForRightToLeftFlowDirection;
 
-	BorderImage = InArgs._BorderImage;
-	BorderBackgroundColor = InArgs._BorderBackgroundColor;
+	SetBorderImage(InArgs._BorderImage);
+	SetBorderBackgroundColor(InArgs._BorderBackgroundColor);
 	SetForegroundColor(InArgs._ForegroundColor);
 
 	if (InArgs._OnMouseButtonDown.IsBound())
@@ -82,13 +92,13 @@ void SBorder::ClearContent()
 
 int32 SBorder::OnPaint( const FPaintArgs& Args, const FGeometry& AllottedGeometry, const FSlateRect& MyCullingRect, FSlateWindowElementList& OutDrawElements, int32 LayerId, const FWidgetStyle& InWidgetStyle, bool bParentEnabled ) const
 {
-	const FSlateBrush* BrushResource = BorderImage.Get();
+	const FSlateBrush* BrushResource = BorderImageAttribute.Get();
 		
 	const bool bEnabled = ShouldBeEnabled(bParentEnabled);
 
 	if ( BrushResource && BrushResource->DrawAs != ESlateBrushDrawType::NoDrawType )
 	{
-		const bool bShowDisabledEffect = ShowDisabledEffect.Get();
+		const bool bShowDisabledEffect = GetShowDisabledEffect();
 		const ESlateDrawEffect DrawEffects = (bShowDisabledEffect && !bEnabled) ? ESlateDrawEffect::DisabledEffect : ESlateDrawEffect::None;
 
 		if (bFlipForRightToLeftFlowDirection && GSlateFlowDirection == EFlowDirection::RightToLeft)
@@ -100,7 +110,7 @@ int32 SBorder::OnPaint( const FPaintArgs& Args, const FGeometry& AllottedGeometr
 				FlippedGeometry.ToPaintGeometry(),
 				BrushResource,
 				DrawEffects,
-				BrushResource->GetTint(InWidgetStyle) * InWidgetStyle.GetColorAndOpacityTint() * BorderBackgroundColor.Get().GetColor(InWidgetStyle)
+				BrushResource->GetTint(InWidgetStyle) * InWidgetStyle.GetColorAndOpacityTint() * BorderBackgroundColorAttribute.Get().GetColor(InWidgetStyle)
 			);
 		}
 		else
@@ -111,7 +121,7 @@ int32 SBorder::OnPaint( const FPaintArgs& Args, const FGeometry& AllottedGeometr
 				AllottedGeometry.ToPaintGeometry(),
 				BrushResource,
 				DrawEffects,
-				BrushResource->GetTint(InWidgetStyle) * InWidgetStyle.GetColorAndOpacityTint() * BorderBackgroundColor.Get().GetColor(InWidgetStyle)
+				BrushResource->GetTint(InWidgetStyle) * InWidgetStyle.GetColorAndOpacityTint() * BorderBackgroundColorAttribute.Get().GetColor(InWidgetStyle)
 			);
 		}
 	}
@@ -119,24 +129,19 @@ int32 SBorder::OnPaint( const FPaintArgs& Args, const FGeometry& AllottedGeometr
 	return SCompoundWidget::OnPaint(Args, AllottedGeometry, MyCullingRect, OutDrawElements, LayerId, InWidgetStyle, bEnabled );
 }
 
-bool SBorder::ComputeVolatility() const
-{
-	return BorderImage.IsBound() || BorderBackgroundColor.IsBound() || DesiredSizeScale.IsBound() || ShowDisabledEffect.IsBound();
-}
-
 FVector2D SBorder::ComputeDesiredSize(float LayoutScaleMultiplier) const
 {
-	return DesiredSizeScale.Get() * SCompoundWidget::ComputeDesiredSize(LayoutScaleMultiplier);
+	return DesiredSizeScaleAttribute.Get() * Super::ComputeDesiredSize(LayoutScaleMultiplier);
 }
 
-void SBorder::SetBorderBackgroundColor(const TAttribute<FSlateColor>& InColorAndOpacity)
+void SBorder::SetBorderBackgroundColor(TAttribute<FSlateColor> InColorAndOpacity)
 {
-	SetAttribute(BorderBackgroundColor, InColorAndOpacity, EInvalidateWidgetReason::Paint);
+	BorderBackgroundColorAttribute.Assign(*this, InColorAndOpacity);
 }
 
-void SBorder::SetDesiredSizeScale(const TAttribute<FVector2D>& InDesiredSizeScale)
+void SBorder::SetDesiredSizeScale(TAttribute<FVector2D> InDesiredSizeScale)
 {
-	SetAttribute(DesiredSizeScale, InDesiredSizeScale, EInvalidateWidgetReason::Layout);
+	DesiredSizeScaleAttribute.Assign(*this, InDesiredSizeScale);
 }
 
 void SBorder::SetHAlign(EHorizontalAlignment HAlign)
@@ -162,12 +167,12 @@ void SBorder::SetPadding(const TAttribute<FMargin>& InPadding)
 	SetAttribute(ChildSlot.SlotPadding, InPadding, EInvalidateWidgetReason::Layout);
 }
 
-void SBorder::SetShowEffectWhenDisabled(const TAttribute<bool>& InShowEffectWhenDisabled)
+void SBorder::SetShowEffectWhenDisabled(TAttribute<bool> InShowEffectWhenDisabled)
 {
-	SetAttribute(ShowDisabledEffect, InShowEffectWhenDisabled, EInvalidateWidgetReason::Paint);
+	ShowDisabledEffectAttribute.Assign(*this, InShowEffectWhenDisabled);
 }
 
-void SBorder::SetBorderImage(const TAttribute<const FSlateBrush*>& InBorderImage)
+void SBorder::SetBorderImage(TAttribute<const FSlateBrush*> InBorderImage)
 {
-	BorderImage.SetImage(*this, InBorderImage);
+	BorderImageAttribute.Assign(*this, InBorderImage);
 }
diff --git a/Engine/Source/Runtime/Slate/Public/Widgets/Layout/SBorder.h b/Engine/Source/Runtime/Slate/Public/Widgets/Layout/SBorder.h
--- a/Engine/Source/Runtime/Slate/Public/Widgets/Layout/SBorder.h
+++ b/Engine/Source/Runtime/Slate/Public/Widgets/Layout/SBorder.h
@@ -23,6 +23,8 @@ class FSlateWindowElementList;
  */
 class SLATE_API SBorder : public SCompoundWidget
 {
+	SLATE_DECLARE_WIDGET(SBorder, SCompoundWidget)
+
 public:
 
 	SLATE_BEGIN_ARGS(SBorder)
@@ -104,10 +106,10 @@ public:
 	void ClearContent();
 
 	/** Sets the color and opacity of the background image of this border. */
-	void SetBorderBackgroundColor(const TAttribute<FSlateColor>& InColorAndOpacity);
+	void SetBorderBackgroundColor(TAttribute<FSlateColor> InColorAndOpacity);
 
 	/** Set the desired size scale multiplier */
-	void SetDesiredSizeScale(const TAttribute<FVector2D>& InDesiredSizeScale);
+	void SetDesiredSizeScale(TAttribute<FVector2D> InDesiredSizeScale);
 	
 	/** See HAlign argument */
 	void SetHAlign(EHorizontalAlignment HAlign);
@@ -119,28 +121,47 @@ public:
 	void SetPadding(const TAttribute<FMargin>& InPadding);
 
 	/** See ShowEffectWhenDisabled attribute */
-	void SetShowEffectWhenDisabled(const TAttribute<bool>& InShowEffectWhenDisabled);
+	void SetShowEffectWhenDisabled(TAttribute<bool> InShowEffectWhenDisabled);
 
 	/** See BorderImage attribute */
-	void SetBorderImage(const TAttribute<const FSlateBrush*>& InBorderImage);
+	void SetBorderImage(TAttribute<const FSlateBrush*> InBorderImage);
 
 public:
 	// SWidget interface
 	virtual int32 OnPaint( const FPaintArgs& Args, const FGeometry& AllottedGeometry, const FSlateRect& MyCullingRect, FSlateWindowElementList& OutDrawElements, int32 LayerId, const FWidgetStyle& InWidgetStyle, bool bParentEnabled ) const override;
-	virtual bool ComputeVolatility() const override;
 	// End of SWidget interface
 
 protected:
-	// Begin SWidget overrides.
+	//~Begin SWidget overrides.
 	virtual FVector2D ComputeDesiredSize(float) const override;
-	// End SWidget overrides.
-
- protected:
+	//~End SWidget overrides.
+
+	const FSlateBrush* GetBorderImageAttribute() const { return BorderImageAttribute.Get(); }
+	const FSlateColor& GetBorderBackgroundColor() const { return BorderBackgroundColorAttribute.Get(); }
+	const FVector2D& GetDesiredSizeScale() const { return DesiredSizeScaleAttribute.Get(); }
+	bool GetShowDisabledEffect() const { return ShowDisabledEffectAttribute.Get(); }
+	TSlateAttributeRef<const FSlateBrush*> GetBorderImageAttributeAttribute() const { return TSlateAttributeRef<const FSlateBrush*>(*this, BorderImageAttribute); }
+	TSlateAttributeRef<FSlateColor> GetBorderBackgroundColorAttribute() const { return TSlateAttributeRef<FSlateColor>(*this, BorderBackgroundColorAttribute); }
+	TSlateAttributeRef<FVector2D> GetDesiredSizeScaleAttribute() const { return TSlateAttributeRef<FVector2D>(*this, DesiredSizeScaleAttribute); }
+	TSlateAttributeRef<bool> GetShowDisabledEffectAttribute() const { return TSlateAttributeRef<bool>(*this, ShowDisabledEffectAttribute); }
+
+#if WITH_EDITORONLY_DATA
+	UE_DEPRECATED(5.0, "Direct access to BorderImage is now deprecated. Use the setter or getter.")
 	FInvalidatableBrushAttribute BorderImage;
-	TAttribute<FSlateColor> BorderBackgroundColor;
-	TAttribute<FVector2D> DesiredSizeScale;
+	UE_DEPRECATED(5.0, "Direct access to BorderBackgroundColor is now deprecated. Use the setter or getter.")
+	FSlateDeprecatedTAttribute<FSlateColor> BorderBackgroundColor;
+	UE_DEPRECATED(5.0, "Direct access to DesiredSizeScale is now deprecated. Use the setter or getter.")
+	FSlateDeprecatedTAttribute<FVector2D> DesiredSizeScale;
+	UE_DEPRECATED(5.0, "Direct access to ShowDisabledEffect is now deprecated. Use the setter or getter.")
+	FSlateDeprecatedTAttribute<bool> ShowDisabledEffect;
+#endif
+
+ private:
+	TSlateAttribute<const FSlateBrush*> BorderImageAttribute;
+	TSlateAttribute<FSlateColor> BorderBackgroundColorAttribute;
+	TSlateAttribute<FVector2D> DesiredSizeScaleAttribute;
 	/** Whether or not to show the disabled effect when this border is disabled */
-	TAttribute<bool> ShowDisabledEffect;
+	TSlateAttribute<bool> ShowDisabledEffectAttribute;
 
 	/** Flips the image if the localization's flow direction is RightToLeft */
 	bool bFlipForRightToLeftFlowDirection;
diff --git a/Engine/Source/Runtime/Slate/Public/Widgets/Views/STableRow.h b/Engine/Source/Runtime/Slate/Public/Widgets/Views/STableRow.h
--- a/Engine/Source/Runtime/Slate/Public/Widgets/Views/STableRow.h
+++ b/Engine/Source/Runtime/Slate/Public/Widgets/Views/STableRow.h
@@ -1135,7 +1135,7 @@ protected:
 		check(InArgs._ExpanderStyleSet);
 		ExpanderStyleSet = InArgs._ExpanderStyleSet;
 
-		this->BorderImage = FInvalidatableBrushAttribute(TAttribute<const FSlateBrush*>(this, &STableRow::GetBorder));
+		SetBorderImage(TAttribute<const FSlateBrush*>(this, &STableRow::GetBorder));
 
 		this->SetForegroundColor(TAttribute<FSlateColor>( this, &STableRow::GetForegroundBasedOnSelection ));
 
