Notification reskin and progress notification support. Progress notifications appear as a notification with short timeout and then progress appears at the bottom of the status bar until progress completes.

diff --git a/Engine/Source/Editor/StatusBar/Private/NotificationTesting.cpp b/Engine/Source/Editor/StatusBar/Private/NotificationTesting.cpp
new file mode 100644
--- /dev/null
+++ b/Engine/Source/Editor/StatusBar/Private/NotificationTesting.cpp
@@ -0,0 +1,219 @@
+// Copyright Epic Games, Inc. All Rights Reserved.
+
+#include "Widgets/Notifications/SNotificationList.h"
+#include "Framework/Notifications/NotificationManager.h"
+#include "TimerManager.h"
+#include "Editor.h"
+
+#if (WITH_EDITOR || IS_PROGRAM) && !UE_BUILD_SHIPPING
+
+struct FTestAsyncTask
+{
+	FProgressNotificationHandle ProgressHandle;
+	FTimerHandle TimerHandle;
+	int32 WorkDone;
+};
+
+static TArray<FTestAsyncTask> Tasks;
+FTimerHandle TestTimerHandle;
+
+static void StartTask(FString TaskName, int32 TotalWork, int32 IncAmount)
+{
+	FText ProgressNotification = FText::FromString(TaskName);
+
+	FTestAsyncTask TestTask;
+	TestTask.WorkDone = 0;
+	TestTask.ProgressHandle = FSlateNotificationManager::Get().StartProgressNotification(ProgressNotification, TotalWork);
+
+	int32 TaskIndex = Tasks.Num();
+
+	GEditor->GetTimerManager()->SetTimer(
+		TestTask.TimerHandle,
+		FTimerDelegate::CreateLambda
+		(
+			[TaskIndex, TotalWork, IncAmount]()
+			{
+				FTestAsyncTask& CurrentTask = Tasks[TaskIndex];
+				if (CurrentTask.WorkDone >= TotalWork)
+				{
+					GEditor->GetTimerManager()->ClearTimer(CurrentTask.TimerHandle);
+
+				}
+				else
+				{
+					CurrentTask.WorkDone += IncAmount;
+					FSlateNotificationManager::Get().UpdateProgressNotification(CurrentTask.ProgressHandle, CurrentTask.WorkDone);
+				}
+			}
+		),
+		1.0f,
+		true
+	);
+
+	Tasks.Add(MoveTemp(TestTask));
+}
+
+static void TestProgressBars()
+{
+	FString TaskName = FString::Printf(TEXT("Reticulating Splines %d"), Tasks.Num() + 1);
+	int32 TotalWork = 10;
+
+	StartTask(TaskName, TotalWork, 1);
+}
+
+
+FAutoConsoleCommand TestProgressBarsCommand(TEXT("Slate.TestProgressNotification"), TEXT(""), FConsoleCommandDelegate::CreateStatic(&TestProgressBars));
+
+static void TestNotifications()
+{
+	static const float Timeout = 15.0f;
+
+	FTicker::GetCoreTicker().AddTicker(TEXT("TestNotifications"), 0.0f, [](float DeltaTime)
+		{
+			FNotificationInfo NotificationInfo(FText::FromString(TEXT("Test Notification 1")));
+			NotificationInfo.FadeInDuration = 2.0f;
+			NotificationInfo.FadeOutDuration = 2.0f;
+			NotificationInfo.ExpireDuration = Timeout;
+
+			auto Notification = FSlateNotificationManager::Get().AddNotification(NotificationInfo);
+			Notification->ExpireAndFadeout();
+
+			return false;
+		});
+
+	FTicker::GetCoreTicker().AddTicker(TEXT("TestNotifications"), 2.0f, [](float DeltaTime)
+		{
+			FNotificationInfo NotificationInfo(FText::FromString(TEXT("Test Notification 2")));
+			NotificationInfo.FadeInDuration = 2.0f;
+			NotificationInfo.FadeOutDuration = 2.0f;
+			NotificationInfo.ExpireDuration = Timeout;
+
+			NotificationInfo.bUseThrobber = true;
+			auto Notification = FSlateNotificationManager::Get().AddNotification(NotificationInfo);
+			Notification->SetCompletionState(SNotificationItem::CS_Pending);
+			Notification->ExpireAndFadeout();
+
+			return false;
+		});
+
+	FTicker::GetCoreTicker().AddTicker(TEXT("TestNotifications"), 3.0f, [](float DeltaTime)
+		{
+			FNotificationInfo NotificationInfo(FText::FromString(TEXT("Test Notification 3")));
+			NotificationInfo.FadeInDuration = 2.0f;
+			NotificationInfo.FadeOutDuration = 2.0f;
+			NotificationInfo.ExpireDuration = Timeout;
+
+			NotificationInfo.bUseSuccessFailIcons = true;
+			auto Notification = FSlateNotificationManager::Get().AddNotification(NotificationInfo);
+			Notification->SetCompletionState(SNotificationItem::CS_Success);
+			Notification->ExpireAndFadeout();
+
+			return false;
+		});
+
+	FTicker::GetCoreTicker().AddTicker(TEXT("TestNotifications"), 4.0f, [](float DeltaTime)
+		{
+			FNotificationInfo NotificationInfo(FText::FromString(TEXT("Test Notification 5")));
+			NotificationInfo.FadeInDuration = 2.0f;
+			NotificationInfo.FadeOutDuration = 2.0f;
+			NotificationInfo.ExpireDuration = Timeout;
+
+			NotificationInfo.bUseSuccessFailIcons = true;
+			auto Notification = FSlateNotificationManager::Get().AddNotification(NotificationInfo);
+			Notification->SetCompletionState(SNotificationItem::CS_Fail);
+			Notification->ExpireAndFadeout();
+
+			return false;
+		});
+
+	FTicker::GetCoreTicker().AddTicker(TEXT("TestNotifications"), 5.0f, [](float DeltaTime)
+		{
+			FNotificationInfo NotificationInfo(FText::FromString(TEXT("Test Notification 6")));
+			NotificationInfo.FadeInDuration = 2.0f;
+			NotificationInfo.FadeOutDuration = 2.0f;
+			NotificationInfo.ExpireDuration = Timeout;
+
+			NotificationInfo.CheckBoxText = FText::FromString(TEXT("Don't ask again"));
+			NotificationInfo.CheckBoxState = ECheckBoxState::Checked;
+			NotificationInfo.CheckBoxStateChanged = FOnCheckStateChanged::CreateStatic([](ECheckBoxState NewState) {});
+
+			auto Notification = FSlateNotificationManager::Get().AddNotification(NotificationInfo);
+
+			Notification->ExpireAndFadeout();
+
+			return false;
+		});
+
+	FTicker::GetCoreTicker().AddTicker(TEXT("TestNotifications"), 6.0f, [](float DeltaTime)
+		{
+			FNotificationInfo NotificationInfo(FText::FromString(TEXT("Test Notification 7")));
+			NotificationInfo.FadeInDuration = 2.0f;
+			NotificationInfo.FadeOutDuration = 2.0f;
+			NotificationInfo.ExpireDuration = Timeout;
+
+			NotificationInfo.Hyperlink = FSimpleDelegate::CreateLambda([]() {});
+			NotificationInfo.HyperlinkText = FText::FromString(TEXT("This is a hyperlink"));
+
+			auto Notification = FSlateNotificationManager::Get().AddNotification(NotificationInfo);
+
+			Notification->ExpireAndFadeout();
+
+			return false;
+		});
+
+	FTicker::GetCoreTicker().AddTicker(TEXT("TestNotifications"), 7.0f, [](float DeltaTime)
+		{
+			FNotificationInfo NotificationInfo(FText::FromString(TEXT("Test Notification 8")));
+			NotificationInfo.FadeInDuration = 2.0f;
+			NotificationInfo.FadeOutDuration = 2.0f;
+			NotificationInfo.ExpireDuration = Timeout;
+
+			FNotificationButtonInfo Button1(FText::FromString("OK"), FText::GetEmpty(), FSimpleDelegate(), SNotificationItem::ECompletionState::CS_None);
+			FNotificationButtonInfo Button2(FText::FromString("CANCEL"), FText::GetEmpty(), FSimpleDelegate(), SNotificationItem::ECompletionState::CS_None);
+
+			NotificationInfo.ButtonDetails.Add(Button1);
+			NotificationInfo.ButtonDetails.Add(Button2);
+
+			auto Notification = FSlateNotificationManager::Get().AddNotification(NotificationInfo);
+			Notification->ExpireAndFadeout();
+
+			return false;
+		});
+
+	FTicker::GetCoreTicker().AddTicker(TEXT("TestNotifications"), 8.0f, [](float DeltaTime)
+		{
+			FNotificationInfo NotificationInfo(FText::FromString(TEXT("Everything Under The Sun. This one also has a lot of text which should wrap to the next line")));
+			NotificationInfo.FadeInDuration = 2.0f;
+			NotificationInfo.FadeOutDuration = 2.0f;
+			NotificationInfo.ExpireDuration = Timeout;
+
+			NotificationInfo.CheckBoxText = FText::FromString(TEXT("Don't ask again"));
+			NotificationInfo.CheckBoxState = ECheckBoxState::Checked;
+			NotificationInfo.CheckBoxStateChanged = FOnCheckStateChanged::CreateStatic([](ECheckBoxState NewState) {});
+
+			NotificationInfo.Hyperlink = FSimpleDelegate::CreateLambda([]() {});
+			NotificationInfo.HyperlinkText = FText::FromString(TEXT("This is a hyperlink"));
+
+			NotificationInfo.bUseSuccessFailIcons = true;
+			NotificationInfo.bUseThrobber = true;
+
+
+			FNotificationButtonInfo Button1(FText::FromString("OK"), FText::GetEmpty(), FSimpleDelegate());
+			FNotificationButtonInfo Button2(FText::FromString("CANCEL"), FText::GetEmpty(), FSimpleDelegate());
+
+			NotificationInfo.ButtonDetails.Add(Button1);
+			NotificationInfo.ButtonDetails.Add(Button2);
+
+
+			auto Notification = FSlateNotificationManager::Get().AddNotification(NotificationInfo);
+			Notification->SetCompletionState(SNotificationItem::CS_Pending);
+
+			Notification->ExpireAndFadeout();
+
+			return false;
+		});
+}
+
+FAutoConsoleCommand TestNotificationCommand(TEXT("Slate.TestNotifications"), TEXT(""), FConsoleCommandDelegate::CreateStatic(&TestNotifications));
+
+#endif // (WITH_EDITOR || IS_PROGRAM) && !UE_BUILD_SHIPPING
diff --git a/Engine/Source/Editor/StatusBar/Private/SStatusBar.cpp b/Engine/Source/Editor/StatusBar/Private/SStatusBar.cpp
--- a/Engine/Source/Editor/StatusBar/Private/SStatusBar.cpp
+++ b/Engine/Source/Editor/StatusBar/Private/SStatusBar.cpp
@@ -14,6 +14,7 @@
 #include "Widgets/Input/SMultiLineEditableTextBox.h"
 #include "Widgets/Layout/SSpacer.h"
 #include "Widgets/Docking/SDockTab.h"
+#include "Widgets/Layout/SSeparator.h"
 #include "TimerManager.h"
 #include "Framework/Commands/Commands.h"
 #include "ToolMenuContext.h"
@@ -24,9 +25,23 @@
 #include "Widgets/Layout/SSplitter.h"
 #include "InputCoreTypes.h"
 #include "Misc/ConfigCacheIni.h"
+#include "Widgets/Notifications/SProgressBar.h"
+#include "Widgets/Notifications/SNotificationList.h"
+#include "Framework/Notifications/NotificationManager.h"
+#include "Widgets/Notifications/INotificationWidget.h"
 
 #define LOCTEXT_NAMESPACE "StatusBar"
 
+namespace StatusBarNotificationConstants
+{
+	// How long progress notification toasts should appear for
+	const float NotificationExpireTime = 5.0f;
+
+	const float NotificationFadeDuration = .15f;
+
+	// Delay before a progress notification becomes visible. This is to avoid the status bar to animate and flicker from short lived notifications. 
+	const double NotificationDelay = 0.5;
+}
 
 class SDrawerOverlay : public SCompoundWidget
 {
@@ -67,7 +82,7 @@ class SDrawerOverlay : public SCompoundWidget
 
 		OnTargetHeightChanged = InArgs._OnTargetHeightChanged;
 
-		BackgroundBrush = FAppStyle::Get().GetBrush("StatusBar.Background");
+		BackgroundBrush = FAppStyle::Get().GetBrush("Brushes.Panel");
 		ShadowBrush = FAppStyle::Get().GetBrush("StatusBar.ContentBrowserShadow");
 
 		bIsResizeHandleHovered = false;
@@ -255,12 +270,235 @@ private:
 	float CurrentHeight;
 	float MinHeight;
 	float MaxHeight;
-	float TargetHeight;
+	float TargetHeight; 
 	float InitialHeightAtResize;
 	bool bIsResizing;
 	bool bIsResizeHandleHovered;
 };
 
+class SStatusBarProgressWidget : public SCompoundWidget, public INotificationWidget
+{
+public:
+	SLATE_BEGIN_ARGS(SStatusBarProgressWidget)
+	{}
+		SLATE_ATTRIBUTE(const FStatusBarProgress*, StatusBarProgress)
+	SLATE_END_ARGS()
+	
+public:
+	void Construct(const FArguments& InArgs, bool bIsShownInNotification = false)
+	{
+		StatusBarProgress = InArgs._StatusBarProgress;
+
+		ChildSlot
+		[
+			SNew(SHorizontalBox)
+			+SHorizontalBox::Slot()
+			.Padding(0.0f, 0.0f, 4.0f, 0.0f)
+			[
+				SNew(SVerticalBox)
+				+ SVerticalBox::Slot()
+				.AutoHeight()
+				.Padding(0.0f, 3.0f, 0.0f, 2.0f)
+				[
+					SAssignNew(ProgressTextWidget, STextBlock)
+				]
+
+				+ SVerticalBox::Slot()
+				//.AutoWidth()
+				[
+					SNew(SBox)
+					.HeightOverride(8)
+					[
+						SNew(SOverlay)
+						+ SOverlay::Slot()
+						.VAlign(VAlign_Center)
+						.Padding(1.0f, 0.0f)
+						[
+							SAssignNew(ProgressBar, SProgressBar)
+							.Percent(0.0f)
+						]
+						+ SOverlay::Slot()
+						[
+							SNew(SImage)
+							.Image(FAppStyle::Get().GetBrush("StatusBar.ProgressOverlay"))
+							.Visibility(EVisibility::HitTestInvisible)
+						]
+					]
+				]
+			]
+			+ SHorizontalBox::Slot()
+			.AutoWidth()
+			.HAlign(HAlign_Right)
+			.VAlign(VAlign_Center)
+			[
+				SAssignNew(PercentText, STextBlock)
+			]
+		];
+	
+		if (bIsShownInNotification)
+		{
+			ProgressTextWidget->SetFont(FAppStyle::Get().GetFontStyle("NotificationList.FontBold"));
+		}
+	}
+
+	virtual void Tick(const FGeometry& AllottedGeometry, const double InCurrentTime, const float InDeltaTime) override
+	{
+		const FStatusBarProgress* ProgressData = StatusBarProgress.Get();
+		if (ProgressData)
+		{
+			float PercentDone = FMath::Clamp((float)ProgressData->TotalWorkDone / ProgressData->TotalWorkToDo, 0.0f, 1.0f);
+			ProgressTextWidget->SetText(ProgressData->DisplayText);
+			PercentText->SetText(FText::AsPercent(PercentDone));
+			ProgressBar->SetPercent(PercentDone);
+		}
+		else if(StatusBarProgress.IsBound())
+		{
+			StatusBarProgress = nullptr;
+			FText CurrentText = ProgressTextWidget->GetText();
+			ProgressTextWidget->SetText(FText::Format(LOCTEXT("CancelledProgressText", "{0} (Canceled)"), CurrentText));
+		}
+	}
+
+	void SetProgressText(FText ProgressText)
+	{
+		ProgressTextWidget->SetText(ProgressText);
+	}
+
+	void SetProgressPercent(float Percent)
+	{
+		ProgressBar->SetPercent(Percent);
+		PercentText->SetText(FText::AsPercent(Percent));
+	}
+	
+	/** INotificationWidget interface */
+	virtual void OnSetCompletionState(SNotificationItem::ECompletionState) override
+	{}
+
+	TSharedRef<SWidget> AsWidget() override
+	{
+		return AsShared();
+	}
+
+	bool UseNotificationBackground() const override { return true; }
+private:
+	TAttribute<const FStatusBarProgress*> StatusBarProgress;
+	TSharedPtr<SProgressBar> ProgressBar;
+	TSharedPtr<STextBlock> PercentText;
+	TSharedPtr<STextBlock> ProgressTextWidget;
+};
+
+class SStatusBarProgressArea : public SCompoundWidget
+{
+	SLATE_BEGIN_ARGS(SStatusBarProgressArea)
+	{}
+		SLATE_EVENT(FOnGetContent, OnGetProgressMenuContent)
+
+	SLATE_END_ARGS()
+
+	void SetPercent(float Percent)
+	{
+		MainProgressWidget->SetProgressPercent(Percent);
+		MainProgressWidget->SetProgressText(FText::AsPercent(Percent));
+	}
+
+	void SetProgressText(FText ProgressText)
+	{
+		MainProgressWidget->SetProgressText(ProgressText);
+	}
+
+public:
+	void Construct(const FArguments& InArgs)
+	{
+		OpenCloseEasingCurve = FCurveSequence(0.0f, 0.15f, ECurveEaseFunction::QuadOut);
+
+		SetVisibility(EVisibility::Collapsed);
+
+		ChildSlot
+		[
+			SAssignNew(Box, SBox)
+			.WidthOverride(300.0f)
+			.Padding(FMargin(4.0f,0.0f))
+			[
+				SAssignNew(ProgressCombo, SComboButton)
+				.MenuPlacement(MenuPlacement_AboveAnchor)
+				.ComboButtonStyle(FAppStyle::Get(), "SimpleComboButton")
+				.OnGetMenuContent(InArgs._OnGetProgressMenuContent)
+				.ButtonContent()
+				[
+					SAssignNew(MainProgressWidget, SStatusBarProgressWidget)
+				]
+			]
+		];
+	}
+
+	void OpenProgressBar()
+	{
+		if(!GetVisibility().IsVisible())
+		{
+			SetVisibility(EVisibility::Visible);
+
+			if (!OpenCloseEasingCurve.IsPlaying())
+			{
+				OpenCloseEasingCurve.Play(SharedThis(this), false, 0.0f, false);
+
+				if (!OpenCloseTimer.IsValid())
+				{
+					AnimationThrottle = FSlateThrottleManager::Get().EnterResponsiveMode();
+					OpenCloseTimer = RegisterActiveTimer(0.0f, FWidgetActiveTimerDelegate::CreateSP(SharedThis(this), &SStatusBarProgressArea::UpdateProgressAnimation));
+				}
+			}
+		}
+	}
+
+	void DismissProgressBar()
+	{
+		if (GetVisibility().IsVisible())
+		{
+			if (OpenCloseEasingCurve.IsForward())
+			{
+				OpenCloseEasingCurve.Reverse();
+			}
+
+			if (!OpenCloseTimer.IsValid())
+			{
+				AnimationThrottle = FSlateThrottleManager::Get().EnterResponsiveMode();
+				OpenCloseTimer = RegisterActiveTimer(0.0f, FWidgetActiveTimerDelegate::CreateSP(this, &SStatusBarProgressArea::UpdateProgressAnimation));
+			}
+
+			ProgressCombo->SetIsOpen(false);
+		}
+	}
+
+	virtual FVector2D ComputeDesiredSize(float Scale) const
+	{
+		return SCompoundWidget::ComputeDesiredSize(Scale) * FVector2D(OpenCloseEasingCurve.GetLerp(), 1.0f);
+	}
+private:
+	EActiveTimerReturnType UpdateProgressAnimation(double CurrentTime, float DeltaTime)
+	{
+		if (!OpenCloseEasingCurve.IsPlaying())
+		{
+			if (OpenCloseEasingCurve.IsAtStart())
+			{
+				SetVisibility(EVisibility::Collapsed);
+			}
+
+			FSlateThrottleManager::Get().LeaveResponsiveMode(AnimationThrottle);
+			OpenCloseTimer.Reset();
+			return EActiveTimerReturnType::Stop;
+		}
+
+		return EActiveTimerReturnType::Continue;
+	}
+
+private:
+	TSharedPtr<SBox> Box;
+	TSharedPtr<SStatusBarProgressWidget> MainProgressWidget;
+	TSharedPtr<SComboButton> ProgressCombo;
+	FCurveSequence OpenCloseEasingCurve;
+	FThrottleRequest AnimationThrottle;
+	TSharedPtr<FActiveTimerHandle> OpenCloseTimer;
+};
 
 SStatusBar::~SStatusBar()
 {
@@ -278,7 +516,7 @@ void SStatusBar::Construct(const FArguments& InArgs, FName InStatusBarName, cons
 	UpArrow = FAppStyle::Get().GetBrush("StatusBar.ContentBrowserUp");
 	DownArrow = FAppStyle::Get().GetBrush("StatusBar.ContentBrowserDown");
 
-	const FSlateBrush* StatusBarBackground = FAppStyle::Get().GetBrush("StatusBar.Background");
+	const FSlateBrush* StatusBarBackground = FAppStyle::Get().GetBrush("Brushes.Panel");
 
 	DrawerEasingCurve = FCurveSequence(0.0f, 0.15f, ECurveEaseFunction::QuadOut);
 
@@ -305,7 +543,6 @@ void SStatusBar::Construct(const FArguments& InArgs, FName InStatusBarName, cons
 			    .Padding(1.0f, 0.0f)
 			    [
 				    SNew(SBorder)
-				    .Padding(0.0f)
 				    .BorderImage(StatusBarBackground)
 				    .VAlign(VAlign_Center)
 				    .Padding(FMargin(6.0f, 0.0f))
@@ -318,7 +555,6 @@ void SStatusBar::Construct(const FArguments& InArgs, FName InStatusBarName, cons
 				.Padding(1.0f, 0.0f)
 				[
 					SNew(SBorder)
-					.Padding(0.0f)
 					.BorderImage(StatusBarBackground)
 					.VAlign(VAlign_Center)
 					.Padding(FMargin(6.0f, 0.0f))
@@ -339,17 +575,31 @@ void SStatusBar::Construct(const FArguments& InArgs, FName InStatusBarName, cons
 						MakeStatusBarToolBarWidget()
 					]
 				]
+				+ SHorizontalBox::Slot()
+				.HAlign(HAlign_Right)
+				.AutoWidth()
+				.Padding(1.0f, 0.0f)
+				[
+					SNew(SBorder)
+					.Padding(0.0f)
+					.BorderImage(StatusBarBackground)
+					.VAlign(VAlign_Center)
+					.Padding(FMargin(6.0f, 0.0f))
+					[
+						MakeProgressBar()
+					]
+				]
 			]
 		]
 	];
 }
 
-void SStatusBar::PushMessage(FStatusBarMessageHandle& InHandle, const TAttribute<FText>& InMessage, const TAttribute<FText>& InHintText)
+void SStatusBar::PushMessage(FStatusBarMessageHandle InHandle, const TAttribute<FText>& InMessage, const TAttribute<FText>& InHintText)
 {
 	MessageStack.Emplace(InMessage, InHintText, InHandle);
 }
 
-void SStatusBar::PopMessage(FStatusBarMessageHandle& InHandle)
+void SStatusBar::PopMessage(FStatusBarMessageHandle InHandle)
 {
 	if (InHandle.IsValid() && MessageStack.Num() > 0)
 	{
@@ -365,6 +615,68 @@ void SStatusBar::ClearAllMessages()
 	MessageStack.Empty();
 }
 
+void SStatusBar::StartProgressNotification(FProgressNotificationHandle InHandle, FText DisplayText, int32 TotalWorkToDo)
+{
+	if (!FindProgressNotification(InHandle))
+	{
+		if(TotalWorkToDo > 0)
+		{
+			ProgressNotifications.Emplace(DisplayText, FPlatformTime::Seconds(), InHandle, TotalWorkToDo);
+
+			// If a notification was already active, refresh its fadeout time 
+			if (TSharedPtr<SNotificationItem> ActiveProgressNotificationPin = ActiveProgressNotification.Pin())
+			{
+				ActiveProgressNotificationPin->SetExpireDuration(StatusBarNotificationConstants::NotificationExpireTime);
+				ActiveProgressNotificationPin->ExpireAndFadeout();
+			}
+			else
+			{
+				bAllowedToRefreshProgressNotification = true;
+			}
+			UpdateProgressStatus();
+		}
+		else
+		{
+			//UE_LOG(LogUnrealEdEngine, Log, TEXT("Progress notification \"%s\" has no work to do so it will not be displayed"), *DisplayText.ToString()))
+		}
+	}
+}
+
+bool SStatusBar::UpdateProgressNotification(FProgressNotificationHandle InHandle, int32 TotalWorkDone, int32 UpdatedTotalWorkToDo, FText UpdatedDisplayText)
+{
+	if (FStatusBarProgress* Progress = FindProgressNotification(InHandle))
+	{
+		Progress->TotalWorkDone = FMath::Clamp(TotalWorkDone, 0, Progress->TotalWorkToDo);
+		if (!UpdatedDisplayText.IsEmpty())
+		{
+			Progress->DisplayText = UpdatedDisplayText;
+		}
+
+		if (UpdatedTotalWorkToDo != 0)
+		{
+			Progress->TotalWorkToDo = UpdatedTotalWorkToDo;
+		}
+
+		UpdateProgressStatus();
+
+		return true;
+	}
+
+	return false;
+}
+
+bool SStatusBar::CancelProgressNotification(FProgressNotificationHandle InHandle)
+{
+	if (ProgressNotifications.RemoveAll([InHandle](const FStatusBarProgress& Progress){ return Progress.Handle == InHandle;}) != 0)
+	{
+		UpdateProgressStatus();
+
+		return true;
+	}
+
+	return false;
+}
+
 EVisibility SStatusBar::GetHelpIconVisibility() const
 {
 	if (MessageStack.Num() > 0)
@@ -573,6 +885,13 @@ TSharedRef<SWidget> SStatusBar::MakeStatusMessageWidget()
 }
 
 
+TSharedRef<SWidget> SStatusBar::MakeProgressBar()
+{
+	return 
+		SAssignNew(ProgressBar, SStatusBarProgressArea)
+		.OnGetProgressMenuContent(this, &SStatusBar::OnGetProgressBarMenuContent);
+}
+
 bool SStatusBar::IsDrawerOpened(const FName DrawerId) const
 {
 	return OpenedDrawerData.Key == DrawerId ? true : false;
@@ -655,6 +974,157 @@ void SStatusBar::RegisterSourceControlStatus()
 	}
 }
 
+FStatusBarProgress* SStatusBar::FindProgressNotification(FProgressNotificationHandle InHandle)
+{
+	return ProgressNotifications.FindByPredicate(
+		[InHandle](const FStatusBarProgress& Progress)
+		{
+			return Progress.Handle == InHandle;
+		});
+}
+
+void SStatusBar::UpdateProgressStatus()
+{
+	int32 NumIncompleteTasks = 0;
+
+	if (ProgressNotifications.Num())
+	{
+		int32 TotalWorkToDo = 0;
+		int32 CurrentWorkDone = 0;
+
+		bool bShouldAnyProgressBeVisible = false;
+
+		double CurrentTime = FPlatformTime::Seconds();
+
+		const FStatusBarProgress* LastIncompleteTask = &ProgressNotifications.Last();
+		for (const FStatusBarProgress& Progress : ProgressNotifications)
+		{
+			TotalWorkToDo += Progress.TotalWorkToDo;
+			CurrentWorkDone += Progress.TotalWorkDone;
+
+			bShouldAnyProgressBeVisible |= ((CurrentTime - Progress.StartTime) >= StatusBarNotificationConstants::NotificationDelay);
+
+			if (Progress.TotalWorkToDo > Progress.TotalWorkDone)
+			{
+				++NumIncompleteTasks;
+				LastIncompleteTask = &Progress;
+			}
+		}
+
+		// Just assume 100% of the work is done if there is no work to do. The progress bar will dismiss in this case but we want to show 100% while its dismissing
+		const float Percent = TotalWorkToDo > 0 ? (float)CurrentWorkDone / TotalWorkToDo : 1.0f;
+		FText StatusBarProgressText;
+		if (NumIncompleteTasks > 1)
+		{
+			StatusBarProgressText = FText::Format(LOCTEXT("ProgressBarLabel", "{0} (+{1} more)"), LastIncompleteTask->DisplayText, FText::AsNumber(NumIncompleteTasks - 1));
+		}
+		else
+		{
+			StatusBarProgressText = LastIncompleteTask->DisplayText;
+		}
+
+		bShouldAnyProgressBeVisible &= (NumIncompleteTasks > 0);
+		
+		TSharedPtr<SStatusBarProgressWidget> ActiveNotificationProgressWidgetPin = ActiveNotificationProgressWidget.Pin();
+
+		if(bShouldAnyProgressBeVisible)
+		{
+			TSharedPtr<SNotificationItem> ActiveProgressNotificationPin = ActiveProgressNotification.Pin();
+		
+			// Show a new notification the first time a new progress task is started assuming we don't already have a notification open
+			if (!ActiveProgressNotificationPin.IsValid() && bAllowedToRefreshProgressNotification)
+			{
+				FNotificationInfo ProgressNotification(FText::GetEmpty());
+
+				ActiveNotificationProgressWidgetPin = SNew(SStatusBarProgressWidget, true);
+				ProgressNotification.ContentWidget = ActiveNotificationProgressWidgetPin.ToSharedRef();
+				ProgressNotification.FadeOutDuration = StatusBarNotificationConstants::NotificationFadeDuration;
+				ProgressNotification.ForWindow = FSlateApplication::Get().FindWidgetWindow(AsShared());
+
+				ActiveProgressNotificationPin = FSlateNotificationManager::Get().AddNotification(ProgressNotification);
+
+				ActiveProgressNotification = ActiveProgressNotificationPin;
+				ActiveNotificationProgressWidget = ActiveNotificationProgressWidgetPin;
+
+				ActiveProgressNotificationPin->SetExpireDuration(StatusBarNotificationConstants::NotificationExpireTime);
+				ActiveProgressNotificationPin->ExpireAndFadeout();
+
+				// Do not show the notification again unless a new task is started.
+				bAllowedToRefreshProgressNotification = false;
+			}
+
+			OpenProgressBar();
+
+			ProgressBar->SetPercent(Percent);
+
+			ProgressBar->SetProgressText(StatusBarProgressText);
+		}
+
+		// Update the notification if it exists. Even if no progress should be visible, if the notification is visible we want to update it while it fades out
+		if (ActiveNotificationProgressWidgetPin)
+		{
+			ActiveNotificationProgressWidgetPin->SetProgressText(StatusBarProgressText);
+			ActiveNotificationProgressWidgetPin->SetProgressPercent(Percent);
+		}
+	}
+
+	if (NumIncompleteTasks == 0)
+	{
+		DismissProgressBar();
+	}
+}
+
+void SStatusBar::OpenProgressBar()
+{
+	ProgressBar->OpenProgressBar();
+}
+
+void SStatusBar::DismissProgressBar()
+{
+	ProgressBar->DismissProgressBar();
+	bAllowedToRefreshProgressNotification = false;
+	ProgressNotifications.Empty();
+}
+
+TSharedRef<SWidget> SStatusBar::OnGetProgressBarMenuContent()
+{
+	FMenuBuilder ProgressBarMenu(false, nullptr);
+
+	const float StatusBarHeight = FAppStyle::Get().GetFloat("StatusBar.Height");
+
+	for (int32 ProgressIndex = 0; ProgressIndex < ProgressNotifications.Num(); ++ProgressIndex)
+	{
+		FStatusBarProgress& Progress = ProgressNotifications[ProgressIndex];
+
+		FProgressNotificationHandle Handle = Progress.Handle;
+
+		const bool bLastProgressBar = (ProgressIndex + 1 == ProgressNotifications.Num());
+
+		TSharedRef<SWidget> MenuWidget =
+			SNew(SBox)
+			.Padding(FMargin(8.0f, ProgressIndex == 0 ? 0.0f : 4.0f, 8.0f, bLastProgressBar ? 0.0f : 8.0f))
+			[
+				SNew(SStatusBarProgressWidget)
+				.StatusBarProgress_Lambda([this, Handle]() { return FindProgressNotification(Handle); })
+			];
+
+
+		ProgressBarMenu.AddWidget(MenuWidget, FText::GetEmpty(), false, false);
+
+		if (!bLastProgressBar)
+		{
+			ProgressBarMenu.AddWidget(SNew(SSeparator).Thickness(1.0f), FText::GetEmpty(), false, false);
+		}
+	}
+
+	return
+		SNew(SBox)
+		.WidthOverride(ProgressBar->GetDesiredSize().X-8.0f)
+		[
+			ProgressBarMenu.MakeWidget()
+		];
+}
+
 void SStatusBar::RegisterDrawer(FStatusBarDrawer&& Drawer)
 {
 	const int32 NumDrawers = RegisteredDrawers.Num();
@@ -662,7 +1132,7 @@ void SStatusBar::RegisterDrawer(FStatusBarDrawer&& Drawer)
 
 	if (RegisteredDrawers.Num() > NumDrawers)
 	{
-		const FSlateBrush* StatusBarBackground = FAppStyle::Get().GetBrush("StatusBar.Background");
+		const FSlateBrush* StatusBarBackground = FAppStyle::Get().GetBrush("Brushes.Panel");
 
 		DrawerBox->AddSlot()
 		.Padding(1.0f, 0.0f)
@@ -676,35 +1146,6 @@ void SStatusBar::RegisterDrawer(FStatusBarDrawer&& Drawer)
 				MakeStatusBarDrawerButton(Drawer)
 			]
 		];
-/*
-			/ *+ SHorizontalBox::Slot()
-				.Padding(1.0f, 0.0f)
-				.AutoWidth()
-				[
-					SNew(SBorder)
-					.Padding(FMargin(2.0f, 0.0f))
-					.BorderImage(StatusBarBackground)
-					.VAlign(VAlign_Center)
-					[
-						MakeStatusBarDrawerButton(ContentBrowserDrawerButton)
-					]
-				]
-				+ SHorizontalBox::Slot()
-				.AutoWidth()
-				.Padding(1.0f, 0.0f)
-				[
-					SNew(SBorder)
-					.Padding(0.0f)
-					.BorderImage(StatusBarBackground)
-					.VAlign(VAlign_Center)
-					.Padding(FMargin(6.0f, 0.0f))
-					[
-						MakeDebugConsoleWidget(InArgs._OnConsoleClosed)
-					]
-				]* /
-		}
-*/
-
 	}
 }
 
diff --git a/Engine/Source/Editor/StatusBar/Private/SStatusBar.h b/Engine/Source/Editor/StatusBar/Private/SStatusBar.h
new file mode 100644
--- /dev/null
+++ b/Engine/Source/Editor/StatusBar/Private/SStatusBar.h
@@ -0,0 +1,238 @@
+// Copyright Epic Games, Inc. All Rights Reserved.
+
+#pragma once
+
+#include "StatusBarSubsystem.h"
+#include "Widgets/SCompoundWidget.h"
+#include "Widgets/DeclarativeSyntaxSupport.h"
+#include "Framework/SlateDelegates.h"
+#include "Animation/CurveSequence.h"
+#include "Application/ThrottleManager.h"
+
+
+DECLARE_DELEGATE_OneParam(FOnStatusBarDrawerOpened, TSharedRef<SStatusBar>&)
+DECLARE_DELEGATE_OneParam(FOnStatusBarDrawerDismissed, const TSharedPtr<SWidget>&)
+DECLARE_DELEGATE_OneParam(FOnStatusBarDrawerTargetHeightChanged, float);
+
+class SWindow;
+class SDockTab;
+class SWidget;
+class SMultiLineEditableTextBox;
+class SDrawerOverlay;
+class SHorizontalBox;
+class SStatusBarProgressArea;
+class SStatusBarProgressWidget;
+
+namespace StatusBarDrawerIds
+{
+	const FName ContentBrowser("ContentBrowser");
+}
+
+/** Data payload for messages in the status bar */
+struct FStatusBarMessage
+{
+	FStatusBarMessage(const TAttribute<FText>& InMessageText, const TAttribute<FText>& InHintText, FStatusBarMessageHandle InHandle)
+		: MessageText(InMessageText)
+		, HintText(InHintText)
+		, Handle(InHandle)
+	{}
+
+	TAttribute<FText> MessageText;
+	TAttribute<FText> HintText;
+	FStatusBarMessageHandle Handle;
+};
+
+/** Data payload for progress bars in the status bar */
+struct FStatusBarProgress
+{
+	FStatusBarProgress(FText InDisplayText, double InStartTime, FProgressNotificationHandle InHandle, int32 InTotalWorkToDo)
+		: DisplayText(InDisplayText)
+		, StartTime(InStartTime)
+		, TotalWorkToDo(InTotalWorkToDo)
+		, TotalWorkDone(0)
+		, Handle(InHandle)
+	{}
+
+	FText DisplayText;
+	double StartTime;
+	int32 TotalWorkToDo;
+	int32 TotalWorkDone;
+	FProgressNotificationHandle Handle;
+};
+
+struct FStatusBarDrawer
+{
+	FStatusBarDrawer(FName InUniqueId)
+		: UniqueId(InUniqueId)
+	{}
+
+	FName UniqueId;
+	FOnGetContent GetDrawerContentDelegate;
+	FOnStatusBarDrawerOpened OnDrawerOpenedDelegate;
+	FOnStatusBarDrawerDismissed OnDrawerDismissedDelegate;
+
+	FText ButtonText;
+	FText ToolTipText;
+	const FSlateBrush* Icon;
+
+	bool operator==(const FName& OtherId) const
+	{
+		return UniqueId == OtherId;
+	}
+
+	bool operator==(const FStatusBarDrawer& Other) const
+	{
+		return UniqueId == Other.UniqueId;
+	}
+};
+
+class SStatusBar : public SCompoundWidget
+{
+	SLATE_BEGIN_ARGS(SStatusBar)	
+	{}
+		SLATE_EVENT(FSimpleDelegate, OnConsoleClosed)
+	SLATE_END_ARGS()
+
+public:
+	~SStatusBar();
+
+	/** SWidget interface */
+	virtual bool SupportsKeyboardFocus() const { return false; }
+	void Construct(const FArguments& InArgs, FName InStatusBarName, const TSharedRef<SDockTab> InParentTab);
+
+	/** 
+	 * Pushes a new status bar message
+	 *
+	 * @param InHandle	A handle to the message for clearing it later
+	 * @param InMessage	The message to display
+	 * @param InHintText	Optional hint text message.  This message will be highlighted to make it stand out
+	 */
+	void PushMessage(FStatusBarMessageHandle InHandle, const TAttribute<FText>& InMessage, const TAttribute<FText>& InHintText);
+
+	/**
+	 * Removes a message from the status bar.  When messages are removed the previous message on the stack (if any) is displayed
+	 *
+	 * @param InHandle	Handle to the status bar message to remove
+	 */
+	void PopMessage(FStatusBarMessageHandle InHandle);
+
+	/**
+	 * Removes all messages from the status bar
+	 */
+	void ClearAllMessages();
+
+	/**
+	 * Called when a progress notification begins
+	 * @param Handle			Handle to the notification
+	 * @param DisplayText		Display text used to describe the type of work to the user
+	 * @param TotalWorkToDo		Arbitrary number of work units to perform.
+	 */
+	void StartProgressNotification(FProgressNotificationHandle InHandle, FText DisplayText, int32 TotalWorkToDo);
+
+	/**
+	 * Called when a notification should be updated.
+	 * @param InHandle				Handle to the notification that was previously created with StartProgressNotification
+	 * @param TotalWorkDone			The total number of work units done for the notification.
+	 * @param UpdatedTotalWorkToDo	UpdatedTotalWorkToDo. This value will be 0 if the total work did not change
+	 * @param UpdatedDisplayText	Updated display text of the notification. This value will be empty if the text did not change
+	 */
+	bool UpdateProgressNotification(FProgressNotificationHandle InHandle, int32 TotalWorkDone, int32 UpdatedTotalWorkToDo, FText UpdatedDisplayText);
+
+	/**
+	 * Called when a notification should be canceled
+	 */
+	bool CancelProgressNotification(FProgressNotificationHandle InHandle);
+
+	TSharedPtr<SDockTab> GetParentTab() const;
+	bool FocusDebugConsole();
+	bool IsDebugConsoleFocused() const;
+
+	void RegisterDrawer(FStatusBarDrawer&& Drawer);
+
+	void OpenDrawer(const FName DrawerId);
+
+	/**
+	 * Dismisses an open drawer with an animation.  The drawer contents are removed once the animation is complete
+	 */
+	void DismissDrawer(const TSharedPtr<SWidget>& NewlyFocusedWidget);
+
+	void CloseDrawerImmediately();
+
+	bool IsDrawerOpened(const FName DrawerId) const;
+private:
+	/** Called when global focus changes which is used to determine if we should close an opened content browser drawer */
+	void OnGlobalFocusChanging(const FFocusEvent& FocusEvent, const FWeakWidgetPath& OldFocusedWidgetPath, const TSharedPtr<SWidget>& OldFocusedWidget, const FWidgetPath& NewFocusedWidgetPath, const TSharedPtr<SWidget>& NewFocusedWidget);
+
+	/** Called when active tab changes which is used to determine if we should close an opened content browser drawer */
+	void OnActiveTabChanged(TSharedPtr<SDockTab> PreviouslyActive, TSharedPtr<SDockTab> NewlyActivated);
+
+	void OnDrawerHeightChanged(float TargetHeight);
+
+	EVisibility GetHelpIconVisibility() const;
+
+	bool GetContentBrowserExpanded() const;
+
+	FText GetStatusBarMessage() const;
+
+	TSharedRef<SWidget> MakeStatusBarDrawerButton(const FStatusBarDrawer& Drawer);
+
+	TSharedRef<SWidget> MakeStatusBarToolBarWidget();
+	TSharedRef<SWidget> MakeDebugConsoleWidget(FSimpleDelegate OnConsoleClosed);
+	TSharedRef<SWidget> MakeStatusMessageWidget();
+	TSharedRef<SWidget> MakeProgressBar();
+
+	void MakeDrawerButtons(TSharedRef<SHorizontalBox> DrawerBox);
+
+	FReply OnDrawerButtonClicked(const FName DrawerId);
+
+	void CreateAnimationTimerIfNeeded();
+
+	EActiveTimerReturnType UpdateDrawerAnimation(double CurrentTime, float DeltaTime);
+
+	void RegisterStatusBarMenu();
+
+	void RegisterSourceControlStatus();
+
+	FStatusBarProgress* FindProgressNotification(FProgressNotificationHandle InHandle);
+
+	void UpdateProgressStatus();
+	void OpenProgressBar();
+	void DismissProgressBar();
+	
+	TSharedRef<SWidget> OnGetProgressBarMenuContent();
+
+private:
+	TArray<FStatusBarMessage> MessageStack;
+	TArray<FStatusBarProgress> ProgressNotifications;
+
+	TSharedPtr<SMultiLineEditableTextBox> ConsoleEditBox;
+	TWeakPtr<SDockTab> ParentTab;
+
+	TArray<FStatusBarDrawer> RegisteredDrawers;
+
+	TWeakPtr<SWindow> WindowWithOverlayContent;
+
+	TPair<FName,TSharedPtr<SDrawerOverlay>> OpenedDrawerData;
+
+	TSharedPtr<SHorizontalBox> DrawerBox;
+
+	TSharedPtr<FActiveTimerHandle> DrawerOpenCloseTimer;
+
+	TSharedPtr<SStatusBarProgressArea> ProgressBar;
+
+	TWeakPtr<SNotificationItem> ActiveProgressNotification;
+
+	TWeakPtr<SStatusBarProgressWidget> ActiveNotificationProgressWidget;
+
+
+	FCurveSequence DrawerEasingCurve;
+	FThrottleRequest AnimationThrottle;
+	const FSlateBrush* UpArrow = nullptr;
+	const FSlateBrush* DownArrow = nullptr;
+	FName StatusBarName;
+	FName StatusBarToolBarName;
+
+	bool bAllowedToRefreshProgressNotification = false;
+
+};
+
diff --git a/Engine/Source/Editor/StatusBar/Private/StatusBarSubsystem.cpp b/Engine/Source/Editor/StatusBar/Private/StatusBarSubsystem.cpp
--- a/Engine/Source/Editor/StatusBar/Private/StatusBarSubsystem.cpp
+++ b/Engine/Source/Editor/StatusBar/Private/StatusBarSubsystem.cpp
@@ -12,21 +12,25 @@
 #include "Widgets/Docking/SDockTab.h"
 #include "Editor.h"
 #include "Toolkits/GlobalEditorCommonCommands.h"
-
+#include "Widgets/Notifications/SNotificationList.h"
+#include "Framework/Notifications/NotificationManager.h"
 
 #define LOCTEXT_NAMESPACE "StatusBar"
 
-
-int32 UStatusBarSubsystem::HandleCounter = 0;
+int32 UStatusBarSubsystem::MessageHandleCounter = 0;
 
 void UStatusBarSubsystem::Initialize(FSubsystemCollectionBase& Collection)
 {
 	FSourceControlCommands::Register();
+
+	FSlateNotificationManager::Get().SetProgressNotificationHandler(this);
 }
 
 void UStatusBarSubsystem::Deinitialize()
 {
 	FSourceControlCommands::Unregister();
+
+	FSlateNotificationManager::Get().SetProgressNotificationHandler(nullptr);
 }
 
 bool UStatusBarSubsystem::FocusDebugConsole(TSharedRef<SWindow> ParentWindow)
@@ -139,7 +143,7 @@ FStatusBarMessageHandle UStatusBarSubsystem::PushStatusBarMessage(FName StatusBa
 {
 	if (TSharedPtr<SStatusBar> StatusBar = GetStatusBar(StatusBarName))
 	{
-		FStatusBarMessageHandle NewHandle(++HandleCounter);
+		FStatusBarMessageHandle NewHandle(++MessageHandleCounter);
 
 		StatusBar->PushMessage(NewHandle, InMessage, InHintText);
 
@@ -170,6 +174,52 @@ void UStatusBarSubsystem::ClearStatusBarMessages(FName StatusBarName)
 	}
 }
 
+void UStatusBarSubsystem::StartProgressNotification(FProgressNotificationHandle Handle, FText DisplayText, int32 TotalWorkToDo)
+{
+	// Find the active status bar to display the progress in
+	for (auto StatusBar : StatusBars)
+	{
+		if (TSharedPtr<SStatusBar> StatusBarPinned = StatusBar.Value.Pin())
+		{
+			TSharedPtr<SDockTab> ParentTab = StatusBarPinned->GetParentTab();
+			if (ParentTab && ParentTab->IsForeground() && ParentTab->GetParentWindow() == FSlateApplication::Get().GetActiveTopLevelRegularWindow())
+			{
+				StatusBarPinned->StartProgressNotification(Handle, DisplayText, TotalWorkToDo);
+				break;
+			}
+		}
+	}
+}
+
+void UStatusBarSubsystem::UpdateProgressNotification(FProgressNotificationHandle Handle, int32 TotalWorkDone, int32 UpdatedTotalWorkToDo, FText UpdatedDisplayText)
+{
+	for (auto StatusBar : StatusBars)
+	{
+		if (TSharedPtr<SStatusBar> StatusBarPinned = StatusBar.Value.Pin())
+		{
+			if (StatusBarPinned->UpdateProgressNotification(Handle, TotalWorkDone, UpdatedTotalWorkToDo, UpdatedDisplayText))
+			{
+				break;
+			}
+		}
+	}
+
+}
+
+void UStatusBarSubsystem::CancelProgressNotification(FProgressNotificationHandle Handle)
+{
+	for (auto StatusBar : StatusBars)
+	{
+		if (TSharedPtr<SStatusBar> StatusBarPinned = StatusBar.Value.Pin())
+		{
+			if (StatusBarPinned->CancelProgressNotification(Handle))
+			{
+				break;
+			}
+		}
+	}
+}
+
 void UStatusBarSubsystem::OnDebugConsoleClosed()
 {
 	if (PreviousKeyboardFocusedWidget.IsValid())
diff --git a/Engine/Source/Editor/StatusBar/Public/SStatusBar.h b/Engine/Source/Editor/StatusBar/Public/SStatusBar.h
deleted file mode 100644
diff --git a/Engine/Source/Editor/StatusBar/Public/StatusBarSubsystem.h b/Engine/Source/Editor/StatusBar/Public/StatusBarSubsystem.h
--- a/Engine/Source/Editor/StatusBar/Public/StatusBarSubsystem.h
+++ b/Engine/Source/Editor/StatusBar/Public/StatusBarSubsystem.h
@@ -4,6 +4,7 @@
 
 #include "CoreMinimal.h"
 #include "EditorSubsystem.h"
+#include "Framework/Notifications/NotificationManager.h"
 
 #include "StatusBarSubsystem.generated.h"
 
@@ -47,7 +48,7 @@ private:
 };
 
 UCLASS()
-class STATUSBAR_API UStatusBarSubsystem : public UEditorSubsystem
+class STATUSBAR_API UStatusBarSubsystem : public UEditorSubsystem, public IProgressNotificationHandler
 {
 	GENERATED_BODY()
 
@@ -112,6 +113,11 @@ public:
 	void ClearStatusBarMessages(FName StatusBarName);
 
 private:
+	/** IProgressNotificationHandler interface */
+	virtual void StartProgressNotification(FProgressNotificationHandle Handle, FText DisplayText, int32 TotalWorkToDo) override;
+	virtual void UpdateProgressNotification(FProgressNotificationHandle Handle, int32 TotalWorkDone, int32 UpdatedTotalWorkToDo, FText UpdatedDisplayText) override;
+	virtual void CancelProgressNotification(FProgressNotificationHandle Handle) override;
+
 	bool ToggleContentBrowser(TSharedRef<SWindow> ParentWindow);
 	void OnDebugConsoleClosed();
 	void CreateContentBrowserIfNeeded();
@@ -127,5 +133,5 @@ private:
 
 	/** The floating content browser that is opened via the content browser button in the status bar */
 	TSharedPtr<SWidget> StatusBarContentBrowser;
-	static int32 HandleCounter;
+	static int32 MessageHandleCounter;
 };
diff --git a/Engine/Source/Runtime/Slate/Private/Framework/Application/SlateApplication.cpp b/Engine/Source/Runtime/Slate/Private/Framework/Application/SlateApplication.cpp
--- a/Engine/Source/Runtime/Slate/Private/Framework/Application/SlateApplication.cpp
+++ b/Engine/Source/Runtime/Slate/Private/Framework/Application/SlateApplication.cpp
@@ -2608,6 +2608,17 @@ TSharedPtr<SWindow> FSlateApplication::GetActiveTopLevelWindow() const
 	return ActiveTopLevelWindow.Pin();
 }
 
+TSharedPtr<SWindow> FSlateApplication::GetActiveTopLevelRegularWindow() const
+{
+	TSharedPtr<SWindow> ActiveWindow = ActiveTopLevelWindow.Pin();
+	while (ActiveWindow && !ActiveWindow->IsRegularWindow())
+	{
+		ActiveWindow = ActiveWindow->GetParentWindow();
+	}
+
+	return ActiveWindow;
+}
+
 TSharedPtr<SWindow> FSlateApplication::GetActiveModalWindow() const
 {
 	return (ActiveModalWindows.Num() > 0) ? ActiveModalWindows.Last() : nullptr;
diff --git a/Engine/Source/Runtime/Slate/Private/Framework/Notifications/NotificationManager.cpp b/Engine/Source/Runtime/Slate/Private/Framework/Notifications/NotificationManager.cpp
--- a/Engine/Source/Runtime/Slate/Private/Framework/Notifications/NotificationManager.cpp
+++ b/Engine/Source/Runtime/Slate/Private/Framework/Notifications/NotificationManager.cpp
@@ -10,9 +10,12 @@
 namespace NotificationManagerConstants
 {
 	// Offsets from the bottom-left corner of the work area
-	const FVector2D NotificationOffset( 15.0f, 15.0f );
+	const FVector2D NotificationOffset( 15.0f, 20.0f );
 }
 
+int32 FSlateNotificationManager::ProgressHandleCounter = 0;
+
+
 FSlateNotificationManager::FRegionalNotificationList::FRegionalNotificationList(const FSlateRect& InRectangle)
 	: Region(InRectangle)
 {
@@ -93,25 +96,65 @@ void FSlateNotificationManager::SetRootWindow( const TSharedRef<SWindow> InRootW
 	RootWindowPtr = InRootWindow;
 }
 
-TSharedRef<SNotificationList> FSlateNotificationManager::CreateStackForArea(const FSlateRect& InRectangle)
+TSharedRef<SNotificationList> FSlateNotificationManager::CreateStackForArea(const FSlateRect& InRectangle, TSharedPtr<SWindow> Window)
 {
+	// Doesn't work, per-pixel transparency on windows is required which is broken
+#define ONE_WINDOW_FOR_NOTIFICATIONS 0
+#if ONE_WINDOW_FOR_NOTIFICATIONS 
+	for (FRegionalNotificationList& List : RegionalLists)
+	{
+		if (FSlateRect::IsRectangleContained(List.Region, InRectangle))
+		{
+			return List.Notifications[0];
+		}
+	}
+
+	TSharedRef<SNotificationList> NewNotificationList = SNew(SNotificationList);
+	TSharedRef<SWindow> NotificationWindow = SWindow::MakeNotificationWindow();
+	NotificationWindow->SetContent(NewNotificationList);
+	NewNotificationList->ParentWindowPtr = NotificationWindow;
+
+	if (Window.IsValid())
+	{
+		FSlateApplication::Get().AddWindowAsNativeChild(NotificationWindow, Window.ToSharedRef());
+	}
+	else
+	{
+		FSlateApplication::Get().AddWindow(NotificationWindow);
+	}
+
+	if (!FSlateApplication::Get().GetActiveModalWindow().IsValid())
+	{
+		if (NotificationWindow->IsActive() || NotificationWindow->HasActiveParent())
+		{
+			NotificationWindow->BringToFront();
+		}
+	}
+
+
+	FRegionalNotificationList NewList(FSlateApplication::Get().GetWorkArea(InRectangle));
+	NewList.Notifications.Add(NewNotificationList);
+	RegionalLists.Add(NewList);
+
+	return NewNotificationList;
+#else
 	TSharedRef<SNotificationList> NotificationList = SNew(SNotificationList);
 	TSharedRef<SWindow> NotificationWindow = SWindow::MakeNotificationWindow();
 	NotificationWindow->SetContent(NotificationList);
 	NotificationList->ParentWindowPtr = NotificationWindow;
 
-	if( RootWindowPtr.IsValid() )
+	if (Window.IsValid())
 	{
-		FSlateApplication::Get().AddWindowAsNativeChild( NotificationWindow, RootWindowPtr.Pin().ToSharedRef() );
+		FSlateApplication::Get().AddWindowAsNativeChild(NotificationWindow, Window.ToSharedRef());
 	}
 	else
 	{
-		FSlateApplication::Get().AddWindow( NotificationWindow );
+		FSlateApplication::Get().AddWindow(NotificationWindow);
 	}
 
-	if( !FSlateApplication::Get().GetActiveModalWindow().IsValid() )
+	if (!FSlateApplication::Get().GetActiveModalWindow().IsValid())
 	{
-		if ( NotificationWindow->IsActive() || NotificationWindow->HasActiveParent() )
+		if (NotificationWindow->IsActive() || NotificationWindow->HasActiveParent())
 		{
 			NotificationWindow->BringToFront();
 		}
@@ -135,19 +178,26 @@ TSharedRef<SNotificationList> FSlateNotificationManager::CreateStackForArea(cons
 	}
 
 	return NotificationList;
+#endif // ONE_WINDOW_FOR_NOTIFICATIONS
+
+	
 }
 
 TSharedPtr<SNotificationItem> FSlateNotificationManager::AddNotification(const FNotificationInfo& Info)
 {
-	check(IsInGameThread() || !"FSlateNotificationManager::AddNotification must be called on game thread. Use QueueNotification if necessary.");
+	checkf(IsInGameThread(), TEXT("FSlateNotificationManager::AddNotification must be called on game thread. Use QueueNotification if necessary."));
 
 	// Early calls of this function can happen before Slate is initialized.
 	if( FSlateApplication::IsInitialized() )
 	{
 		FSlateRect PreferredWorkArea;
-		// Retrieve the main editor window to display the notification in, otherwise use the preferred work area
-		if (RootWindowPtr.IsValid())
+		if (Info.ForWindow.IsValid())
 		{
+			PreferredWorkArea = FSlateApplication::Get().GetWorkArea(Info.ForWindow->GetRectInScreen());
+		}
+		else if (RootWindowPtr.IsValid())
+		{
+			// Retrieve the main editor window to display the notification in, otherwise use the preferred work area
 			PreferredWorkArea = FSlateApplication::Get().GetWorkArea(RootWindowPtr.Pin()->GetRectInScreen());
 		}
 		else
@@ -155,7 +205,7 @@ TSharedPtr<SNotificationItem> FSlateNotificationManager::AddNotification(const F
 			PreferredWorkArea = FSlateApplication::Get().GetPreferredWorkArea();
 		}
 
-		TSharedRef<SNotificationList> List = CreateStackForArea(PreferredWorkArea);
+		TSharedRef<SNotificationList> List = CreateStackForArea(PreferredWorkArea, Info.ForWindow.IsValid() ? Info.ForWindow : RootWindowPtr.Pin());
 
 		return List->AddNotification(Info);
 	}
@@ -168,6 +218,41 @@ void FSlateNotificationManager::QueueNotification(FNotificationInfo* Info)
 	PendingNotifications.Push(Info);
 }
 
+FProgressNotificationHandle FSlateNotificationManager::StartProgressNotification(FText DisplayText, int32 TotalWorkToDo)
+{
+	if (ProgressNotificationHandler)
+	{
+		FProgressNotificationHandle NewHandle(++ProgressHandleCounter);
+		ProgressNotificationHandler->StartProgressNotification(NewHandle, DisplayText, TotalWorkToDo);
+
+		return NewHandle;
+	}
+
+	return FProgressNotificationHandle();
+}
+
+void FSlateNotificationManager::UpdateProgressNotification(FProgressNotificationHandle InHandle, int32 TotalWorkDone, int32 UpdatedTotalWorkToDo, FText UpdatedDisplayText)
+{
+	if (ProgressNotificationHandler)
+	{
+		ProgressNotificationHandler->UpdateProgressNotification(InHandle, TotalWorkDone, UpdatedTotalWorkToDo, UpdatedDisplayText);
+	}
+}
+
+void FSlateNotificationManager::CancelProgressNotification(FProgressNotificationHandle InHandle)
+{
+	if (ProgressNotificationHandler)
+	{
+		ProgressNotificationHandler->CancelProgressNotification(InHandle);
+	}
+}
+
+void FSlateNotificationManager::SetProgressNotificationHandler(IProgressNotificationHandler* NewHandler)
+{
+	ProgressNotificationHandler = NewHandler;
+}
+
+
 void FSlateNotificationManager::GetWindows(TArray< TSharedRef<SWindow> >& OutWindows) const
 {
 	for (const FRegionalNotificationList& RegionList : RegionalLists)
@@ -256,173 +341,4 @@ void FSlateNotificationManager::ForceNotificationsInFront( const TSharedRef<SWin
 	}
 }
 
-#if (WITH_EDITOR || IS_PROGRAM) && !UE_BUILD_SHIPPING
-
-static void TestNotifications()
-{
-	FTicker::GetCoreTicker().AddTicker(TEXT("TestNotifications"), 0.0f, [](float DeltaTime)
-		{
-			FNotificationInfo NotificationInfo(FText::FromString(TEXT("Test Notification 1")));
-			NotificationInfo.FadeInDuration = 2.0f;
-			NotificationInfo.FadeOutDuration = 2.0f;
-			NotificationInfo.ExpireDuration = 10.0f;
-
-			auto Notification = FSlateNotificationManager::Get().AddNotification(NotificationInfo);
-			Notification->ExpireAndFadeout();
-
-			return false;
-		});
-
-	FTicker::GetCoreTicker().AddTicker(TEXT("TestNotifications"), 1.0f, [](float DeltaTime)
-		{
-			FNotificationInfo NotificationInfo(FText::FromString(TEXT("Test Notification 2")));
-			NotificationInfo.FadeInDuration = 2.0f;
-			NotificationInfo.FadeOutDuration = 2.0f;
-			NotificationInfo.ExpireDuration = 10.0f;
-
-			NotificationInfo.bUseLargeFont = false;
-			auto Notification = FSlateNotificationManager::Get().AddNotification(NotificationInfo);
-			Notification->ExpireAndFadeout();
-
-			return false;
-		});
-
-	FTicker::GetCoreTicker().AddTicker(TEXT("TestNotifications"), 2.0f, [](float DeltaTime)
-		{
-			FNotificationInfo NotificationInfo(FText::FromString(TEXT("Test Notification 3")));
-			NotificationInfo.FadeInDuration = 2.0f;
-			NotificationInfo.FadeOutDuration = 2.0f;
-			NotificationInfo.ExpireDuration = 10.0f;
-
-			NotificationInfo.bUseThrobber = true;
-			auto Notification = FSlateNotificationManager::Get().AddNotification(NotificationInfo);
-			Notification->SetCompletionState(SNotificationItem::CS_Pending);
-			Notification->ExpireAndFadeout();
-
-			return false;
-		});
-
-	FTicker::GetCoreTicker().AddTicker(TEXT("TestNotifications"), 3.0f, [](float DeltaTime)
-		{
-			FNotificationInfo NotificationInfo(FText::FromString(TEXT("Test Notification 4")));
-			NotificationInfo.FadeInDuration = 2.0f;
-			NotificationInfo.FadeOutDuration = 2.0f;
-			NotificationInfo.ExpireDuration = 10.0f;
-
-			NotificationInfo.bUseSuccessFailIcons = true;
-			auto Notification = FSlateNotificationManager::Get().AddNotification(NotificationInfo);
-			Notification->SetCompletionState(SNotificationItem::CS_Success);
-			Notification->ExpireAndFadeout();
-
-			return false;
-		});
-
-	FTicker::GetCoreTicker().AddTicker(TEXT("TestNotifications"), 4.0f, [](float DeltaTime)
-		{
-			FNotificationInfo NotificationInfo(FText::FromString(TEXT("Test Notification 5")));
-			NotificationInfo.FadeInDuration = 2.0f;
-			NotificationInfo.FadeOutDuration = 2.0f;
-			NotificationInfo.ExpireDuration = 10.0f;
-
-			NotificationInfo.bUseSuccessFailIcons = true;
-			auto Notification = FSlateNotificationManager::Get().AddNotification(NotificationInfo);
-			Notification->SetCompletionState(SNotificationItem::CS_Fail);
-			Notification->ExpireAndFadeout();
-
-			return false;
-		});
-
-	FTicker::GetCoreTicker().AddTicker(TEXT("TestNotifications"), 5.0f, [](float DeltaTime)
-		{
-			FNotificationInfo NotificationInfo(FText::FromString(TEXT("Test Notification 6")));
-			NotificationInfo.FadeInDuration = 2.0f;
-			NotificationInfo.FadeOutDuration = 2.0f;
-			NotificationInfo.ExpireDuration = 10.0f;
-
-			NotificationInfo.CheckBoxText = FText::FromString(TEXT("Don't ask again"));
-			NotificationInfo.CheckBoxState = ECheckBoxState::Checked;
-			NotificationInfo.CheckBoxStateChanged = FOnCheckStateChanged::CreateStatic([](ECheckBoxState NewState) {});
-
-			auto Notification = FSlateNotificationManager::Get().AddNotification(NotificationInfo);
-
-			Notification->ExpireAndFadeout();
-
-			return false;
-		});
-
-	FTicker::GetCoreTicker().AddTicker(TEXT("TestNotifications"), 6.0f, [](float DeltaTime)
-		{
-			FNotificationInfo NotificationInfo(FText::FromString(TEXT("Test Notification 7")));
-			NotificationInfo.FadeInDuration = 2.0f;
-			NotificationInfo.FadeOutDuration = 2.0f;
-			NotificationInfo.ExpireDuration = 10.0f;
-
-			NotificationInfo.Hyperlink = FSimpleDelegate::CreateLambda([]() {});
-			NotificationInfo.HyperlinkText = FText::FromString(TEXT("This is a hyperlink"));
-
-			auto Notification = FSlateNotificationManager::Get().AddNotification(NotificationInfo);
-
-			Notification->ExpireAndFadeout();
-
-			return false;
-		});
-
-	FTicker::GetCoreTicker().AddTicker(TEXT("TestNotifications"), 7.0f, [](float DeltaTime)
-		{
-			FNotificationInfo NotificationInfo(FText::FromString(TEXT("Test Notification 8")));
-			NotificationInfo.FadeInDuration = 2.0f;
-			NotificationInfo.FadeOutDuration = 2.0f;
-			NotificationInfo.ExpireDuration = 10.0f;
-
-			FNotificationButtonInfo Button1(FText::FromString("Ok"), FText::GetEmpty(), FSimpleDelegate(), SNotificationItem::ECompletionState::CS_None);
-			FNotificationButtonInfo Button2(FText::FromString("Cancel"), FText::GetEmpty(), FSimpleDelegate(), SNotificationItem::ECompletionState::CS_None);
-
-			NotificationInfo.ButtonDetails.Add(Button1);
-			NotificationInfo.ButtonDetails.Add(Button2);
-
-			auto Notification = FSlateNotificationManager::Get().AddNotification(NotificationInfo);
-
-			Notification->ExpireAndFadeout();
-
-			return false;
-		});
-
-	FTicker::GetCoreTicker().AddTicker(TEXT("TestNotifications"), 8.0f, [](float DeltaTime)
-		{
-			FNotificationInfo NotificationInfo(FText::FromString(TEXT("Everything Under The Sun")));
-			NotificationInfo.FadeInDuration = 2.0f;
-			NotificationInfo.FadeOutDuration = 2.0f;
-			NotificationInfo.ExpireDuration = 10.0f;
-
-			NotificationInfo.CheckBoxText = FText::FromString(TEXT("Don't ask again"));
-			NotificationInfo.CheckBoxState = ECheckBoxState::Checked;
-			NotificationInfo.CheckBoxStateChanged = FOnCheckStateChanged::CreateStatic([](ECheckBoxState NewState) {});
-
-			NotificationInfo.Hyperlink = FSimpleDelegate::CreateLambda([]() {});
-			NotificationInfo.HyperlinkText = FText::FromString(TEXT("This is a hyperlink"));
-
-			NotificationInfo.bUseSuccessFailIcons = true;
-			NotificationInfo.bUseThrobber = true;
-
-
-			FNotificationButtonInfo Button1(FText::FromString("OK"), FText::GetEmpty(), FSimpleDelegate());
-			FNotificationButtonInfo Button2(FText::FromString("Cancel"), FText::GetEmpty(), FSimpleDelegate());
-
-			NotificationInfo.ButtonDetails.Add(Button1);
-			NotificationInfo.ButtonDetails.Add(Button2);
-
-
-			auto Notification = FSlateNotificationManager::Get().AddNotification(NotificationInfo);
-			Notification->SetCompletionState(SNotificationItem::CS_Pending);
-
-			Notification->ExpireAndFadeout();
-
-			return false;
-		});
-}
-
-FAutoConsoleCommand TestNotificationCommand(TEXT("Slate.TestNotifications"), TEXT(""), FConsoleCommandDelegate::CreateStatic(&TestNotifications));
-
-#endif // (WITH_EDITOR || IS_PROGRAM) && !UE_BUILD_SHIPPING
-
 #undef LOCTEXT_NAMESPACE
diff --git a/Engine/Source/Runtime/Slate/Private/Widgets/Notifications/SNotificationList.cpp b/Engine/Source/Runtime/Slate/Private/Widgets/Notifications/SNotificationList.cpp
--- a/Engine/Source/Runtime/Slate/Private/Widgets/Notifications/SNotificationList.cpp
+++ b/Engine/Source/Runtime/Slate/Private/Widgets/Notifications/SNotificationList.cpp
@@ -16,6 +16,7 @@
 #include "Widgets/Images/SThrobber.h"
 #include "Widgets/Notifications/INotificationWidget.h"
 #include "Framework/Notifications/NotificationManager.h"
+#include "Styling/StyleColors.h"
 
 /////////////////////////////////////////////////
 // SNotificationExtendable
@@ -74,17 +75,22 @@ public:
 		// Make visible
 		SetVisibility(EVisibility::Visible);
 
+		TSharedRef<SWidget> Self = AsShared();
+
 		// Play Fadein animation
 		FadeAnimation = FCurveSequence();
 		FadeCurve = FadeAnimation.AddCurve(0.f, FadeInDuration.Get());
-		FadeAnimation.Play( this->AsShared() );
+		FadeAnimation.Play(Self);
+
+		ThrobberAnimation = FCurveSequence(0.0f, 1.0f);
+		ThrobberAnimation.Play(Self, true);
 
 		// Scale up/flash animation
 		IntroAnimation = FCurveSequence();
 		ScaleCurveX = IntroAnimation.AddCurve(0.2f, 0.3f, ECurveEaseFunction::QuadOut);
 		ScaleCurveY = IntroAnimation.AddCurve(0.f, 0.2f);
 		GlowCurve = IntroAnimation.AddCurve(0.5f, 0.55f, ECurveEaseFunction::QuadOut);
-		IntroAnimation.Play( this->AsShared() );
+		IntroAnimation.Play(Self);
 
 		// When a fade in occurs, we need a high framerate for the animation to look good
 		if( FadeInDuration.Get() > KINDA_SMALL_NUMBER && bAllowThrottleWhenFrameRateIsLow && !ThrottleHandle.IsValid() )
@@ -168,9 +174,9 @@ protected:
 		}
 
 		// Clear reference
-		if( MyList.IsValid() )
+		if(TSharedPtr< SNotificationList> MyListPinned = MyList.Pin())
 		{
-			MyList.Pin()->NotificationItemFadedOut(SharedThis(this));
+			MyListPinned->NotificationItemFadedOut(SharedThis(this));
 		}
 	}
 
@@ -185,9 +191,12 @@ protected:
 	{
 		// if we have a parent window, we need to make that transparent, rather than
 		// this widget
-		if(MyList.IsValid() && MyList.Pin()->ParentWindowPtr.IsValid())
+		if(TSharedPtr<SNotificationList> MyListPinned = MyList.Pin())
 		{
-			MyList.Pin()->ParentWindowPtr.Pin()->SetOpacity( FadeCurve.GetLerp() );
+			if (TSharedPtr<SWindow> MyWindow = MyListPinned->ParentWindowPtr.Pin())
+			{
+				MyWindow->SetOpacity(FadeCurve.GetLerp());
+			}
 			return FLinearColor(1,1,1,1);
 		}
 		else
@@ -220,20 +229,45 @@ protected:
 		return FVector2D(  ScaleCurveX.GetLerp(), ScaleCurveY.GetLerp() );
 	}
 
+	FMargin GetIconPadding() const
+	{
+		if (GetThrobberVisibility().IsVisible())
+		{
+			// Account for the throbber taking up space
+			return FMargin(-6.0f, 0.0f, 0.0f, 0.0f);
+		}
+
+		return FMargin(0);
+	}
+
 	/** Gets the visibility for the throbber */
 	EVisibility GetThrobberVisibility() const
 	{
 		return CompletionState == CS_Pending ? EVisibility::Visible : EVisibility::Collapsed;
 	}
 
+	TOptional<FSlateRenderTransform> GetThrobberTransform() const
+	{
+		const float DeltaAngle = ThrobberAnimation.GetLerp()*2*PI;
+		return FSlateRenderTransform(FQuat2D(DeltaAngle));
+	}
+
 	EVisibility GetSuccessFailImageVisibility() const
 	{
 		return (CompletionState == CS_Success || CompletionState == CS_Fail) ? EVisibility::Visible : EVisibility::Collapsed;
 	}
 
-	const FSlateBrush* GetSuccessFailImage() const
+	const FSlateBrush* GetNotificationIcon() const
 	{
-		return CompletionState == CS_Success ? FCoreStyle::Get().GetBrush("NotificationList.SuccessImage") : FCoreStyle::Get().GetBrush("NotificationList.FailImage");
+		switch (CompletionState)
+		{
+		case ECompletionState::CS_Success:
+			return FAppStyle::Get().GetBrush("Icons.SuccessWithColor.Large");
+		case ECompletionState::CS_Fail:
+			return FAppStyle::Get().GetBrush("Icons.ErrorWithColor.Large");
+		default:
+			return BaseIcon.Get();
+		}
 	}
 
 public:
@@ -243,16 +277,18 @@ public:
 protected:
 
 	/** The text displayed in this text block */
-	TAttribute< FText > Text;
+	TAttribute<FText> Text;
 
 	/** The fade in duration for this element */
-	TAttribute< float > FadeInDuration;
+	TAttribute<float> FadeInDuration;
 
 	/** The fade out duration for this element */
-	TAttribute< float > FadeOutDuration;
+	TAttribute<float> FadeOutDuration;
 
 	/** The duration before a fadeout for this element */
-	TAttribute< float > ExpireDuration;
+	TAttribute<float> ExpireDuration;
+
+	TAttribute<const FSlateBrush*> BaseIcon;
 
 	/** The default glow color for pulse animation. */
 	FLinearColor DefaultGlowColor = FLinearColor(1.0f, 1.0f, 1.0f);
@@ -272,7 +308,7 @@ protected:
 	FCurveHandle ScaleCurveX;
 	FCurveHandle ScaleCurveY;
 	FCurveHandle GlowCurve;
-
+	FCurveSequence ThrobberAnimation;
 	/** The completion state change animation */
 	FCurveSequence CompletionStateAnimation;
 
@@ -281,6 +317,78 @@ protected:
 };
 
 
+class SNotificationBackground : public SBorder
+{
+public:
+	SLATE_BEGIN_ARGS(SNotificationBackground) {}
+		SLATE_ATTRIBUTE(FMargin, Padding)
+		/** ColorAndOpacity is the color and opacity of content in the border */
+		SLATE_ATTRIBUTE(FLinearColor, ColorAndOpacity)
+		/** BorderBackgroundColor refers to the actual color and opacity of the supplied border image.*/
+		SLATE_ATTRIBUTE(FSlateColor, BorderBackgroundColor)
+
+		SLATE_ATTRIBUTE(FVector2D, DesiredSizeScale)
+
+		SLATE_DEFAULT_SLOT(FArguments, Content)
+	SLATE_END_ARGS()
+
+	void Construct(const FArguments& InArgs)
+	{
+		WatermarkBrush = FAppStyle::Get().GetBrush("NotificationList.Watermark");
+		BorderBrush = FAppStyle::Get().GetBrush("NotificationList.ItemBackground_Border");
+
+		WatermarkTint = FStyleColors::Notifications.GetSpecifiedColor() * FLinearColor(1.15f, 1.15f, 1.15f, 1.0f);
+
+		SBorder::Construct(
+			SBorder::FArguments()
+			.BorderImage(FAppStyle::Get().GetBrush("NotificationList.ItemBackground"))
+			.Padding(InArgs._Padding)
+			.BorderBackgroundColor(InArgs._BorderBackgroundColor)
+			.ColorAndOpacity(InArgs._ColorAndOpacity)
+			.DesiredSizeScale(InArgs._DesiredSizeScale)
+			[
+				InArgs._Content.Widget
+			]
+		);
+	}
+
+	virtual int32 OnPaint(const FPaintArgs& Args, const FGeometry& AllottedGeometry, const FSlateRect& MyCullingRect, FSlateWindowElementList& OutDrawElements, int32 LayerId, const FWidgetStyle& InWidgetStyle, bool bParentEnabled) const override
+	{
+		if (WatermarkBrush && WatermarkBrush->DrawAs != ESlateBrushDrawType::NoDrawType)
+		{
+			// The watermark should be 250% bigger than the size of the notification to get the desired effect
+			const float SizeY = AllottedGeometry.GetLocalSize().Y;
+			const float WatermarkSize = FMath::Clamp(SizeY * 2.50f, 130.0f, 220.0f);
+	
+			FSlateDrawElement::MakeBox(
+				OutDrawElements,
+				LayerId + 1,
+				AllottedGeometry.ToPaintGeometry(FVector2D(4.0f, 4.0f), FVector2D(WatermarkSize, WatermarkSize)),
+				WatermarkBrush,
+				ESlateDrawEffect::None,
+				WatermarkTint * InWidgetStyle.GetColorAndOpacityTint() * BorderBackgroundColor.Get().GetColor(InWidgetStyle));
+		}
+
+
+		int32 OutLayerId = SBorder::OnPaint(Args, AllottedGeometry, MyCullingRect, OutDrawElements, LayerId, InWidgetStyle, bParentEnabled);
+
+		FSlateDrawElement::MakeBox(
+			OutDrawElements,
+			OutLayerId,
+			AllottedGeometry.ToPaintGeometry(),
+			BorderBrush,
+			ESlateDrawEffect::None,
+			BorderBrush->GetTint(InWidgetStyle) * InWidgetStyle.GetColorAndOpacityTint());
+
+		return OutLayerId+1;
+	}
+	
+private:
+	const FSlateBrush* WatermarkBrush = nullptr;
+	const FSlateBrush* BorderBrush = nullptr;
+	FLinearColor WatermarkTint;
+};
+
 /////////////////////////////////////////////////
 // SNotificationItemImpl
 
@@ -290,7 +398,6 @@ class SNotificationItemImpl : public SNotificationExtendable
 public:
 	SLATE_BEGIN_ARGS( SNotificationItemImpl )
 		: _Text()
-		, _Font()
 		, _Image()	
 		, _FadeInDuration(0.5f)
 		, _FadeOutDuration(2.f)
@@ -299,37 +406,36 @@ public:
 	{}
 
 		/** The text displayed in this text block */
-		SLATE_ATTRIBUTE( FText, Text )
-		/** Sets the font used to draw the text */
-		SLATE_ATTRIBUTE(FSlateFontInfo, Font)
+		SLATE_ARGUMENT(FText, Text)
+
 		/** Setup information for the buttons on the notification */ 
-		SLATE_ATTRIBUTE(TArray<FNotificationButtonInfo>, ButtonDetails)
+		SLATE_ARGUMENT(TArray<FNotificationButtonInfo>, ButtonDetails)
 		/** The icon image to display next to the text */
-		SLATE_ATTRIBUTE( const FSlateBrush*, Image )
+		SLATE_ARGUMENT(const FSlateBrush*, Image)
 		/** The fade in duration for this element */
-		SLATE_ATTRIBUTE( float, FadeInDuration )
+		SLATE_ARGUMENT(float, FadeInDuration)
 		/** The fade out duration for this element */
-		SLATE_ATTRIBUTE( float, FadeOutDuration )
+		SLATE_ARGUMENT(float, FadeOutDuration)
 		/** The duration before a fadeout for this element */
-		SLATE_ATTRIBUTE( float, ExpireDuration )
+		SLATE_ARGUMENT(float, ExpireDuration)
 		/** Controls whether or not to add the animated throbber */
-		SLATE_ATTRIBUTE( bool, bUseThrobber)
+		SLATE_ARGUMENT(bool, bUseThrobber)
 		/** Controls whether or not to display the success and fail icons */
-		SLATE_ATTRIBUTE( bool, bUseSuccessFailIcons)
+		SLATE_ARGUMENT(bool, bUseSuccessFailIcons)
 		/** When true the larger bolder font will be used to display the message */
-		SLATE_ATTRIBUTE( bool, bUseLargeFont)
+		SLATE_ARGUMENT(bool, bUseLargeFont)
 		/** When set this forces the width of the box, used to stop resizeing on text change */
-		SLATE_ARGUMENT( FOptionalSize, WidthOverride )
+		SLATE_ARGUMENT(FOptionalSize, WidthOverride)
 		/** When set this will display a check box on the notification; handles getting the current check box state */
-		SLATE_ATTRIBUTE( ECheckBoxState, CheckBoxState )
+		SLATE_ATTRIBUTE(ECheckBoxState, CheckBoxState)
 		/** When set this will display a check box on the notification; handles setting the new check box state */
-		SLATE_EVENT( FOnCheckStateChanged, CheckBoxStateChanged );
+		SLATE_EVENT(FOnCheckStateChanged, CheckBoxStateChanged);
 		/** Text to display for the check box message */
-		SLATE_ATTRIBUTE( FText, CheckBoxText );
+		SLATE_ATTRIBUTE(FText, CheckBoxText );
 		/** When set this will display as a hyperlink on the right side of the notification. */
-		SLATE_EVENT( FSimpleDelegate, Hyperlink)
+		SLATE_EVENT(FSimpleDelegate, Hyperlink)
 		/** Text to display for the hyperlink (if Hyperlink is valid) */
-		SLATE_ATTRIBUTE( FText, HyperlinkText );
+		SLATE_ATTRIBUTE(FText, HyperlinkText);
 
 	SLATE_END_ARGS()
 	/**
@@ -348,194 +454,152 @@ public:
 
 		ChildSlot
 		[
-			SNew(SBorder)
-			.BorderImage(FCoreStyle::Get().GetBrush("NotificationList.ItemBackground"))
-			.BorderBackgroundColor(this, &SNotificationItemImpl::GetContentColor)
-			.ColorAndOpacity(this, &SNotificationItemImpl::GetContentColorRaw)
-			.DesiredSizeScale(this, &SNotificationItemImpl::GetItemScale)
-			[
-				SNew(SBorder)
-				.Padding( FMargin(5) )
-				.BorderImage(FCoreStyle::Get().GetBrush("NotificationList.ItemBackground_Border"))
-				.BorderBackgroundColor(this, &SNotificationItemImpl::GetGlowColor)
+/*
+			SNew(SOverlay)
+			+ SOverlay::Slot()
+			[*/
+				SNew(SNotificationBackground)
+				.Padding(FMargin(16, 8))
+				.BorderBackgroundColor(this, &SNotificationItemImpl::GetContentColor)
+				.ColorAndOpacity(this, &SNotificationItemImpl::GetContentColorRaw)
+				.DesiredSizeScale(this, &SNotificationItemImpl::GetItemScale)
 				[
 					ConstructInternals(InArgs)
 				]
+/*
 			]
+			+SOverlay::Slot()
+			[
+				SNew(SImage)
+				.Image(FAppStyle::Get().GetBrush("NotificationList.ItemBackground_Border"))
+				.Visibility(EVisibility::HitTestInvisible)
+			]*/
 		];
 	}
 	
 	/**
 	 * Returns the internals of the notification
 	 */
-	TSharedRef<SHorizontalBox> ConstructInternals( const FArguments& InArgs ) 
+	TSharedRef<SWidget> ConstructInternals( const FArguments& InArgs ) 
 	{
 		CheckBoxStateChanged = InArgs._CheckBoxStateChanged;
 		Hyperlink = InArgs._Hyperlink;
 		HyperlinkText = InArgs._HyperlinkText;
+		BaseIcon = InArgs._Image;
 
-		TSharedRef<SHorizontalBox> HorizontalBox = SNew(SHorizontalBox);
+		// Container for the text and optional interactive widgets (buttons, check box, and hyperlink)
+		TSharedRef<SVerticalBox> InteractiveWidgetsBox = SNew(SVerticalBox);
 
-		// Notification image
-		HorizontalBox->AddSlot()
-		.AutoWidth()
-		.Padding(10.f, 0.f, 0.f, 0.f)
-		.VAlign(VAlign_Center)
-		.HAlign(HAlign_Left)
+		InteractiveWidgetsBox->AddSlot()
 		[
-			SNew(SImage)
-			.Image( InArgs._Image )
+			SAssignNew(MyTextBlock, STextBlock)
+			.Text(Text)
+			.Font(FAppStyle::Get().GetFontStyle(TEXT("NotificationList.FontBold")))
+			.WrapTextAt(InArgs._WidthOverride.IsSet() ? InArgs._WidthOverride.Get()-50.0f : 0.0f )
 		];
 
 		{
-			FSlateFontInfo Font = InArgs._Font.Get();
-
-			if (!Font.HasValidFont())
-			{
-				Font = InArgs._bUseLargeFont.Get()
-					? FCoreStyle::Get().GetFontStyle(TEXT("NotificationList.FontBold"))
-					: FCoreStyle::Get().GetFontStyle(TEXT("NotificationList.FontLight"));
-			}
-
-			// Container for the text and optional interactive widgets (buttons, check box, and hyperlink)
-			TSharedRef<SVerticalBox> TextAndInteractiveWidgetsBox = SNew(SVerticalBox);
-
-			HorizontalBox->AddSlot()
-			.AutoWidth()
-			.Padding(10.f, 0.f, 15.f, 0.f)
+			InteractiveWidgetsBox->AddSlot()
+			.AutoHeight()
+			.Padding(FMargin(0.0f, 10.0f, 0.0f, 2.0f))
 			.VAlign(VAlign_Center)
 			.HAlign(HAlign_Left)
 			[
-				TextAndInteractiveWidgetsBox
+				SNew(SHyperlink)
+				.Visibility(this, &SNotificationItemImpl::GetHyperlinkVisibility)
+				.Text(this, &SNotificationItemImpl::GetHyperlinkText)
+				.OnNavigate(this, &SNotificationItemImpl::OnHyperlinkClicked)
 			];
+		}
 
-			// Build Text box
-			TextAndInteractiveWidgetsBox->AddSlot()
+		{
+			InteractiveWidgetsBox->AddSlot()
 			.AutoHeight()
+			.Padding(FMargin(0.0f, 10.0f, 0.0f, 2.0f))
+			.VAlign(VAlign_Center)
+			.HAlign(HAlign_Left)
 			[
-				SNew(SBox)
-				.WidthOverride(InArgs._WidthOverride)
+				SNew(SCheckBox)
+				.Visibility(this, &SNotificationItemImpl::GetCheckBoxVisibility)
+				.IsChecked(InArgs._CheckBoxState)
+				.OnCheckStateChanged(CheckBoxStateChanged)
+				.Padding(FMargin(2.0f, 0.0f))
 				[
-					SAssignNew(MyTextBlock, STextBlock)
-					.Text(Text)
-					.Font(Font)
-					.AutoWrapText(InArgs._WidthOverride.IsSet()) // only auto-wrap the text if we've been given a size constraint; otherwise, fill the notification area
+					SNew(STextBlock)
+					.Text(InArgs._CheckBoxText)
 				]
 			];
+		}
 
-			TSharedRef<SHorizontalBox> InteractiveWidgetsBox = SNew(SHorizontalBox);
-			TextAndInteractiveWidgetsBox->AddSlot()
-			.AutoHeight()
-			[
-				InteractiveWidgetsBox
-			];
-
+		{
 			// Adds any buttons that were passed in.
+			TSharedRef<SHorizontalBox> ButtonsBox = SNew(SHorizontalBox);
+
+			for (int32 idx = 0; idx < InArgs._ButtonDetails.Num(); idx++)
 			{
-				TSharedRef<SHorizontalBox> ButtonsBox = SNew(SHorizontalBox);
-				for (int32 idx = 0; idx < InArgs._ButtonDetails.Get().Num(); idx++)
-				{
-					FNotificationButtonInfo Button = InArgs._ButtonDetails.Get()[idx];
-
-					ButtonsBox->AddSlot()
-					.AutoWidth()
-					.HAlign(HAlign_Left)
-					.VAlign(VAlign_Center)
-					.Padding(0.0f, 0.0f, 4.0f, 0.0f)
-					[
-						SNew(SButton)
-						.Text(Button.Text)
-						.ToolTipText(Button.ToolTip)
-						.OnClicked(this, &SNotificationItemImpl::OnButtonClicked, Button.Callback) 
-						.Visibility( this, &SNotificationItemImpl::GetButtonVisibility, Button.VisibilityOnNone, Button.VisibilityOnPending, Button.VisibilityOnSuccess, Button.VisibilityOnFail )
-					];
-				}
-				InteractiveWidgetsBox->AddSlot()
+				FNotificationButtonInfo Button = InArgs._ButtonDetails[idx];
+
+				ButtonsBox->AddSlot()
 				.AutoWidth()
-				.Padding(0.0f, 2.0f, 0.0f, 0.0f)
+
 				.VAlign(VAlign_Center)
-				.HAlign(HAlign_Left)
+				.Padding(0.0f, 0.0f, 4.0f, 0.0f)
 				[
-					ButtonsBox
+					SNew(SButton)
+					.Text(Button.Text)
+					.ToolTipText(Button.ToolTip)
+					.OnClicked(this, &SNotificationItemImpl::OnButtonClicked, Button.Callback) 
+					.Visibility( this, &SNotificationItemImpl::GetButtonVisibility, Button.VisibilityOnNone, Button.VisibilityOnPending, Button.VisibilityOnSuccess, Button.VisibilityOnFail )
 				];
 			}
 
-			// Adds a check box, but only visible when bound
 			InteractiveWidgetsBox->AddSlot()
-			.AutoWidth()
+			.AutoHeight()
+			.Padding(0.0f, 5.0f, 0.0f, 0.0f)
 			.VAlign(VAlign_Center)
-			.HAlign(HAlign_Left)
-			[
-				SNew(SBox)
-				.Padding(FMargin(0.0f, 2.0f, 4.0f, 0.0f))
-				.VAlign(VAlign_Center)
-				.HAlign(HAlign_Left)
-				.Visibility(this, &SNotificationItemImpl::GetCheckBoxVisibility)
-				[
-					SNew(SCheckBox)
-					.IsChecked(InArgs._CheckBoxState)
-					.OnCheckStateChanged(CheckBoxStateChanged)
-					[
-						SNew(STextBlock)
-						.Text(InArgs._CheckBoxText)
-					]
-				]
-			];
-
-			// Adds a hyperlink, but only visible when bound
-			InteractiveWidgetsBox->AddSlot()
-			.VAlign(VAlign_Bottom)
 			.HAlign(HAlign_Right)
 			[
-				SNew(SBox)
-				.Padding(FMargin(0.0f, 2.0f, 0.0f, 2.0f))
-				.VAlign(VAlign_Center)
-				.HAlign(HAlign_Left)
-				.Visibility(this, &SNotificationItemImpl::GetHyperlinkVisibility)
-				[
-					SNew(SHyperlink)
-					.Text(this, &SNotificationItemImpl::GetHyperlinkText)
-					.OnNavigate(this, &SNotificationItemImpl::OnHyperlinkClicked)
-				]
+				ButtonsBox
 			];
 		}
 
-		if (InArgs._bUseThrobber.Get())
-		{
-			// Build pending throbber
-			HorizontalBox->AddSlot()
-			.AutoWidth()
+		return
+			SNew(SBox)
+			.WidthOverride(InArgs._WidthOverride)
 			[
-				SNew(SBox)
-				.Padding(FMargin(5.f, 0.f, 10.f, 0.f))
-				.HAlign(HAlign_Center)
-				.VAlign(VAlign_Center)
-				.Visibility( this, &SNotificationItemImpl::GetThrobberVisibility )
+				SNew(SHorizontalBox)
+				+ SHorizontalBox::Slot()
+				.AutoWidth()
+				.Padding(MakeAttributeSP(this, &SNotificationItemImpl::GetIconPadding))
+				.VAlign(VAlign_Top)
+				.HAlign(HAlign_Left)
 				[
-					SNew(SThrobber)
+					SNew(SOverlay)
+					+SOverlay::Slot()
+					.VAlign(VAlign_Center)
+					.HAlign(HAlign_Center)
+					[
+						SNew(SImage)
+						.Image(FAppStyle::Get().GetBrush("NotificationList.Throbber"))
+						.Visibility(this, &SNotificationItemImpl::GetThrobberVisibility)
+						.RenderTransform(this, &SNotificationItemImpl::GetThrobberTransform)
+						.RenderTransformPivot(FVector2D(.5f,.5f))
+					] 
+					+ SOverlay::Slot()
+					.VAlign(VAlign_Center)
+					.HAlign(HAlign_Center)
+					[
+						SNew(SImage)
+						.Image(InArgs._bUseSuccessFailIcons ? MakeAttributeSP<const FSlateBrush*>(this, &SNotificationItemImpl::GetNotificationIcon) : BaseIcon)
+						.DesiredSizeOverride(FVector2D(32,32))
+					]
 				]
-			];
-		}
-
-		if (InArgs._bUseSuccessFailIcons.Get())
-		{
-			// Build success/fail image
-			HorizontalBox->AddSlot()
-			.AutoWidth()
-			[
-				SNew(SBox)
-				.Padding(FMargin(8.f, 0.f, 10.f, 0.f))
-				.HAlign(HAlign_Center)
-				.VAlign(VAlign_Center)
-				.Visibility( this, &SNotificationItemImpl::GetSuccessFailImageVisibility )
+				+ SHorizontalBox::Slot()
+				.Padding(10.f, 0.f, 0.f, 0.f)
 				[
-					SNew(SImage)
-					.Image( this, &SNotificationItemImpl::GetSuccessFailImage )
+					InteractiveWidgetsBox
 				]
 			];
-		}
-
-		return HorizontalBox;
 	}
 
 	/** Sets the text and delegate for the hyperlink */
@@ -609,7 +673,7 @@ protected:
 	FSimpleDelegate Hyperlink;
 
 	/** Text to display for the hyperlink message */
-	TAttribute< FText > HyperlinkText;
+	TAttribute<FText> HyperlinkText;
 };
 
 /////////////////////////////////////////////////
@@ -631,6 +695,8 @@ public:
 		SLATE_ATTRIBUTE(float, FadeOutDuration)
 		/** The duration before a fadeout for this element */
 		SLATE_ATTRIBUTE(float, ExpireDuration)
+		/** When set this forces the width of the notification */
+		SLATE_ARGUMENT(FOptionalSize, WidthOverride)
 		/** The widget that provides the notification to display */
 		SLATE_ARGUMENT(TSharedPtr<INotificationWidget>, ContentWidget)
 
@@ -650,8 +716,37 @@ public:
 		ExpireDuration = InArgs._ExpireDuration;
 		NotificationWidget = InArgs._ContentWidget;
 
-		ChildSlot
-			[
+		TSharedRef<SWidget> Internals = SNullWidget::NullWidget;
+
+		if (NotificationWidget->UseNotificationBackground())
+		{
+			Internals =
+				/*SNew(SOverlay)
+				+ SOverlay::Slot()
+				[*/
+					SNew(SNotificationBackground)
+					.Padding(FMargin(16, 8))
+					.BorderBackgroundColor(this, &SNotificationItemExternalImpl::GetContentColor)
+					.ColorAndOpacity(this, &SNotificationItemExternalImpl::GetContentColorRaw)
+					.DesiredSizeScale(this, &SNotificationItemExternalImpl::GetItemScale)
+					[
+						SNew(SBox)
+						.WidthOverride(InArgs._WidthOverride)
+						[
+							NotificationWidget->AsWidget()
+						]
+					];
+			/*	]
+				+SOverlay::Slot()
+				[
+					SNew(SImage)
+					.Image(FAppStyle::Get().GetBrush("NotificationList.ItemBackground_Border"))
+					.Visibility(EVisibility::HitTestInvisible)
+				];*/
+		}
+		else
+		{
+			Internals = 
 				SNew(SBorder)
 				.Padding(0.f)
 				.BorderImage(FCoreStyle::Get().GetBrush("NoBorder"))
@@ -659,21 +754,30 @@ public:
 				.ColorAndOpacity(this, &SNotificationItemExternalImpl::GetContentColorRaw)
 				.DesiredSizeScale(this, &SNotificationItemExternalImpl::GetItemScale)
 				[
-					SNew(SOverlay)
-					+ SOverlay::Slot()
+					SNew(SBox)
+					.WidthOverride(InArgs._WidthOverride)
 					[
-						NotificationWidget->AsWidget()
+						SNew(SOverlay)
+						+ SOverlay::Slot()
+						[
+							NotificationWidget->AsWidget()
+						]
+						+ SOverlay::Slot()
+						[
+							SNew(SBorder)
+							.Padding(0)
+							.BorderImage(FCoreStyle::Get().GetBrush("NotificationList.ItemBackground_Border_Transparent"))
+							.BorderBackgroundColor(this, &SNotificationItemExternalImpl::GetGlowColor)
+							.Visibility(EVisibility::SelfHitTestInvisible)
+						]
 					]
-					+ SOverlay::Slot()
-					[
-						SNew(SBorder)
-						.Padding(0)
-						.BorderImage(FCoreStyle::Get().GetBrush("NotificationList.ItemBackground_Border_Transparent"))
-						.BorderBackgroundColor(this, &SNotificationItemExternalImpl::GetGlowColor)
-						.Visibility(EVisibility::SelfHitTestInvisible)
-					]
-				]
-			];
+				];
+		}
+
+		ChildSlot
+		[
+			Internals
+		];
 	}
 
 	/** SNotificationItem interface */
@@ -702,16 +806,16 @@ TSharedRef<SNotificationItem> SNotificationList::AddNotification(const FNotifica
 				.ContentWidget(Info.ContentWidget)
 				.FadeInDuration(Info.FadeInDuration)
 				.ExpireDuration(Info.ExpireDuration)
-				.FadeOutDuration(Info.FadeOutDuration);
+				.FadeOutDuration(Info.FadeOutDuration)
+				.WidthOverride(Info.WidthOverride);
 		}
 		else
 		{
-			static const FSlateBrush* CachedImage = FCoreStyle::Get().GetBrush("NotificationList.DefaultMessage");
+			static const FSlateBrush* CachedImage = FAppStyle::Get().GetBrush("AppIconFlat");
 
 			// Create notification.
 			NewItem = SNew(SNotificationItemImpl)
 				.Text(Info.Text)
-				.Font(Font)
 				.ButtonDetails(Info.ButtonDetails)
 				.Image((Info.Image != nullptr) ? Info.Image : CachedImage)
 				.FadeInDuration(Info.FadeInDuration)
@@ -776,7 +880,6 @@ void SNotificationList::NotificationItemFadedOut (const TSharedRef<SNotification
 void SNotificationList::Construct(const FArguments& InArgs)
 {
 	bDone = false;
-	Font = InArgs._Font;
 
 	ChildSlot
 	[
diff --git a/Engine/Source/Runtime/Slate/Public/Framework/Application/SlateApplication.h b/Engine/Source/Runtime/Slate/Public/Framework/Application/SlateApplication.h
--- a/Engine/Source/Runtime/Slate/Public/Framework/Application/SlateApplication.h
+++ b/Engine/Source/Runtime/Slate/Public/Framework/Application/SlateApplication.h
@@ -1354,6 +1354,9 @@ public:
 	}
 
 	virtual TSharedPtr<SWindow> GetActiveTopLevelWindow() const override;
+
+	virtual TSharedPtr<SWindow> GetActiveTopLevelRegularWindow() const override;
+
 	virtual const FSlateBrush* GetAppIcon() const override;
 	virtual const FSlateBrush* GetAppIconSmall() const override;
 
diff --git a/Engine/Source/Runtime/Slate/Public/Framework/Notifications/NotificationManager.h b/Engine/Source/Runtime/Slate/Public/Framework/Notifications/NotificationManager.h
--- a/Engine/Source/Runtime/Slate/Public/Framework/Notifications/NotificationManager.h
+++ b/Engine/Source/Runtime/Slate/Public/Framework/Notifications/NotificationManager.h
@@ -14,6 +14,61 @@ struct FNotificationInfo;
 
 template <class T> class TLockFreePointerListLIFO;
 
+/** 
+ * Handle to an active progress notification.
+ * Used to update the notification
+ */
+struct FProgressNotificationHandle
+{
+	friend class FSlateNotificationManager;
+
+	FProgressNotificationHandle()
+		: Id(INDEX_NONE)
+	{}
+
+	void Reset() { Id = INDEX_NONE; }
+	bool IsValid() const { return Id != INDEX_NONE; }
+
+	bool operator==(const FProgressNotificationHandle& OtherHandle) const
+	{
+		return Id == OtherHandle.Id;
+	}
+private:
+	FProgressNotificationHandle(int32 InId)
+		: Id(InId)
+	{}
+
+	int32 Id;
+};
+
+/** Base class for any handlers that display progress bars for progres notifications */
+class IProgressNotificationHandler
+{
+public:
+	virtual ~IProgressNotificationHandler() {}
+
+	/**
+	 * Called when a progress notification begins
+ 	 * @param Handle			Handle to the notification
+	 * @param DisplayText		Display text used to describe the type of work to the user
+	 * @param TotalWorkToDo		Arbitrary number of work units to perform. 
+	 */
+	virtual void StartProgressNotification(FProgressNotificationHandle Handle, FText DisplayText, int32 TotalWorkToDo) = 0;
+
+	/**
+	 * Called when a notification should be updated. 
+	 * @param InHandle				Handle to the notification that was previously created with StartProgressNotification
+	 * @param TotalWorkDone			The total number of work units done for the notification.
+ 	 * @param UpdatedTotalWorkToDo	UpdatedTotalWorkToDo. This value will be 0 if the total work did not change
+	 * @param UpdatedDisplayText	Updated display text of the notification. This value will be empty if the text did not change
+	 */
+	virtual void UpdateProgressNotification(FProgressNotificationHandle Handle, int32 TotalWorkDone, int32 UpdatedTotalWorkToDo, FText UpdatedDisplayText) = 0;
+
+	/**
+	 * Called when a notification should be cancelled
+	 */
+	virtual void CancelProgressNotification(FProgressNotificationHandle Handle) = 0;
+};
 /**
  * A class which manages a group of notification windows                 
  */
@@ -45,13 +100,44 @@ public:
 	 */
 	void QueueNotification(FNotificationInfo* Info);
 
+	/**
+	 * Begin a progress notification. These notifications should be used for background work not blocking work. Use SlowTask for blocking work
+	 * @param DisplayText		Display text used to describe the type of work to the user
+	 * @param TotalWorkToDo		Arbitrary number of work units to perform. If more than one progress notification is active this work is summed with any other active notifications. 
+	 *				This value usually represents the number of steps to be completed
+ 	 *
+	 * @return NotificationHandle	A handle that can be used to update the progress or cancel the notification
+	 */
+	FProgressNotificationHandle StartProgressNotification(FText DisplayText, int32 TotalWorkToDo);
+
+	/**
+	 * Updates a progress notification. 
+	 * @param InHandle		Handle to the notification that was previously created with StartProgressNotification
+	 * @param TotalWorkDone		The total number of work uits done for the notification. This is NOT an incremental value. 
+ 	 * @param UpdatedTotalWorkToDo	(Optional) If the total work to do changes you can update the value here. Note that if the total work left increases but the total work done does not, the progress bar may go backwards
+	 * @param UpdatedDisplayText	(Optional) Updated display text of the notification
+	 * @return NotificationHandle	A handle that can be used to update the progress or cancel the notification
+	 */
+	void UpdateProgressNotification(FProgressNotificationHandle InHandle, int32 TotalWorkDone, int32 UpdatedTotalWorkToDo = 0, FText UpdatedDisplayText = FText::GetEmpty());
+
+	/**
+	 * Cancels an active notification
+	 */
+	void CancelProgressNotification(FProgressNotificationHandle InHandle);
+
+	/**
+	 * Sets the progress notification handler for the current application. Only one handler is used at a time. 
+	 * This handler is not managed in any way. If your handler is being destroyed call SetProgressNotificationHandler(nullptr)
+	 */
+	void SetProgressNotificationHandler(IProgressNotificationHandler* NewHandler);
+
 	/**
 	 * Called back from the SlateApplication when a window is activated/resized
 	 * We need to keep notifications topmost in the z-order so we manage it here directly
 	 * as there isn't a cross-platform OS-level way of making a 'topmost child'.
 	 * @param ActivateEvent 	Information about the activation event
 	 */
-	void ForceNotificationsInFront( const TSharedRef<SWindow> InWindow );
+	void ForceNotificationsInFront(const TSharedRef<SWindow> InWindow);
 
 	/**
 	 * Gets all the windows that represent notifications
@@ -85,7 +171,7 @@ protected:
 	FSlateNotificationManager();
 
 	/** Create a notification list for the specified screen rectangle */
-	TSharedRef<SNotificationList> CreateStackForArea(const FSlateRect& InRectangle);
+	TSharedRef<SNotificationList> CreateStackForArea(const FSlateRect& InRectangle, TSharedPtr<SWindow> Window);
 
 	/** FCoreDelegates::OnPreExit shutdown callback */
 	void ShutdownOnPreExit();
@@ -111,15 +197,21 @@ private:
 		FSlateRect Region;
 	};
 
+	/** Handle to a system which updates progress notifications (e.g a status bar) */
+	IProgressNotificationHandler* ProgressNotificationHandler = nullptr;
+
 	/** A window under which all of the notification windows will nest. */
 	TWeakPtr<SWindow> RootWindowPtr;
 
 	/** An array of notification lists grouped by work area regions */
-	TArray< FRegionalNotificationList > RegionalLists;
+	TArray<FRegionalNotificationList> RegionalLists;
 
 	/** Thread safe queue of notifications to display */
 	TLockFreePointerListLIFO<FNotificationInfo> PendingNotifications;
 
+	/** Counter used to create progress handles */
+	static int32 ProgressHandleCounter;
+
 	/** Whether notifications should be displayed or not.  This can be used to globally suppress notification pop-ups */
 	bool bAllowNotifications = true;
 
diff --git a/Engine/Source/Runtime/Slate/Public/Widgets/Notifications/INotificationWidget.h b/Engine/Source/Runtime/Slate/Public/Widgets/Notifications/INotificationWidget.h
--- a/Engine/Source/Runtime/Slate/Public/Widgets/Notifications/INotificationWidget.h
+++ b/Engine/Source/Runtime/Slate/Public/Widgets/Notifications/INotificationWidget.h
@@ -11,4 +11,5 @@ class SLATE_API INotificationWidget
 public:
 	virtual void OnSetCompletionState(SNotificationItem::ECompletionState State) = 0;
 	virtual TSharedRef< SWidget > AsWidget() = 0;
+	virtual bool UseNotificationBackground() const { return false; }
 };
diff --git a/Engine/Source/Runtime/Slate/Public/Widgets/Notifications/SNotificationList.h b/Engine/Source/Runtime/Slate/Public/Widgets/Notifications/SNotificationList.h
--- a/Engine/Source/Runtime/Slate/Public/Widgets/Notifications/SNotificationList.h
+++ b/Engine/Source/Runtime/Slate/Public/Widgets/Notifications/SNotificationList.h
@@ -107,7 +107,7 @@ struct FNotificationInfo
 	 *
 	 * @param	InText	Text string to display for this notification
 	 */
-	FNotificationInfo( const FText& InText )
+	FNotificationInfo(const FText& InText)
 		: ContentWidget(),
 		Text(InText),
 		ButtonDetails(),
@@ -118,7 +118,7 @@ struct FNotificationInfo
 		bUseThrobber(true),
 		bUseSuccessFailIcons(true),
 		bUseLargeFont(true),
-		WidthOverride(),
+		WidthOverride(320.0f),
 		bFireAndForget(true),
 		CheckBoxState(ECheckBoxState::Unchecked),
 		CheckBoxStateChanged(),
@@ -144,7 +144,7 @@ struct FNotificationInfo
 		bUseThrobber(false),
 		bUseSuccessFailIcons(false),
 		bUseLargeFont(false),
-		WidthOverride(),
+		WidthOverride(320.0f),
 		bFireAndForget(true),
 		CheckBoxState(ECheckBoxState::Unchecked),
 		CheckBoxStateChanged(),
@@ -191,19 +191,22 @@ struct FNotificationInfo
 	bool bFireAndForget;
 
 	/** When set this will display a check box on the notification; handles getting the current check box state */
-	TAttribute< ECheckBoxState > CheckBoxState;
+	TAttribute<ECheckBoxState> CheckBoxState;
 
 	/** When set this will display a check box on the notification; handles setting the new check box state */
 	FOnCheckStateChanged CheckBoxStateChanged;
 
 	/** Text to display for the check box message */
-	TAttribute< FText > CheckBoxText;
+	TAttribute<FText> CheckBoxText;
 
 	/** When set this will display as a hyperlink on the right side of the notification. */
 	FSimpleDelegate Hyperlink;
 
 	/** Text to display for the hyperlink message */
-	TAttribute< FText > HyperlinkText;
+	TAttribute<FText> HyperlinkText;
+
+	/** A specific window to put the notification in. If this is null, the root window of the application will be used */
+	TSharedPtr<SWindow> ForWindow;
 
 	/** True if we should throttle the editor while the notification is transitioning and performance is poor, to make sure the user can see the animation */
 	bool bAllowThrottleWhenFrameRateIsLow;
@@ -222,8 +225,6 @@ class SLATE_API SNotificationList
 public:
 
 	SLATE_BEGIN_ARGS( SNotificationList ){}
-		/** Sets the font used to draw the text */
-		SLATE_ATTRIBUTE(FSlateFontInfo, Font)
 	SLATE_END_ARGS()
 	
 	/**
@@ -257,9 +258,6 @@ protected:
 	/** The parent window of this list. */
 	TWeakPtr<SWindow> ParentWindowPtr;
 
-	/** Holds passed in font */
-	TAttribute<FSlateFontInfo> Font;
-
 	/** Flag to auto-destroy this list. */
 	bool bDone;
 
diff --git a/Engine/Source/Runtime/SlateCore/Private/Styling/StarshipCoreStyle.cpp b/Engine/Source/Runtime/SlateCore/Private/Styling/StarshipCoreStyle.cpp
--- a/Engine/Source/Runtime/SlateCore/Private/Styling/StarshipCoreStyle.cpp
+++ b/Engine/Source/Runtime/SlateCore/Private/Styling/StarshipCoreStyle.cpp
@@ -282,6 +282,7 @@ TSharedRef<ISlateStyle> FStarshipCoreStyle::Create()
 		Style->Set("AppIconPadding", FMargin(5.f, 5.f, 5.f, 5.f));
 #endif
 
+		Style->Set("AppIconFlat", new IMAGE_BRUSH_SVG("Starship/Common/unreal", FVector2D(32, 32), FStyleColors::Foreground));
 		Style->Set("AppIcon.Small", new IMAGE_BRUSH_SVG("Starship/Common/unreal-small", Icon24x24, FStyleColors::Foreground));
 		Style->Set("AppIconPadding.Small", FMargin(4.f, 4.f, 0.f, 0.f));
 
@@ -293,8 +294,8 @@ TSharedRef<ISlateStyle> FStarshipCoreStyle::Create()
 	
 		Style->Set("Icons.Download", new IMAGE_BRUSH("Icons/icon_Downloads_16x", Icon16x16));
 
-		Style->Set("Icons.Error", new IMAGE_BRUSH_SVG("Starship/Common/alert-circle", Icon16x16));
-		Style->Set("Icons.ErrorWithColor", new IMAGE_BRUSH_SVG("Starship/Common/alert-circle", Icon16x16, FStyleColors::Error));
+		Style->Set("Icons.Error", new IMAGE_BRUSH_SVG("Starship/Common/alert-triangle", Icon16x16));
+		Style->Set("Icons.ErrorWithColor", new IMAGE_BRUSH_SVG("Starship/Common/alert-triangle", Icon16x16, FStyleColors::Error));
 
 		Style->Set("Icons.Warning", new IMAGE_BRUSH_SVG("Starship/Common/alert-triangle", Icon16x16));
 		Style->Set("Icons.WarningWithColor", new IMAGE_BRUSH_SVG("Starship/Common/alert-triangle", Icon16x16, FStyleColors::Warning));
@@ -302,6 +303,18 @@ TSharedRef<ISlateStyle> FStarshipCoreStyle::Create()
 		Style->Set("Icons.Info", new IMAGE_BRUSH_SVG("Starship/Common/Info", Icon16x16));
 		Style->Set("Icons.InfoWithColor", new IMAGE_BRUSH_SVG("Starship/Common/Info", Icon16x16, FStyleColors::Foreground));
 
+		Style->Set("Icons.Success", new IMAGE_BRUSH_SVG("Starship/Common/check-circle", Icon16x16));
+		Style->Set("Icons.SuccessWithColor", new IMAGE_BRUSH_SVG("Starship/Common/check-circle", Icon16x16, FStyleColors::Success));
+
+		Style->Set("Icons.Error.Large", new IMAGE_BRUSH_SVG("Starship/Common/alert-triangle-large", Icon32x32));
+		Style->Set("Icons.ErrorWithColor.Large", new IMAGE_BRUSH_SVG("Starship/Common/alert-triangle-large", Icon32x32, FStyleColors::Error));
+
+		Style->Set("Icons.Warning.Large", new IMAGE_BRUSH_SVG("Starship/Common/alert-triangle-large", Icon32x32));
+		Style->Set("Icons.WarningWithColor.Large", new IMAGE_BRUSH_SVG("Starship/Common/alert-triangle-large", Icon32x32, FStyleColors::Warning));
+
+		Style->Set("Icons.Success.Large", new IMAGE_BRUSH_SVG("Starship/Common/check-circle-large", Icon32x32));
+		Style->Set("Icons.SuccessWithColor.Large", new IMAGE_BRUSH_SVG("Starship/Common/check-circle-large", Icon32x32, FStyleColors::Success));
+
 		Style->Set("Icons.box-perspective", new IMAGE_BRUSH_SVG("Starship/Common/box-perspective", Icon16x16));
 		Style->Set("Icons.cylinder", new IMAGE_BRUSH_SVG("Starship/Common/cylinder", Icon16x16));
 		Style->Set("Icons.pyramid", new IMAGE_BRUSH_SVG("Starship/Common/pyriamid", Icon16x16));
@@ -599,14 +612,25 @@ TSharedRef<ISlateStyle> FStarshipCoreStyle::Create()
 
 	// SNotificationList defaults...
 	{
-		Style->Set("NotificationList.FontBold", DEFAULT_FONT("Bold", 16));
+		Style->Set("NotificationList.FontBold", DEFAULT_FONT("Bold", 12));
 		Style->Set("NotificationList.FontLight", DEFAULT_FONT("Light", 12));
-		Style->Set("NotificationList.ItemBackground", new BOX_BRUSH("Old/Menu_Background", FMargin(8.0f / 64.0f)));
-		Style->Set("NotificationList.ItemBackground_Border", new BOX_BRUSH("Old/Menu_Background_Inverted_Border_Bold", FMargin(8.0f / 64.0f)));
+		Style->Set("NotificationList.ItemBackground", new FSlateColorBrush(FStyleColors::Notifications));
+		
+		Style->Set("NotificationList.ItemBackground_Border", new FSlateRoundedBoxBrush(FStyleColors::Transparent, 5.0f, FStyleColors::Black, 1.0f));
 		Style->Set("NotificationList.ItemBackground_Border_Transparent", new BOX_BRUSH("Old/Notification_Border_Flash", FMargin(8.0f / 64.0f)));
+
+		Style->Set("NotificationList.Throbber", new IMAGE_BRUSH("Starship/Notifications/Throbber", FVector2D(40, 40)));
+
+		Style->Set("NotificationList.Watermark", new IMAGE_BRUSH_SVG("Starship/Notifications/Watermark", FVector2D(128, 128)));
+
+		// Should be deleted. everywhere using this is looking for a generic "alert" icon
+		Style->Set("NotificationList.DefaultMessage", new IMAGE_BRUSH("Common/EventMessage_Default", Icon40x40));
+
+		// Should be deleted. everywhere using this should be using Icons.Success
 		Style->Set("NotificationList.SuccessImage", new IMAGE_BRUSH("Icons/notificationlist_success", Icon16x16));
 		Style->Set("NotificationList.FailImage", new IMAGE_BRUSH("Icons/notificationlist_fail", Icon16x16));
-		Style->Set("NotificationList.DefaultMessage", new IMAGE_BRUSH("Common/EventMessage_Default", Icon40x40));
+
+	
 	}
 
 	// SSeparator defaults...
@@ -689,6 +713,10 @@ TSharedRef<ISlateStyle> FStarshipCoreStyle::Create()
 
 		Style->Set("Window", Window);
 		
+		Window.SetCornerRadius(8.0f);
+
+		Style->Set("NotificationWindow", Window);
+
 		Style->Set("ChildWindow.Background", new FSlateColorBrush(FStyleColors::Recessed));
 	}
 
diff --git a/Engine/Source/Runtime/SlateCore/Private/Styling/StyleColors.cpp b/Engine/Source/Runtime/SlateCore/Private/Styling/StyleColors.cpp
--- a/Engine/Source/Runtime/SlateCore/Private/Styling/StyleColors.cpp
+++ b/Engine/Source/Runtime/SlateCore/Private/Styling/StyleColors.cpp
@@ -44,6 +44,8 @@ const FSlateColor FStyleColors::Select = EStyleColor::Select;
 const FSlateColor FStyleColors::SelectInactive = EStyleColor::SelectInactive;
 const FSlateColor FStyleColors::SelectParent = EStyleColor::SelectParent;
 const FSlateColor FStyleColors::SelectHover = EStyleColor::SelectHover;
+
+const FSlateColor FStyleColors::Notifications = EStyleColor::Notifications;
 // if select ==  primary shouldnt we have a select pressed which is the same as primary press?
 
 const FSlateColor FStyleColors::AccentBlue = EStyleColor::AccentBlue;
@@ -61,6 +63,7 @@ const FSlateColor FStyleColors::AccentFolder = EStyleColor::AccentFolder;
 
 const FSlateColor FStyleColors::Warning= EStyleColor::Warning;
 const FSlateColor FStyleColors::Error = EStyleColor::Error;
+const FSlateColor FStyleColors::Success = EStyleColor::Success;
 
 USlateThemeManager::USlateThemeManager()
 {
@@ -104,6 +107,9 @@ void USlateThemeManager::InitalizeDefaults()
 	SetDefaultColor(EStyleColor::SelectInactive, COLOR("#99B3BFFF"));
 	SetDefaultColor(EStyleColor::SelectParent, COLOR("#2C323AFF"));
 	SetDefaultColor(EStyleColor::SelectHover, GetDefaultColor(EStyleColor::Panel));
+
+	SetDefaultColor(EStyleColor::Notifications, COLOR("464B50FF"));
+
 	// if select ==  primary shouldnt we have a select pressed which is the same as primary press?
 
 	SetDefaultColor(EStyleColor::AccentBlue, COLOR("#26BBFFFF"));
@@ -121,6 +127,7 @@ void USlateThemeManager::InitalizeDefaults()
 
 	SetDefaultColor(EStyleColor::Warning, COLOR("#FFB800FF"));
 	SetDefaultColor(EStyleColor::Error, COLOR("#EF3535FF"));
+	SetDefaultColor(EStyleColor::Success, COLOR("#1FE44BFF"));
 }
 
 #if ALLOW_THEMES
diff --git a/Engine/Source/Runtime/SlateCore/Private/Widgets/SWindow.cpp b/Engine/Source/Runtime/SlateCore/Private/Widgets/SWindow.cpp
--- a/Engine/Source/Runtime/SlateCore/Private/Widgets/SWindow.cpp
+++ b/Engine/Source/Runtime/SlateCore/Private/Widgets/SWindow.cpp
@@ -433,6 +433,7 @@ TSharedRef<SWindow> SWindow::MakeNotificationWindow()
 {
 	TSharedRef<SWindow> NewWindow =
 		SNew(SWindow)
+		.Style(FAppStyle::Get(), "NotificationWindow")
 		.Type( EWindowType::Notification )
 		.SupportsMaximize( false )
 		.SupportsMinimize( false )
diff --git a/Engine/Source/Runtime/SlateCore/Public/Application/SlateApplicationBase.h b/Engine/Source/Runtime/SlateCore/Public/Application/SlateApplicationBase.h
--- a/Engine/Source/Runtime/SlateCore/Public/Application/SlateApplicationBase.h
+++ b/Engine/Source/Runtime/SlateCore/Public/Application/SlateApplicationBase.h
@@ -52,6 +52,9 @@ class IWindowTitleBar
 {
 public:
 	virtual void Flash( ) = 0;
+	virtual void UpdateWindowMenu(TSharedPtr<SWidget> MenuContent) = 0;
+	virtual void UpdateBackgroundContent(TSharedPtr<SWidget> BackgroundContent) = 0;
+	virtual void SetUseLargeIcon(bool bUseSmallIcon) = 0;
 };
 
 class FSlateApplicationBase;
@@ -164,6 +167,13 @@ public:
 	 */
 	virtual TSharedPtr<SWindow> GetActiveTopLevelWindow() const = 0;
 
+	/**
+	 * Gets the active regular top-level window.
+	 * A regular window is a non-menu, non-tooltip, non-cursor decorator window
+	 * @return The top level window, or nullptr if no Slate windows are currently active.
+	 */
+	virtual TSharedPtr<SWindow> GetActiveTopLevelRegularWindow() const = 0;
+
 	/**
 	 * Get a list of all top-level windows in the application, excluding virtual windows.
 	 *
@@ -177,6 +187,7 @@ public:
 	 * @return The icon.
 	 */
 	virtual const FSlateBrush* GetAppIcon( ) const = 0;
+	virtual const FSlateBrush* GetAppIconSmall() const = 0;
 
 	/**
 	 * Gets the ratio SlateUnit / ScreenPixel.
diff --git a/Engine/Source/Runtime/SlateCore/Public/Styling/StyleColors.h b/Engine/Source/Runtime/SlateCore/Public/Styling/StyleColors.h
--- a/Engine/Source/Runtime/SlateCore/Public/Styling/StyleColors.h
+++ b/Engine/Source/Runtime/SlateCore/Public/Styling/StyleColors.h
@@ -47,6 +47,7 @@ enum class EStyleColor : uint8
 	SelectInactive,
 	SelectParent,
 	SelectHover,
+	Notifications,
 	AccentBlue,
 	AccentPurple,
 	AccentPink,
@@ -61,6 +62,7 @@ enum class EStyleColor : uint8
 	AccentFolder,
 	Warning,
 	Error,
+	Success,
 
 	/** Only user colors should be below this line
 	 * To use user colors:
@@ -315,6 +317,8 @@ struct SLATECORE_API FStyleColors
 	static const FSlateColor SelectParent;
 	static const FSlateColor SelectHover;
 
+	static const FSlateColor Notifications;
+
 	static const FSlateColor AccentBlue;
 	static const FSlateColor AccentPurple;
 	static const FSlateColor AccentPink;
@@ -330,4 +334,5 @@ struct SLATECORE_API FStyleColors
 
 	static const FSlateColor Warning;
 	static const FSlateColor Error;
+	static const FSlateColor Success;
 };
