Fix deprecation warnings

diff --git a/Engine/Source/Runtime/Slate/Private/Widgets/Notifications/SNotificationList.cpp b/Engine/Source/Runtime/Slate/Private/Widgets/Notifications/SNotificationList.cpp
--- a/Engine/Source/Runtime/Slate/Private/Widgets/Notifications/SNotificationList.cpp
+++ b/Engine/Source/Runtime/Slate/Private/Widgets/Notifications/SNotificationList.cpp
@@ -366,7 +366,7 @@ public:
 				AllottedGeometry.ToPaintGeometry(FVector2D(4.0f, 4.0f), FVector2D(WatermarkSize, WatermarkSize)),
 				WatermarkBrush,
 				ESlateDrawEffect::None,
-				WatermarkTint * InWidgetStyle.GetColorAndOpacityTint() * BorderBackgroundColor.Get().GetColor(InWidgetStyle));
+				WatermarkTint * InWidgetStyle.GetColorAndOpacityTint() * GetBorderBackgroundColor().GetColor(InWidgetStyle));
 		}
 
 
