Prevent notification throbber image from appearing if bUseThrobber is false when a notification is created

diff --git a/Engine/Source/Runtime/Slate/Private/Widgets/Notifications/SNotificationList.cpp b/Engine/Source/Runtime/Slate/Private/Widgets/Notifications/SNotificationList.cpp
--- a/Engine/Source/Runtime/Slate/Private/Widgets/Notifications/SNotificationList.cpp
+++ b/Engine/Source/Runtime/Slate/Private/Widgets/Notifications/SNotificationList.cpp
@@ -581,7 +581,7 @@ public:
 					[
 						SNew(SImage)
 						.Image(FAppStyle::Get().GetBrush("NotificationList.Throbber"))
-						.Visibility(this, &SNotificationItemImpl::GetThrobberVisibility)
+						.Visibility(InArgs._bUseThrobber ? MakeAttributeSP(this, &SNotificationItemImpl::GetThrobberVisibility) : EVisibility::Collapsed)
 						.RenderTransform(this, &SNotificationItemImpl::GetThrobberTransform)
 						.RenderTransformPivot(FVector2D(.5f,.5f))
 					] 
