Slate: Fix simple implicit type conversion warnings in Slate

diff --git a/Engine/Source/Runtime/Slate/Private/Framework/Docking/SDockingTabStack.cpp b/Engine/Source/Runtime/Slate/Private/Framework/Docking/SDockingTabStack.cpp
--- a/Engine/Source/Runtime/Slate/Private/Framework/Docking/SDockingTabStack.cpp
+++ b/Engine/Source/Runtime/Slate/Private/Framework/Docking/SDockingTabStack.cpp
@@ -831,7 +831,7 @@ TSharedPtr<FTabManager::FLayoutNode> SDockingTabStack::GatherPersistentLayout()
 void SDockingTabStack::ClearReservedSpace()
 {
 	bShowingTitleBarArea = false;
-	(*TitleBarSlot).Padding(0);
+	(*TitleBarSlot).Padding(0.f);
 }
 
 void SDockingTabStack::ReserveSpaceForWindowChrome(EChromeElement Element, bool bIncludePaddingForMenuBar, bool bOnlyMinorTabs)
@@ -842,9 +842,9 @@ void SDockingTabStack::ReserveSpaceForWindowChrome(EChromeElement Element, bool
 	#else
 		static const float TopPaddingForMenuBar = 27.0f;
 		static const float LeftPaddingForIcon = FSlateApplication::Get().GetAppIcon()->GetImageSize().X;
-		const FMargin ControlsPadding = FMargin(8, bIncludePaddingForMenuBar ? TopPaddingForMenuBar : 5, 0, 0);
+		const FMargin ControlsPadding = FMargin(8.f, bIncludePaddingForMenuBar ? TopPaddingForMenuBar : 5.f, 0.f, 0.f);
 		// If we are including top padding for the menu bar we do not need to pad the left side since we will be below the left icon
-		const FMargin IconPadding = FMargin(bIncludePaddingForMenuBar ? LeftPaddingForIcon + 12 : 25, bOnlyMinorTabs ? 5 : 0, 0, 0);
+		const FMargin IconPadding = FMargin(bIncludePaddingForMenuBar ? LeftPaddingForIcon + 12.f : 25.f, bOnlyMinorTabs ? 5.f : 0.f, 0.f, 0.f);
 	#endif
 
 	bShowingTitleBarArea = true;
diff --git a/Engine/Source/Runtime/Slate/Private/Framework/MultiBox/MultiBox.cpp b/Engine/Source/Runtime/Slate/Private/Framework/MultiBox/MultiBox.cpp
--- a/Engine/Source/Runtime/Slate/Private/Framework/MultiBox/MultiBox.cpp
+++ b/Engine/Source/Runtime/Slate/Private/Framework/MultiBox/MultiBox.cpp
@@ -704,7 +704,7 @@ void SMultiBoxWidget::AddBlockWidget(const FMultiBlock& Block, TSharedPtr<SHoriz
 				Slot.AutoWidth();
 			}
 
-			Slot.Padding(0);
+			Slot.Padding(0.f);
 			Slot
 			[
 				FinalWidgetWithHook
diff --git a/Engine/Source/Runtime/Slate/Private/Framework/MultiBox/MultiBoxCustomization.cpp b/Engine/Source/Runtime/Slate/Private/Framework/MultiBox/MultiBoxCustomization.cpp
--- a/Engine/Source/Runtime/Slate/Private/Framework/MultiBox/MultiBoxCustomization.cpp
+++ b/Engine/Source/Runtime/Slate/Private/Framework/MultiBox/MultiBoxCustomization.cpp
@@ -16,7 +16,7 @@ void SCustomToolbarPreviewWidget::BuildMultiBlockWidget(const ISlateStyle* Style
 	ChildSlot
 	[
 		SNew( SBorder )
-		.Padding(0)
+		.Padding(0.f)
 		.BorderImage( FCoreStyle::Get().GetBrush("NoBorder") )
 		.HAlign(HAlign_Center)
 		.VAlign(VAlign_Center)
diff --git a/Engine/Source/Runtime/Slate/Private/Framework/MultiBox/SClippingHorizontalBox.cpp b/Engine/Source/Runtime/Slate/Private/Framework/MultiBox/SClippingHorizontalBox.cpp
--- a/Engine/Source/Runtime/Slate/Private/Framework/MultiBox/SClippingHorizontalBox.cpp
+++ b/Engine/Source/Runtime/Slate/Private/Framework/MultiBox/SClippingHorizontalBox.cpp
@@ -99,7 +99,7 @@ void SClippingHorizontalBox::AddWrapButton()
 		SNew( SComboButton )
 		.HasDownArrow( false )
 		.ButtonStyle(&ToolBarStyle.ButtonStyle)
-		.ContentPadding( 0 )
+		.ContentPadding( 0.f )
 		.ToolTipText( NSLOCTEXT("Slate", "ExpandToolbar", "Click to expand toolbar") )
 		.OnGetMenuContent( OnWrapButtonClicked )
 		.Cursor( EMouseCursor::Default )
@@ -113,7 +113,7 @@ void SClippingHorizontalBox::AddWrapButton()
 	// Add the wrap button
 	AddSlot()
 	.AutoWidth()
-	.Padding( 0 )
+	.Padding( 0.f )
 	[
 		WrapButton.ToSharedRef()
 	];
diff --git a/Engine/Source/Runtime/Slate/Private/Framework/MultiBox/SToolBarButtonBlock.cpp b/Engine/Source/Runtime/Slate/Private/Framework/MultiBox/SToolBarButtonBlock.cpp
--- a/Engine/Source/Runtime/Slate/Private/Framework/MultiBox/SToolBarButtonBlock.cpp
+++ b/Engine/Source/Runtime/Slate/Private/Framework/MultiBox/SToolBarButtonBlock.cpp
@@ -261,7 +261,7 @@ void SToolBarButtonBlock::BuildMultiBlockWidget(const ISlateStyle* StyleSet, con
 		[
 			// Create a button
 			SNew(SButton)
-			.ContentPadding(0)
+			.ContentPadding(0.f)
 			// Use the tool bar item style for this button
 			.ButtonStyle(&ButtonStyle)
 			.IsEnabled(this, &SToolBarButtonBlock::IsEnabled)
diff --git a/Engine/Source/Runtime/Slate/Private/Framework/MultiBox/SToolBarComboButtonBlock.cpp b/Engine/Source/Runtime/Slate/Private/Framework/MultiBox/SToolBarComboButtonBlock.cpp
--- a/Engine/Source/Runtime/Slate/Private/Framework/MultiBox/SToolBarComboButtonBlock.cpp
+++ b/Engine/Source/Runtime/Slate/Private/Framework/MultiBox/SToolBarComboButtonBlock.cpp
@@ -193,7 +193,7 @@ void SToolBarComboButtonBlock::BuildMultiBlockWidget(const ISlateStyle* StyleSet
 	[
 		SAssignNew(ComboButtonWidget, SComboButton)
 		.AddMetaData<FTagMetaData>(FTagMetaData(TagName))
-		.ContentPadding(0)
+		.ContentPadding(0.f)
 		.ComboButtonStyle(ComboStyle)
 		.ButtonStyle(&ButtonStyle)
 		.ToolTipText(ToolBarComboButtonBlock->ToolTip)
diff --git a/Engine/Source/Runtime/Slate/Private/Widgets/Docking/SDockTab.cpp b/Engine/Source/Runtime/Slate/Private/Widgets/Docking/SDockTab.cpp
--- a/Engine/Source/Runtime/Slate/Private/Widgets/Docking/SDockTab.cpp
+++ b/Engine/Source/Runtime/Slate/Private/Widgets/Docking/SDockTab.cpp
@@ -388,7 +388,7 @@ SDockTab::SDockTab()
 	, OnTabClosed()
 	, OnTabActivated()
 	, OnCanCloseTab()
-	, ContentAreaPadding( 2 )
+	, ContentAreaPadding( 2.f )
 	, bShouldAutosize(false)
 	, TabColorScale(FLinearColor::Transparent)
 	, LastActivationTime(0.0)
diff --git a/Engine/Source/Runtime/Slate/Private/Widgets/Input/SExpandableButton.cpp b/Engine/Source/Runtime/Slate/Private/Widgets/Input/SExpandableButton.cpp
--- a/Engine/Source/Runtime/Slate/Private/Widgets/Input/SExpandableButton.cpp
+++ b/Engine/Source/Runtime/Slate/Private/Widgets/Input/SExpandableButton.cpp
@@ -53,7 +53,7 @@ void SExpandableButton::Construct(const FArguments& InArgs)
 				.Visibility( this, &SExpandableButton::GetCollapsedVisibility )
 				.OnClicked( InArgs._OnExpansionClicked )
 				.ButtonStyle( FCoreStyle::Get(), "NoBorder" )
-				.ContentPadding( 0 )
+				.ContentPadding( 0.f )
 				[
 					CollapsedButtonContent
 				]
@@ -66,7 +66,7 @@ void SExpandableButton::Construct(const FArguments& InArgs)
 				SNew(SButton)
 				.Visibility( this, &SExpandableButton::GetExpandedVisibility )
 				.ButtonStyle( FCoreStyle::Get(), "NoBorder" )
-				.ContentPadding( 0 )
+				.ContentPadding( 0.f )
 				.VAlign(VAlign_Center)
 				[
 					ExpandedButtonContent
@@ -90,7 +90,7 @@ void SExpandableButton::Construct(const FArguments& InArgs)
 				// Close expansion button
 				SNew(SButton)
 				.ButtonStyle( FCoreStyle::Get(), "NoBorder" )
-				.ContentPadding( 0 )
+				.ContentPadding( 0.f )
 				.Visibility( this, &SExpandableButton::GetExpandedVisibility )
 				.OnClicked(InArgs._OnCloseClicked)
 				[
diff --git a/Engine/Source/Runtime/Slate/Private/Widgets/Input/SMenuAnchor.cpp b/Engine/Source/Runtime/Slate/Private/Widgets/Input/SMenuAnchor.cpp
--- a/Engine/Source/Runtime/Slate/Private/Widgets/Input/SMenuAnchor.cpp
+++ b/Engine/Source/Runtime/Slate/Private/Widgets/Input/SMenuAnchor.cpp
@@ -284,7 +284,7 @@ bool SMenuAnchor::IsOpenViaCreatedWindow() const
 void SMenuAnchor::SetContent(TSharedRef<SWidget> InContent)
 {
 	Children[0]
-	.Padding(0)
+	.Padding(0.f)
 	[
 		InContent
 	];
diff --git a/Engine/Source/Runtime/Slate/Private/Widgets/Input/SRotatorInputBox.cpp b/Engine/Source/Runtime/Slate/Private/Widgets/Input/SRotatorInputBox.cpp
--- a/Engine/Source/Runtime/Slate/Private/Widgets/Input/SRotatorInputBox.cpp
+++ b/Engine/Source/Runtime/Slate/Private/Widgets/Input/SRotatorInputBox.cpp
@@ -20,7 +20,7 @@ void SRotatorInputBox::Construct( const SRotatorInputBox::FArguments& InArgs )
 			.AllowSpin(InArgs._AllowSpin)
 			.MinSliderValue(0.0f)
 			.MaxSliderValue(359.999f)
-			.LabelPadding(FMargin(3))
+			.LabelPadding(FMargin(3.0f))
 			.LabelLocation(SNumericEntryBox<float>::ELabelLocation::Inside)
 			.Label()
 			[
@@ -48,7 +48,7 @@ void SRotatorInputBox::Construct( const SRotatorInputBox::FArguments& InArgs )
 			.AllowSpin(InArgs._AllowSpin)
 			.MinSliderValue(0.0f)
 			.MaxSliderValue(359.999f)
-			.LabelPadding(FMargin(3))
+			.LabelPadding(FMargin(3.0f))
 			.LabelLocation(SNumericEntryBox<float>::ELabelLocation::Inside)
 			.Label()
 			[
@@ -76,7 +76,7 @@ void SRotatorInputBox::Construct( const SRotatorInputBox::FArguments& InArgs )
 			.AllowSpin(InArgs._AllowSpin)
 			.MinSliderValue(0.0f)
 			.MaxSliderValue(359.999f)
-			.LabelPadding(FMargin(3))
+			.LabelPadding(FMargin(3.0f))
 			.LabelLocation(SNumericEntryBox<float>::ELabelLocation::Inside)
 			.Label()
 			[
diff --git a/Engine/Source/Runtime/Slate/Private/Widgets/Input/SSearchBox.cpp b/Engine/Source/Runtime/Slate/Private/Widgets/Input/SSearchBox.cpp
--- a/Engine/Source/Runtime/Slate/Private/Widgets/Input/SSearchBox.cpp
+++ b/Engine/Source/Runtime/Slate/Private/Widgets/Input/SSearchBox.cpp
@@ -138,7 +138,7 @@ void SSearchBox::Construct( const FArguments& InArgs )
 		SNew(SButton)
 		.Visibility(this, &SSearchBox::GetXVisibility)
 		.ButtonStyle(FAppStyle::Get(), "SimpleButton")
-		.ContentPadding(0)
+		.ContentPadding(0.f)
 		.HAlign(HAlign_Center)
 		.VAlign(VAlign_Center)
 		.OnClicked(this, &SSearchBox::OnClearSearch)
diff --git a/Engine/Source/Runtime/Slate/Private/Widgets/Input/STextComboPopup.cpp b/Engine/Source/Runtime/Slate/Private/Widgets/Input/STextComboPopup.cpp
--- a/Engine/Source/Runtime/Slate/Private/Widgets/Input/STextComboPopup.cpp
+++ b/Engine/Source/Runtime/Slate/Private/Widgets/Input/STextComboPopup.cpp
@@ -19,7 +19,7 @@ void STextComboPopup::Construct( const FArguments& InArgs )
 	[
 		SNew(SBorder)
 		. BorderImage(FCoreStyle::Get().GetBrush("PopupText.Background"))
-		. Padding(10)
+		. Padding(10.f)
 		[
 			SNew(SVerticalBox)
 			+SVerticalBox::Slot()
diff --git a/Engine/Source/Runtime/Slate/Private/Widgets/Input/STextEntryPopup.cpp b/Engine/Source/Runtime/Slate/Private/Widgets/Input/STextEntryPopup.cpp
--- a/Engine/Source/Runtime/Slate/Private/Widgets/Input/STextEntryPopup.cpp
+++ b/Engine/Source/Runtime/Slate/Private/Widgets/Input/STextEntryPopup.cpp
@@ -16,7 +16,7 @@ void STextEntryPopup::Construct( const FArguments& InArgs )
 	[
 		SNew(SBorder)
 		. BorderImage(FCoreStyle::Get().GetBrush("PopupText.Background"))
-		. Padding(10)
+		. Padding(10.f)
 		[
 			SNew(SHorizontalBox)
 			+SHorizontalBox::Slot()
diff --git a/Engine/Source/Runtime/Slate/Private/Widgets/Input/SVectorInputBox.cpp b/Engine/Source/Runtime/Slate/Private/Widgets/Input/SVectorInputBox.cpp
--- a/Engine/Source/Runtime/Slate/Private/Widgets/Input/SVectorInputBox.cpp
+++ b/Engine/Source/Runtime/Slate/Private/Widgets/Input/SVectorInputBox.cpp
@@ -61,7 +61,7 @@ void SVectorInputBox::ConstructX(const FArguments& InArgs, TSharedRef<SHorizonta
 		.Delta(InArgs._SpinDelta)
 		.OnBeginSliderMovement(InArgs._OnBeginSliderMovement)
 		.OnEndSliderMovement(InArgs._OnEndSliderMovement)
-		.LabelPadding(FMargin(3))
+		.LabelPadding(FMargin(3.f))
 		.LabelLocation(SNumericEntryBox<float>::ELabelLocation::Inside)
 		.Label()
 		[
diff --git a/Engine/Source/Runtime/Slate/Private/Widgets/Layout/SScrollBar.cpp b/Engine/Source/Runtime/Slate/Private/Widgets/Layout/SScrollBar.cpp
--- a/Engine/Source/Runtime/Slate/Private/Widgets/Layout/SScrollBar.cpp
+++ b/Engine/Source/Runtime/Slate/Private/Widgets/Layout/SScrollBar.cpp
@@ -38,7 +38,7 @@ void SScrollBar::Construct(const FArguments& InArgs)
 				.BorderImage(BackgroundBrush)
 				.HAlign(HorizontalAlignment)
 				.VAlign(VerticalAlignment)
-				.Padding(0)
+				.Padding(0.f)
 				[
 					SAssignNew(Track, SScrollBarTrack)
 					.Orientation(InArgs._Orientation)
@@ -57,7 +57,7 @@ void SScrollBar::Construct(const FArguments& InArgs)
 						SAssignNew(DragThumb, SBorder)
 						.HAlign(HAlign_Center)
 						.VAlign(VAlign_Center)
-						.Padding(0)
+						.Padding(0.f)
 						[
 							SAssignNew(ThicknessSpacer, SSpacer)
 							.Size(InArgs._Thickness)
diff --git a/Engine/Source/Runtime/Slate/Private/Widgets/Notifications/SPopUpErrorText.cpp b/Engine/Source/Runtime/Slate/Private/Widgets/Notifications/SPopUpErrorText.cpp
--- a/Engine/Source/Runtime/Slate/Private/Widgets/Notifications/SPopUpErrorText.cpp
+++ b/Engine/Source/Runtime/Slate/Private/Widgets/Notifications/SPopUpErrorText.cpp
@@ -8,7 +8,7 @@ void SPopupErrorText::Construct( const FArguments& InArgs )
 	SComboButton::Construct( SComboButton::FArguments()
 		.ComboButtonStyle( FCoreStyle::Get(), "MessageLogListingComboButton" )
 		.HasDownArrow(false)
-		.ContentPadding(0)
+		.ContentPadding(0.f)
 		.IsFocusable(false)
 		.ButtonContent()
 		[
diff --git a/Engine/Source/Runtime/Slate/Private/Widgets/Views/SHeaderRow.cpp b/Engine/Source/Runtime/Slate/Private/Widgets/Views/SHeaderRow.cpp
--- a/Engine/Source/Runtime/Slate/Private/Widgets/Views/SHeaderRow.cpp
+++ b/Engine/Source/Runtime/Slate/Private/Widgets/Views/SHeaderRow.cpp
@@ -459,7 +459,7 @@ void SHeaderRow::Construct( const FArguments& InArgs )
 	}
 
 	SBorder::Construct( SBorder::FArguments()
-		.Padding( 0 )
+		.Padding( 0.f )
 		.BorderImage( &Style->BackgroundBrush )
 		.ForegroundColor( Style->ForegroundColor )
 	);
