Slate: Fix simple implicit type conversion warnings in Slate

diff --git a/Engine/Source/Runtime/Slate/Private/Framework/Application/SlateUser.cpp b/Engine/Source/Runtime/Slate/Private/Framework/Application/SlateUser.cpp
--- a/Engine/Source/Runtime/Slate/Private/Framework/Application/SlateUser.cpp
+++ b/Engine/Source/Runtime/Slate/Private/Framework/Application/SlateUser.cpp
@@ -1309,7 +1309,7 @@ void FSlateUser::UpdateTooltip(const FMenuStack& MenuStack, bool bCanSpawnNewToo
 
 		// How long since the tooltip was summoned?
 		const double PlatformSeconds = FPlatformTime::Seconds();
-		const float TimeSinceSummon = PlatformSeconds - TooltipSummonDelay - ActiveTooltipInfo.SummonTime;
+		const float TimeSinceSummon = (float)(PlatformSeconds - TooltipSummonDelay - ActiveTooltipInfo.SummonTime);
 		const float TooltipOpacity = FMath::Clamp<float>(TimeSinceSummon / TooltipIntroDuration, 0.0f, 1.0f);
 
 		// Update window opacity
diff --git a/Engine/Source/Runtime/Slate/Private/Framework/Docking/FDockingDragOperation.cpp b/Engine/Source/Runtime/Slate/Private/Framework/Docking/FDockingDragOperation.cpp
--- a/Engine/Source/Runtime/Slate/Private/Framework/Docking/FDockingDragOperation.cpp
+++ b/Engine/Source/Runtime/Slate/Private/Framework/Docking/FDockingDragOperation.cpp
@@ -101,16 +101,16 @@ FSlateRect FDockingDragOperation::GetPreviewAreaForDirection ( const FSlateRect&
 	switch( DockingDirection )
 	{
 		case SDockingNode::LeftOf:
-			TargetRect.Right = TargetRect.Left + Size.X * 0.5;
+			TargetRect.Right = TargetRect.Left + Size.X * 0.5f;
 		break;
 		case SDockingNode::Above:
-			TargetRect.Bottom = TargetRect.Top + Size.Y * 0.5;
+			TargetRect.Bottom = TargetRect.Top + Size.Y * 0.5f;
 		break;
 		case SDockingNode::RightOf:
-			TargetRect.Left = TargetRect.Left + Size.X * 0.5;
+			TargetRect.Left = TargetRect.Left + Size.X * 0.5f;
 		break;
 		case SDockingNode::Below:
-			TargetRect.Top = TargetRect.Top + Size.Y * 0.5;
+			TargetRect.Top = TargetRect.Top + Size.Y * 0.5f;
 		break;
 		case SDockingNode::Center:
 		break;
diff --git a/Engine/Source/Runtime/Slate/Private/Framework/Docking/TabManager.cpp b/Engine/Source/Runtime/Slate/Private/Framework/Docking/TabManager.cpp
--- a/Engine/Source/Runtime/Slate/Private/Framework/Docking/TabManager.cpp
+++ b/Engine/Source/Runtime/Slate/Private/Framework/Docking/TabManager.cpp
@@ -227,8 +227,8 @@ TSharedRef<FTabManager::FLayoutNode> FTabManager::FLayout::NewFromString_Helper(
 			{
 				FVector2D WindowSize;
 
-				WindowSize.X = JsonObject->GetNumberField( TEXT("WindowSize_X") );
-				WindowSize.Y = JsonObject->GetNumberField( TEXT("WindowSize_Y") );
+				WindowSize.X = (float)JsonObject->GetNumberField( TEXT("WindowSize_X") );
+				WindowSize.Y = (float)JsonObject->GetNumberField( TEXT("WindowSize_Y") );
 
 				NewArea = FTabManager::NewArea( WindowSize );
 			}
@@ -239,11 +239,11 @@ TSharedRef<FTabManager::FLayoutNode> FTabManager::FLayout::NewFromString_Helper(
 				FVector2D WindowPosition = FVector2D::ZeroVector;
 				FVector2D WindowSize;
 
-				WindowPosition.X = JsonObject->GetNumberField( TEXT("WindowPosition_X") );
-				WindowPosition.Y = JsonObject->GetNumberField( TEXT("WindowPosition_Y") );
+				WindowPosition.X = (float)JsonObject->GetNumberField( TEXT("WindowPosition_X") );
+				WindowPosition.Y = (float)JsonObject->GetNumberField( TEXT("WindowPosition_Y") );
 
-				WindowSize.X = JsonObject->GetNumberField( TEXT("WindowSize_X") );
-				WindowSize.Y = JsonObject->GetNumberField( TEXT("WindowSize_Y") );
+				WindowSize.X = (float)JsonObject->GetNumberField( TEXT("WindowSize_X") );
+				WindowSize.Y = (float)JsonObject->GetNumberField( TEXT("WindowSize_Y") );
 
 				bool bIsMaximized = JsonObject->GetBoolField(TEXT("bIsMaximized"));
 
@@ -253,7 +253,7 @@ TSharedRef<FTabManager::FLayoutNode> FTabManager::FLayout::NewFromString_Helper(
 			break;
 		}
 		
-		NewArea->SetSizeCoefficient( JsonObject->GetNumberField( TEXT("SizeCoefficient") ) );
+		NewArea->SetSizeCoefficient((float)JsonObject->GetNumberField( TEXT("SizeCoefficient") ) );
 		NewArea->SetOrientation( local::OrientationFromString( JsonObject->GetStringField(TEXT("Orientation")) ) );
 
 		TArray< TSharedPtr<FJsonValue> > ChildNodeValues = JsonObject->GetArrayField(TEXT("nodes"));
@@ -267,7 +267,7 @@ TSharedRef<FTabManager::FLayoutNode> FTabManager::FLayout::NewFromString_Helper(
 	else if ( NodeType == TEXT("Splitter") )
 	{
 		TSharedRef<FTabManager::FSplitter> NewSplitter =  FTabManager::NewSplitter();
-		NewSplitter->SetSizeCoefficient( JsonObject->GetNumberField(TEXT("SizeCoefficient")) );
+		NewSplitter->SetSizeCoefficient((float)JsonObject->GetNumberField(TEXT("SizeCoefficient")) );
 		NewSplitter->SetOrientation( local::OrientationFromString( JsonObject->GetStringField(TEXT("Orientation")) ) );
 		TArray< TSharedPtr<FJsonValue> > ChildNodeValues = JsonObject->GetArrayField(TEXT("nodes"));
 		for( int32 ChildIndex=0; ChildIndex < ChildNodeValues.Num(); ++ChildIndex )
@@ -279,7 +279,7 @@ TSharedRef<FTabManager::FLayoutNode> FTabManager::FLayout::NewFromString_Helper(
 	else if ( NodeType == TEXT("Stack") )
 	{
 		TSharedRef<FTabManager::FStack> NewStack = FTabManager::NewStack();
-		NewStack->SetSizeCoefficient( JsonObject->GetNumberField(TEXT("SizeCoefficient")) );
+		NewStack->SetSizeCoefficient((float)JsonObject->GetNumberField(TEXT("SizeCoefficient")) );
 		NewStack->SetHideTabWell( JsonObject->GetBoolField(TEXT("HideTabWell")) );
 
 		if(JsonObject->HasField(TEXT("ForegroundTab")))
@@ -297,7 +297,7 @@ TSharedRef<FTabManager::FLayoutNode> FTabManager::FLayout::NewFromString_Helper(
 			TabId = FGlobalTabmanager::Get()->GetTabTypeForPotentiallyLegacyTab(TabId);
 
 			FString SidebarLocation;
-			double SidebarSizeCoefficient = .15f;
+			float SidebarSizeCoefficient = .15f;
 			if (TabAsJson->TryGetStringField(TEXT("SidebarLocation"), SidebarLocation))
 			{
 				TabAsJson->TryGetNumberField(TEXT("SidebarCoeff"), SidebarSizeCoefficient);
@@ -326,7 +326,7 @@ TSharedPtr<FTabManager::FLayout> FTabManager::FLayout::NewFromString( const FStr
 	{
 		const FString LayoutName = JsonObject->GetStringField(TEXT("Name"));
 		TSharedRef<FTabManager::FLayout> NewLayout = FTabManager::NewLayout( *LayoutName );
-		int32 PrimaryAreaIndex = FMath::TruncToInt( JsonObject->GetNumberField(TEXT("PrimaryAreaIndex")) );
+		int32 PrimaryAreaIndex = FMath::TruncToInt((float)JsonObject->GetNumberField(TEXT("PrimaryAreaIndex")) );
 
 		TArray< TSharedPtr<FJsonValue> > Areas = JsonObject->GetArrayField(TEXT("Areas"));
 		for(int32 AreaIndex=0; AreaIndex < Areas.Num(); ++AreaIndex)
diff --git a/Engine/Source/Runtime/Slate/Private/Widgets/Layout/SBackgroundBlur.cpp b/Engine/Source/Runtime/Slate/Private/Widgets/Layout/SBackgroundBlur.cpp
--- a/Engine/Source/Runtime/Slate/Private/Widgets/Layout/SBackgroundBlur.cpp
+++ b/Engine/Source/Runtime/Slate/Private/Widgets/Layout/SBackgroundBlur.cpp
@@ -143,7 +143,7 @@ int32 SBackgroundBlur::OnPaint(const FPaintArgs& Args, const FGeometry& Allotted
 				{
 					OutDrawElements.PushClip(FSlateClippingZone(AllottedGeometry));
 
-					FSlateDrawElement::MakePostProcessPass(OutDrawElements, LayerId, PaintGeometry, FVector4(KernelSize, ComputedStrength, RenderTargetWidth, RenderTargetHeight), DownsampleAmount);
+					FSlateDrawElement::MakePostProcessPass(OutDrawElements, LayerId, PaintGeometry, FVector4((float)KernelSize, ComputedStrength, (float)RenderTargetWidth, (float)RenderTargetHeight), DownsampleAmount);
 
 					OutDrawElements.PopClip();
 				}
diff --git a/Engine/Source/Runtime/Slate/Private/Widgets/Layout/SRadialBox.cpp b/Engine/Source/Runtime/Slate/Private/Widgets/Layout/SRadialBox.cpp
--- a/Engine/Source/Runtime/Slate/Private/Widgets/Layout/SRadialBox.cpp
+++ b/Engine/Source/Runtime/Slate/Private/Widgets/Layout/SRadialBox.cpp
@@ -40,7 +40,7 @@ void SRadialBox::Construct( const FArguments& InArgs )
 	bUseAllottedWidth = InArgs._UseAllottedWidth;
 	StartingAngle = InArgs._StartingAngle;
 	bDistributeItemsEvenly = InArgs._bDistributeItemsEvenly;
-	SectorCentralAngle = NormalizeAngle(InArgs._SectorCentralAngle);
+	SectorCentralAngle = InternalNormalizeAngle(InArgs._SectorCentralAngle);
 	AngleBetweenItems = InArgs._AngleBetweenItems;
 
 	for ( int32 ChildIndex=0; ChildIndex < InArgs.Slots.Num(); ++ChildIndex )
@@ -178,3 +178,9 @@ int32 SRadialBox::NormalizeAngle(int32 Angle) const
 	int32 NormalizedAngle = Angle % 360;
 	return NormalizedAngle < 0 ? NormalizedAngle + 360 : NormalizedAngle;
 }
+
+float SRadialBox::InternalNormalizeAngle(float Angle) const
+{
+	float NormalizedAngle = FMath::Fmod(Angle, 360.f);
+	return NormalizedAngle < 0.f ? NormalizedAngle + 360.f : NormalizedAngle;
+}
diff --git a/Engine/Source/Runtime/Slate/Private/Widgets/Layout/SUniformWrapPanel.cpp b/Engine/Source/Runtime/Slate/Private/Widgets/Layout/SUniformWrapPanel.cpp
--- a/Engine/Source/Runtime/Slate/Private/Widgets/Layout/SUniformWrapPanel.cpp
+++ b/Engine/Source/Runtime/Slate/Private/Widgets/Layout/SUniformWrapPanel.cpp
@@ -86,7 +86,7 @@ void SUniformWrapPanel::OnArrangeChildren( const FGeometry& AllottedGeometry, FA
 
 				if (row == (NumRows - 1))
 				{
-					float NumLastRowColumns = NumVisibleChildren % AdjNumColumns != 0 ? NumVisibleChildren % AdjNumColumns : AdjNumColumns;
+					int32 NumLastRowColumns = (NumVisibleChildren % AdjNumColumns != 0) ? NumVisibleChildren % AdjNumColumns : AdjNumColumns;
 					if (HAlign.Get() == HAlign_Right)
 					{
 						LeftSlop = FMath::FloorToFloat(AllottedGeometry.GetLocalSize().X - (CellSize.X * NumLastRowColumns));
diff --git a/Engine/Source/Runtime/Slate/Private/Widgets/Text/STextBlock.cpp b/Engine/Source/Runtime/Slate/Private/Widgets/Text/STextBlock.cpp
--- a/Engine/Source/Runtime/Slate/Private/Widgets/Text/STextBlock.cpp
+++ b/Engine/Source/Runtime/Slate/Private/Widgets/Text/STextBlock.cpp
@@ -254,11 +254,11 @@ FVector2D STextBlock::ComputeDesiredSize(float LayoutScaleMultiplier) const
 	{
 		const FVector2D LocalShadowOffset = GetShadowOffset();
 
-		const float LocalOutlineSize = GetFont().OutlineSettings.OutlineSize;
+		const int32 LocalOutlineSize = GetFont().OutlineSettings.OutlineSize;
 
 		// Account for the outline width impacting both size of the text by multiplying by 2
 		// Outline size in Y is accounted for in MaxHeight calculation in Measure()
-		const FVector2D ComputedOutlineSize(LocalOutlineSize * 2, LocalOutlineSize);
+		const FVector2D ComputedOutlineSize((float)(LocalOutlineSize * 2), (float)LocalOutlineSize);
 		const FVector2D TextSize = FSlateApplication::Get().GetRenderer()->GetFontMeasureService()->Measure(GetText(), GetFont()) + ComputedOutlineSize + LocalShadowOffset;
 
 		CachedSimpleDesiredSize = FVector2D(FMath::Max(MinDesiredWidth.Get(0.0f), TextSize.X), TextSize.Y);
diff --git a/Engine/Source/Runtime/Slate/Public/Widgets/Input/SNumericEntryBox.h b/Engine/Source/Runtime/Slate/Public/Widgets/Input/SNumericEntryBox.h
--- a/Engine/Source/Runtime/Slate/Public/Widgets/Input/SNumericEntryBox.h
+++ b/Engine/Source/Runtime/Slate/Public/Widgets/Input/SNumericEntryBox.h
@@ -72,7 +72,7 @@ public:
 		, _Label()
 		, _LabelVAlign( VAlign_Fill )
 		, _LabelLocation(ELabelLocation::Outside)
-		, _LabelPadding(FMargin(3,0) )
+		, _LabelPadding(FMargin(3.f,0.f) )
 		, _BorderForegroundColor(FAppStyle::Get().GetWidgetStyle<FSpinBoxStyle>("NumericEntrySpinBox").ForegroundColor)
 		, _BorderBackgroundColor(FLinearColor::White)
 		, _UndeterminedString( SNumericEntryBox<NumericType>::DefaultUndeterminedString )
@@ -80,13 +80,13 @@ public:
 		, _ShiftMouseMovePixelPerDelta(1)
 		, _SupportDynamicSliderMaxValue(false)
 		, _SupportDynamicSliderMinValue(false)
-		, _Delta(0)
+		, _Delta(NumericType(0))
 		, _MinValue(TNumericLimits<NumericType>::Lowest())
 		, _MaxValue(TNumericLimits<NumericType>::Max())
-		, _MinSliderValue(0)				
-		, _MaxSliderValue(100)
+		, _MinSliderValue(NumericType(0))				
+		, _MaxSliderValue(NumericType(100))
 		, _SliderExponent(1.f)
-		, _MinDesiredValueWidth(0)		
+		, _MinDesiredValueWidth(0.f)
 	{}		
 
 		/** Style to use for the editable text box within this widget */
@@ -274,7 +274,7 @@ public:
 				.BorderImage(this, &SNumericEntryBox<NumericType>::GetBorderImage)
 				.BorderBackgroundColor(InArgs._BorderBackgroundColor)
 				.ForegroundColor(InArgs._BorderForegroundColor)
-				.Padding(0)
+				.Padding(0.f)
 				[
 					Overlay
 				]
@@ -299,7 +299,7 @@ public:
 					.BorderImage(this, &SNumericEntryBox<NumericType>::GetBorderImage)
 					.BorderBackgroundColor(InArgs._BorderBackgroundColor)
 					.ForegroundColor(InArgs._BorderForegroundColor)
-					.Padding(0)
+					.Padding(0.f)
 					[
 						Overlay
 					]
@@ -318,7 +318,7 @@ public:
 			.ForegroundColor(ForegroundColor)
 			.VAlign(VAlign_Center)
 			.HAlign(HAlign_Left)
-			.Padding(FMargin(1, 0, 6, 0))
+			.Padding(FMargin(1.f, 0.f, 6.f, 0.f))
 			[
 				SNew(STextBlock)
 				.Text(LabelText)
diff --git a/Engine/Source/Runtime/Slate/Public/Widgets/Layout/SRadialBox.h b/Engine/Source/Runtime/Slate/Public/Widgets/Layout/SRadialBox.h
--- a/Engine/Source/Runtime/Slate/Public/Widgets/Layout/SRadialBox.h
+++ b/Engine/Source/Runtime/Slate/Public/Widgets/Layout/SRadialBox.h
@@ -99,10 +99,14 @@ public:
 	void SetUseAllottedWidth(bool bInUseAllottedWidth);
 
 	/** Mods the angle so it's between 0-360 */
+	UE_DEPRECATED(5.0, "NormalizeAngle is deprecated. You should use the FRotator::NormalizeAxis.")
 	int32 NormalizeAngle(int32 Angle) const;
 
 private:
 
+	/** Mods the angle so it's between 0-360 */
+	float InternalNormalizeAngle(float Angle) const;
+
 	/** How wide this panel should appear to be. */
 	TAttribute<float> PreferredWidth;
 
