Fix thumb position calculation for virtual joysticks

#ue5

diff --git a/Engine/Source/Runtime/Slate/Private/Widgets/Input/SVirtualJoystick.cpp b/Engine/Source/Runtime/Slate/Private/Widgets/Input/SVirtualJoystick.cpp
--- a/Engine/Source/Runtime/Slate/Private/Widgets/Input/SVirtualJoystick.cpp
+++ b/Engine/Source/Runtime/Slate/Private/Widgets/Input/SVirtualJoystick.cpp
@@ -276,7 +276,7 @@ FVector2D SVirtualJoystick::ComputeThumbPosition(int32 ControlIndex, const FVect
 	const FControlData& Control = Controls[ControlIndex];
 
 	// figure out position around center
-	FVector2D Offset = LocalCoord - Controls[ControlIndex].VisualCenter;
+	FVector2D Offset = LocalCoord - Control.VisualCenter;
 	// only do work if we aren't at the center
 	if (Offset == FVector2D(0, 0))
 	{
@@ -293,13 +293,14 @@ FVector2D SVirtualJoystick::ComputeThumbPosition(int32 ControlIndex, const FVect
 		float SinAngle = FMath::Sin(Angle);
 		float XTerm = CosAngle / (Control.CorrectedVisualSize.X * 0.5f);
 		float YTerm = SinAngle / (Control.CorrectedVisualSize.Y * 0.5f);
-		DistanceToEdgeSqr = XTerm * XTerm + YTerm * YTerm;
+		float XYTermSqr = XTerm * XTerm + YTerm * YTerm;
+		DistanceToEdgeSqr = 1.0f / XYTermSqr;
 
 		// only clamp 
 		if (DistanceToTouchSqr > DistanceToEdgeSqr)
 		{
-			float DistanceToEdge = FMath::InvSqrt(DistanceToEdgeSqr);
-			Position = FVector2D(DistanceToEdge * CosAngle,  DistanceToEdge * SinAngle);
+			float DistanceToEdge = FMath::InvSqrt(XYTermSqr);
+			Position = FVector2D(DistanceToEdge * CosAngle, DistanceToEdge * SinAngle);
 		}
 		else
 		{
