Fixed out of bounds when clicking on a >> menu extender, due to no commands.

diff --git a/Engine/Source/Runtime/Slate/Public/Framework/MultiBox/MultiBox.h b/Engine/Source/Runtime/Slate/Public/Framework/MultiBox/MultiBox.h
--- a/Engine/Source/Runtime/Slate/Public/Framework/MultiBox/MultiBox.h
+++ b/Engine/Source/Runtime/Slate/Public/Framework/MultiBox/MultiBox.h
@@ -398,7 +398,7 @@ public:
 	bool bHasSearchWidget;
 
 	/* Returns the last command list used */
-	const TSharedPtr<const FUICommandList> GetLastCommandList() const { return CommandLists.Last(); }
+	const TSharedPtr<const FUICommandList> GetLastCommandList() const { return CommandLists.Num() > 0 ? CommandLists.Last() : nullptr; }
 
 private:
 	
