Improved logging messages for editor application mode and tab manager tracking, in the search for the cause of temporary "unrecognized tab" instances.

diff --git a/Engine/Source/Editor/UnrealEd/Private/ApplicationMode.cpp b/Engine/Source/Editor/UnrealEd/Private/ApplicationMode.cpp
--- a/Engine/Source/Editor/UnrealEd/Private/ApplicationMode.cpp
+++ b/Engine/Source/Editor/UnrealEd/Private/ApplicationMode.cpp
@@ -28,7 +28,7 @@ FApplicationMode::FApplicationMode(FName InModeName, FText(*GetLocalizedMode)(co
 
 void FApplicationMode::DeactivateMode(TSharedPtr<FTabManager> InTabManager)
 {
-	UE_LOG(LogApplicationMode, Log, TEXT("Deactivating ApplicationMode '%s'"), *ModeName.ToString());
+	UE_LOG(LogApplicationMode, Log, TEXT("Deactivating ApplicationMode '%s': 0x%p with TabManager: 0x%p"), *ModeName.ToString(), this, InTabManager.Get());
 
 	// Save the layout to INI
 	check(InTabManager.IsValid());
@@ -49,7 +49,7 @@ void FApplicationMode::DeactivateMode(TSharedPtr<FTabManager> InTabManager)
 
 TSharedRef<FTabManager::FLayout> FApplicationMode::ActivateMode(TSharedPtr<FTabManager> InTabManager)
 {
-	UE_LOG(LogApplicationMode, Log, TEXT("Activating ApplicationMode '%s'"), *ModeName.ToString());
+	UE_LOG(LogApplicationMode, Log, TEXT("Activating ApplicationMode '%s': 0x%p with TabManager: 0x%p"), *ModeName.ToString(), this, InTabManager.Get());
 
 	check(InTabManager.IsValid());
 	RegisterTabFactories(InTabManager);
diff --git a/Engine/Source/Runtime/Slate/Private/Framework/Docking/TabManager.cpp b/Engine/Source/Runtime/Slate/Private/Framework/Docking/TabManager.cpp
--- a/Engine/Source/Runtime/Slate/Private/Framework/Docking/TabManager.cpp
+++ b/Engine/Source/Runtime/Slate/Private/Framework/Docking/TabManager.cpp
@@ -862,7 +862,7 @@ void FTabManager::CloseAllAreas()
 
 TSharedRef<FTabManager::FLayout> FTabManager::PersistLayout() const
 {
-	UE_LOG(LogSlate, Log, TEXT("Saving tab layout to %s"), *ActiveLayoutName.ToString());
+	UE_LOG(LogSlate, Log, TEXT("Saving tab layout to %s in TabManager 0x%p"), *ActiveLayoutName.ToString(), this);
 
 	TSharedRef<FLayout> PersistentLayout = FTabManager::NewLayout( this->ActiveLayoutName );
 	
@@ -918,7 +918,7 @@ FTabSpawnerEntry& FTabManager::RegisterTabSpawner(const FName TabId, const FOnSp
 	ensure(!TabSpawner.Contains(TabId));
 	ensure(!FGlobalTabmanager::Get()->IsLegacyTabType(TabId));
 
-	UE_LOG(LogSlate, Log, TEXT("Registering tab spawner for layout '%s': '%s'"), *ActiveLayoutName.ToString(), *TabId.ToString());
+	UE_LOG(LogSlate, Log, TEXT("Registering tab spawner for '%s' in layout '%s' on TabManager 0x%p"), *TabId.ToString(), *ActiveLayoutName.ToString(), this);
 
 	TSharedRef<FTabSpawnerEntry> NewSpawnerEntry = MakeShareable(new FTabSpawnerEntry(TabId, OnSpawnTab, CanSpawnTab));
 	TabSpawner.Add(TabId, NewSpawnerEntry);
@@ -927,7 +927,7 @@ FTabSpawnerEntry& FTabManager::RegisterTabSpawner(const FName TabId, const FOnSp
 
 bool FTabManager::UnregisterTabSpawner( const FName TabId )
 {
-	UE_LOG(LogSlate, Log, TEXT("Unregistering tab spawner for layout '%s': '%s'"), *ActiveLayoutName.ToString(), *TabId.ToString());
+	UE_LOG(LogSlate, Log, TEXT("Unregistering tab spawner for '%s' in layout '%s' on TabManager 0x%p"), *TabId.ToString(), *ActiveLayoutName.ToString(), this);
 
 	return TabSpawner.Remove( TabId ) > 0;
 }
@@ -937,7 +937,7 @@ void FTabManager::UnregisterAllTabSpawners()
 	if (TabSpawner.Num())
 	{
 		FString AllTabIdsString = FString::JoinBy(TabSpawner, TEXT(", "), [](const TPair<FName, TSharedRef<FTabSpawnerEntry>> Entry) { return Entry.Key.ToString(); });
-		UE_LOG(LogSlate, Log, TEXT("Unregistering all tab spawners for layout '%s': %s"), *ActiveLayoutName.ToString(), *AllTabIdsString);
+		UE_LOG(LogSlate, Log, TEXT("Unregistering all tab spawners on TabManager 0x%p for layout '%s': %s"), this, *ActiveLayoutName.ToString(), *AllTabIdsString);
 	}
 
 	TabSpawner.Empty();
@@ -948,7 +948,7 @@ TSharedPtr<SWidget> FTabManager::RestoreFrom(const TSharedRef<FLayout>& Layout,
 {
 	ActiveLayoutName = Layout->LayoutName;
 
-	UE_LOG(LogSlate, Log, TEXT("Restoring tab layout: %s"), *ActiveLayoutName.ToString());
+	UE_LOG(LogSlate, Log, TEXT("Restoring tab layout: %s in TabManager 0x%p"), *ActiveLayoutName.ToString(), this);
 
 	TSharedPtr<SDockingArea> PrimaryDockArea;
 	for (int32 AreaIndex=0; AreaIndex < Layout->Areas.Num(); ++AreaIndex )
