Added missing copyright text

diff --git a/Engine/Source/Runtime/Slate/Private/Widgets/Notifications/SNotificationBackground.cpp b/Engine/Source/Runtime/Slate/Private/Widgets/Notifications/SNotificationBackground.cpp
--- a/Engine/Source/Runtime/Slate/Private/Widgets/Notifications/SNotificationBackground.cpp
+++ b/Engine/Source/Runtime/Slate/Private/Widgets/Notifications/SNotificationBackground.cpp
@@ -1,3 +1,5 @@
+// Copyright Epic Games, Inc. All Rights Reserved.
+
 #include "Widgets/Notifications/SNotificationBackground.h"
 #include "Styling/StyleColors.h"
 
