Fix the content browser drawer new user notification being hard coded to ctrl+tab, ignoring any user set key binding

diff --git a/Engine/Source/Editor/StatusBar/Private/StatusBarSubsystem.cpp b/Engine/Source/Editor/StatusBar/Private/StatusBarSubsystem.cpp
--- a/Engine/Source/Editor/StatusBar/Private/StatusBarSubsystem.cpp
+++ b/Engine/Source/Editor/StatusBar/Private/StatusBarSubsystem.cpp
@@ -126,7 +126,7 @@ class SNewUserTipNotification : public SCompoundWidget
 								[
 									SNew(STextBlock)
 									.TextStyle(FAppStyle::Get(), "DialogButtonText")
-									.Text(FText::FromString("CTRL"))
+									.Text(FGlobalEditorCommonCommands::Get().OpenContentBrowserDrawer->GetActiveChord(EMultipleKeyBindingIndex::Primary)->GetModifierText(FText::GetEmpty()))
 									.ColorAndOpacity(FStyleColors::ForegroundHover)
 								]
 							]
@@ -148,7 +148,7 @@ class SNewUserTipNotification : public SCompoundWidget
 								[
 									SNew(STextBlock)
 									.TextStyle(FAppStyle::Get(), "DialogButtonText")
-									.Text(FText::FromString("SPACE"))
+									.Text(FGlobalEditorCommonCommands::Get().OpenContentBrowserDrawer->GetActiveChord(EMultipleKeyBindingIndex::Primary)->Key.GetDisplayName(false))
 									.ColorAndOpacity(FStyleColors::ForegroundHover)
 								]
 							]
diff --git a/Engine/Source/Runtime/Slate/Private/Framework/Commands/InputChord.cpp b/Engine/Source/Runtime/Slate/Private/Framework/Commands/InputChord.cpp
--- a/Engine/Source/Runtime/Slate/Private/Framework/Commands/InputChord.cpp
+++ b/Engine/Source/Runtime/Slate/Private/Framework/Commands/InputChord.cpp
@@ -9,9 +9,30 @@
 
 /**
  * Returns the friendly, localized string name of this key binding
- * @todo Slate: Got to be a better way to do this
  */
 FText FInputChord::GetInputText( ) const
+{
+	FFormatNamedArguments Args;
+	Args.Add(TEXT("Modifiers"), GetModifierText());
+	Args.Add(TEXT("Key"), GetKeyText());
+
+	return FText::Format(LOCTEXT("InputText", "{Modifiers}{Key}"), Args);
+}
+
+
+FText FInputChord::GetKeyText() const
+{
+	FText OutString;
+
+	if (Key.IsValid() && !Key.IsModifierKey())
+	{
+		OutString = Key.GetDisplayName();
+	}
+
+	return OutString;
+}
+
+FText FInputChord::GetModifierText(TOptional<FText> ModifierAppender) const
 {
 #if PLATFORM_MAC
 	const FText CommandText = LOCTEXT("KeyName_Control", "Ctrl");
@@ -24,7 +45,7 @@ FText FInputChord::GetInputText( ) const
 	const FText ShiftText = LOCTEXT("KeyName_Shift", "Shift");
 
 
-	const FText AppenderText = Key != EKeys::Invalid ? LOCTEXT("ModAppender", "+") : FText::GetEmpty();
+	const FText AppenderText = Key != EKeys::Invalid ? ModifierAppender.Get(LOCTEXT("ModAppender", "+")) : FText::GetEmpty();
 
 	FFormatNamedArguments Args;
 	int32 ModCount = 0;
@@ -63,22 +84,7 @@ FText FInputChord::GetInputText( ) const
 
 	}
 
-	Args.Add(TEXT("Key"), GetKeyText());
-
-	return FText::Format(LOCTEXT("FourModifiers", "{Mod1}{Appender1}{Mod2}{Appender2}{Mod3}{Appender3}{Mod4}{Appender4}{Key}"), Args);
-}
-
-
-FText FInputChord::GetKeyText( ) const
-{
-	FText OutString; // = KeyGetDisplayName(Key);
-
-	if (Key.IsValid() && !Key.IsModifierKey())
-	{
-		OutString = Key.GetDisplayName();
-	}
-
-	return OutString;
+	return FText::Format(LOCTEXT("FourModifiers", "{Mod1}{Appender1}{Mod2}{Appender2}{Mod3}{Appender3}{Mod4}{Appender4}"), Args);
 }
 
 FInputChord::ERelationshipType FInputChord::GetRelationship( const FInputChord& OtherChord ) const
diff --git a/Engine/Source/Runtime/Slate/Public/Framework/Commands/InputChord.h b/Engine/Source/Runtime/Slate/Public/Framework/Commands/InputChord.h
--- a/Engine/Source/Runtime/Slate/Public/Framework/Commands/InputChord.h
+++ b/Engine/Source/Runtime/Slate/Public/Framework/Commands/InputChord.h
@@ -165,14 +165,21 @@ struct SLATE_API FInputChord
 	 *
 	 * @return A localized string.
 	 */
-	FText GetInputText( ) const;
+	FText GetInputText() const;
 	
 	/**
 	 * Gets the key represented as a localized string.
 	 *
 	 * @return A localized string.
 	 */
-	FText GetKeyText( ) const;
+	FText GetKeyText() const;
+
+	/**
+	 * Gets the localized string for the modifier portion of the coord. 
+	 * 
+	 * @param ModifierAppender. Text to append between each modifier and the key. If this is not specified a '+' character will be used
+	 */
+	FText GetModifierText(TOptional<FText> ModifierAppender = TOptional<FText>()) const;
 
 	/**
 	 * Checks whether this chord requires an modifier keys to be pressed.
