Slate: Convert STextBlock's attribute to SlateAttribute.

diff --git a/Engine/Source/Runtime/Slate/Private/Widgets/Text/STextBlock.cpp b/Engine/Source/Runtime/Slate/Private/Widgets/Text/STextBlock.cpp
--- a/Engine/Source/Runtime/Slate/Private/Widgets/Text/STextBlock.cpp
+++ b/Engine/Source/Runtime/Slate/Private/Widgets/Text/STextBlock.cpp
@@ -12,16 +12,68 @@
 #include "Widgets/Accessibility/SlateAccessibleWidgets.h"
 #endif
 
-DECLARE_CYCLE_STAT(TEXT("STextBlock::SetText Time"), Stat_SlateTextBlockSetText, STATGROUP_SlateVerbose)
 DECLARE_CYCLE_STAT(TEXT("STextBlock::OnPaint Time"), Stat_SlateTextBlockOnPaint, STATGROUP_SlateVerbose)
 DECLARE_CYCLE_STAT(TEXT("STextBlock::ComputeDesiredSize"), Stat_SlateTextBlockCDS, STATGROUP_SlateVerbose)
-DECLARE_CYCLE_STAT(TEXT("STextBlock::ComputeVolitility"), Stat_SlateTextBlockCV, STATGROUP_SlateVerbose)
+
+
+SLATE_IMPLEMENT_WIDGET(STextBlock)
+void STextBlock::PrivateRegisterAttributes(FSlateAttributeInitializer& AttributeInitializer)
+{
+	//~ If we are in SimpleTextMode the invalidation reason can be different.
+	struct FInvalidation
+	{
+		static EInvalidateWidgetReason GetInvalidationNoneIfSimpleTextMode(const SWidget& OwningWidget)
+		{
+			return static_cast<const STextBlock&>(OwningWidget).bSimpleTextMode ? EInvalidateWidgetReason::None : EInvalidateWidgetReason::Layout;
+		}
+		static EInvalidateWidgetReason GetInvalidationPaintIfSimpleTextMode(const SWidget& OwningWidget)
+		{
+			return static_cast<const STextBlock&>(OwningWidget).bSimpleTextMode ? EInvalidateWidgetReason::Paint : EInvalidateWidgetReason::Layout;
+		}
+	};
+
+	SLATE_ADD_MEMBER_ATTRIBUTE_DEFINITION(AttributeInitializer, BoundText, EInvalidateWidgetReason::Layout);
+	SLATE_ADD_MEMBER_ATTRIBUTE_DEFINITION(AttributeInitializer, Font, EInvalidateWidgetReason::Layout);
+	SLATE_ADD_MEMBER_ATTRIBUTE_DEFINITION(AttributeInitializer, StrikeBrush, FInvalidation::GetInvalidationNoneIfSimpleTextMode);
+	SLATE_ADD_MEMBER_ATTRIBUTE_DEFINITION(AttributeInitializer, ColorAndOpacity, FInvalidation::GetInvalidationPaintIfSimpleTextMode);
+	SLATE_ADD_MEMBER_ATTRIBUTE_DEFINITION(AttributeInitializer, ShadowOffset, EInvalidateWidgetReason::Layout);
+	SLATE_ADD_MEMBER_ATTRIBUTE_DEFINITION(AttributeInitializer, ShadowColorAndOpacity, FInvalidation::GetInvalidationPaintIfSimpleTextMode);
+	SLATE_ADD_MEMBER_ATTRIBUTE_DEFINITION(AttributeInitializer, HighlightColor, FInvalidation::GetInvalidationNoneIfSimpleTextMode);
+	SLATE_ADD_MEMBER_ATTRIBUTE_DEFINITION(AttributeInitializer, HighlightShape, FInvalidation::GetInvalidationNoneIfSimpleTextMode);
+	SLATE_ADD_MEMBER_ATTRIBUTE_DEFINITION(AttributeInitializer, HighlightText, FInvalidation::GetInvalidationNoneIfSimpleTextMode);
+	SLATE_ADD_MEMBER_ATTRIBUTE_DEFINITION(AttributeInitializer, WrapTextAt, FInvalidation::GetInvalidationNoneIfSimpleTextMode);
+	SLATE_ADD_MEMBER_ATTRIBUTE_DEFINITION(AttributeInitializer, AutoWrapText, FInvalidation::GetInvalidationNoneIfSimpleTextMode);
+	SLATE_ADD_MEMBER_ATTRIBUTE_DEFINITION(AttributeInitializer, WrappingPolicy, FInvalidation::GetInvalidationNoneIfSimpleTextMode);
+	SLATE_ADD_MEMBER_ATTRIBUTE_DEFINITION(AttributeInitializer, TransformPolicy, FInvalidation::GetInvalidationNoneIfSimpleTextMode);
+	SLATE_ADD_MEMBER_ATTRIBUTE_DEFINITION(AttributeInitializer, Margin, FInvalidation::GetInvalidationNoneIfSimpleTextMode);
+	SLATE_ADD_MEMBER_ATTRIBUTE_DEFINITION(AttributeInitializer, Justification, FInvalidation::GetInvalidationNoneIfSimpleTextMode);
+	SLATE_ADD_MEMBER_ATTRIBUTE_DEFINITION(AttributeInitializer, LineHeightPercentage, FInvalidation::GetInvalidationNoneIfSimpleTextMode);
+	SLATE_ADD_MEMBER_ATTRIBUTE_DEFINITION(AttributeInitializer, MinDesiredWidth, EInvalidateWidgetReason::Layout);
+}
 
 STextBlock::STextBlock()
+	: BoundText(*this)
+	, Font(*this)
+	, StrikeBrush(*this)
+	, ColorAndOpacity(*this)
+	, ShadowOffset(*this)
+	, ShadowColorAndOpacity(*this)
+	, HighlightColor(*this)
+	, HighlightShape(*this)
+	, HighlightText(*this)
+	, WrapTextAt(*this)
+	, AutoWrapText(*this)
+	, WrappingPolicy(*this)
+	, TransformPolicy(*this)
+	, Margin(*this)
+	, Justification(*this)
+	, LineHeightPercentage(*this)
+	, MinDesiredWidth(*this)
+	, Union_Flags(0)
+	, bSimpleTextMode(false)
 {
 	SetCanTick(false);
 	bCanSupportFocus = false;
-	bSimpleTextMode = false;
 
 #if WITH_ACCESSIBILITY
 	AccessibleBehavior = EAccessibleBehavior::Auto;
@@ -38,33 +90,31 @@ void STextBlock::Construct( const FArguments& InArgs )
 {
 	TextStyle = *InArgs._TextStyle;
 
-	HighlightText = InArgs._HighlightText;
-	WrapTextAt = InArgs._WrapTextAt;
-	AutoWrapText = InArgs._AutoWrapText;
-	WrappingPolicy = InArgs._WrappingPolicy;
+	SetHighlightText(InArgs._HighlightText);
+	SetWrapTextAt(InArgs._WrapTextAt);
+	SetAutoWrapText(InArgs._AutoWrapText);
+	SetWrappingPolicy(InArgs._WrappingPolicy);
 
-	if (InArgs._TransformPolicy.IsSet()) 
-	{
-		TransformPolicy = InArgs._TransformPolicy;
-	}
-	Margin = InArgs._Margin;
-	LineHeightPercentage = InArgs._LineHeightPercentage;
-	Justification = InArgs._Justification;
-	MinDesiredWidth = InArgs._MinDesiredWidth;
-
-	Font = InArgs._Font;
-	StrikeBrush = InArgs._StrikeBrush;
-	ColorAndOpacity = InArgs._ColorAndOpacity;
-	ShadowOffset = InArgs._ShadowOffset;
-	ShadowColorAndOpacity = InArgs._ShadowColorAndOpacity;
-	HighlightColor = InArgs._HighlightColor;
-	HighlightShape = InArgs._HighlightShape;
+	SetTransformPolicy(InArgs._TransformPolicy);
+
+	SetMargin(InArgs._Margin);
+	SetLineHeightPercentage(InArgs._LineHeightPercentage);
+	SetJustification(InArgs._Justification);
+	SetMinDesiredWidth(InArgs._MinDesiredWidth);
+
+	SetFont(InArgs._Font);
+	SetStrikeBrush(InArgs._StrikeBrush);
+	SetColorAndOpacity(InArgs._ColorAndOpacity);
+	SetShadowOffset(InArgs._ShadowOffset);
+	SetShadowColorAndOpacity(InArgs._ShadowColorAndOpacity);
+	SetHighlightColor(InArgs._HighlightColor);
+	SetHighlightShape(InArgs._HighlightShape);
 
 	bSimpleTextMode = InArgs._SimpleTextMode;
 
 	SetOnMouseDoubleClick(InArgs._OnDoubleClicked);
 
-	BoundText = InArgs._Text;
+	SetText(InArgs._Text);
 
 	//if(!bSimpleTextMode)
 	{
@@ -76,99 +126,65 @@ void STextBlock::Construct( const FArguments& InArgs )
 
 FSlateFontInfo STextBlock::GetFont() const
 {
-	return Font.IsSet() ? Font.Get() : TextStyle.Font;
+	return bIsAttributeFontSet ? Font.Get() : TextStyle.Font;
 }
 
 const FSlateBrush* STextBlock::GetStrikeBrush() const
 {
-	return StrikeBrush.IsSet() ? StrikeBrush.Get() : &TextStyle.StrikeBrush;
+	return bIsAttributeStrikeBrushSet ? StrikeBrush.Get() : &TextStyle.StrikeBrush;
 }
 
 FSlateColor STextBlock::GetColorAndOpacity() const
 {
-	return ColorAndOpacity.IsSet() ? ColorAndOpacity.Get() : TextStyle.ColorAndOpacity;
+	return bIsAttributeColorAndOpacitySet ? ColorAndOpacity.Get() : TextStyle.ColorAndOpacity;
 }
 
 FVector2D STextBlock::GetShadowOffset() const
 {
-	return ShadowOffset.IsSet() ? ShadowOffset.Get() : TextStyle.ShadowOffset;
+	return bIsAttributeShadowOffsetSet ? ShadowOffset.Get() : TextStyle.ShadowOffset;
 }
 
 FLinearColor STextBlock::GetShadowColorAndOpacity() const
 {
-	return ShadowColorAndOpacity.IsSet() ? ShadowColorAndOpacity.Get() : TextStyle.ShadowColorAndOpacity;
+	return bIsAttributeShadowColorAndOpacitySet ? ShadowColorAndOpacity.Get() : TextStyle.ShadowColorAndOpacity;
 }
 
 FSlateColor STextBlock::GetHighlightColor() const
 {
-	return HighlightColor.IsSet() ? HighlightColor.Get() : TextStyle.HighlightColor;
+	return bIsAttributeHighlightColorSet ? HighlightColor.Get() : TextStyle.HighlightColor;
 }
 
 const FSlateBrush* STextBlock::GetHighlightShape() const
 {
-	return HighlightShape.IsSet() ? HighlightShape.Get() : &TextStyle.HighlightShape;
-}
-
-void STextBlock::InvalidateText(EInvalidateWidget InvalidateReason)
-{
-	if (bSimpleTextMode && EnumHasAnyFlags(InvalidateReason, EInvalidateWidget::Layout))
-	{
-		CachedSimpleDesiredSize.Reset();
-	}
-
-	Invalidate(InvalidateReason);
-}
-
-void STextBlock::SetText( const TAttribute< FString >& InText )
-{
-	if (InText.IsSet() && !InText.IsBound())
-	{
-		SetText(FText::AsCultureInvariant(InText.Get()));
-		return;
-	}
-
-	SCOPE_CYCLE_COUNTER(Stat_SlateTextBlockSetText);
-	BoundText = MakeAttributeLambda([InText]()
-	{
-		return FText::AsCultureInvariant(InText.Get(FString()));
-	});
-	InvalidateText(EInvalidateWidget::LayoutAndVolatility);
+	return bIsAttributeHighlightShapeSet ? HighlightShape.Get() : &TextStyle.HighlightShape;
 }
 
-void STextBlock::SetText( const FString& InText )
+ETextTransformPolicy STextBlock::GetTransformPolicyImpl() const
 {
-	SetText(FText::AsCultureInvariant(InText));
+	return bIsAttributeTransformPolicySet ? TransformPolicy.Get() : TextStyle.TransformPolicy;
 }
 
-void STextBlock::SetText( const TAttribute< FText >& InText )
+void STextBlock::InvalidateText(EInvalidateWidgetReason InvalidateReason)
 {
-	if (InText.IsSet() && !InText.IsBound())
+	if (bSimpleTextMode && EnumHasAnyFlags(InvalidateReason, EInvalidateWidgetReason::Layout))
 	{
-		SetText(InText.Get());
-		return;
+		CachedSimpleDesiredSize.Reset();
 	}
 
-	SCOPE_CYCLE_COUNTER(Stat_SlateTextBlockSetText);
-	BoundText = InText;
-	InvalidateText(EInvalidateWidget::LayoutAndVolatility);
+	Invalidate(InvalidateReason);
 }
 
-void STextBlock::SetText( const FText& InText )
+void STextBlock::SetText(TAttribute<FText> InText)
 {
-	SCOPE_CYCLE_COUNTER(Stat_SlateTextBlockSetText);
-
-	if (!BoundText.IsBound() && InText.IdenticalTo(BoundText.Get(), ETextIdenticalModeFlags::DeepCompare | ETextIdenticalModeFlags::LexicalCompareInvariants))
-	{
-		return;
-	}
-
-	BoundText = InText;
-	InvalidateText(EInvalidateWidget::LayoutAndVolatility);
+	// Cache the IsBound.
+	//When the attribute is not bound, we need to go through all the other bound property to check if it is bound.
+	bIsAttributeBoundTextBound = InText.IsBound();
+	BoundText.Assign(*this, MoveTemp(InText));
 }
 
 void STextBlock::SetHighlightText(TAttribute<FText> InText)
 {
-	HighlightText = InText;
+	HighlightText.Assign(*this, MoveTemp(InText));
 }
 
 int32 STextBlock::OnPaint( const FPaintArgs& Args, const FGeometry& AllottedGeometry, const FSlateRect& MyCullingRect, FSlateWindowElementList& OutDrawElements, int32 LayerId, const FWidgetStyle& InWidgetStyle, bool bParentEnabled ) const
@@ -184,7 +200,7 @@ int32 STextBlock::OnPaint( const FPaintArgs& Args, const FGeometry& AllottedGeom
 
 		const bool bShouldBeEnabled = ShouldBeEnabled(bParentEnabled);
 
-		const FText& LocalText = GetText();
+		const FText& LocalText = BoundText.Get();
 		FSlateFontInfo LocalFont = GetFont();
 
 		if (ShouldDropShadow)
@@ -235,11 +251,11 @@ int32 STextBlock::OnPaint( const FPaintArgs& Args, const FGeometry& AllottedGeom
 		// HACK: Due to the nature of wrapping and layout, we may have been arranged in a different box than what we were cached with.  Which
 		// might update wrapping, so make sure we always set the desired size to the current size of the text layout, which may have changed
 		// during paint.
-		bool bCanWrap = WrapTextAt.Get() > 0 || AutoWrapText.Get();
+		const bool bCanWrap = WrapTextAt.Get() > 0 || AutoWrapText.Get();
 
 		if (bCanWrap && !NewDesiredSize.Equals(LastDesiredSize))
 		{
-			const_cast<STextBlock*>(this)->Invalidate(EInvalidateWidget::Layout);
+			const_cast<STextBlock*>(this)->Invalidate(EInvalidateWidgetReason::Layout);
 		}
 	}
 
@@ -254,76 +270,52 @@ FVector2D STextBlock::ComputeDesiredSize(float LayoutScaleMultiplier) const
 	{
 		const FVector2D LocalShadowOffset = GetShadowOffset();
 
-		const int32 LocalOutlineSize = GetFont().OutlineSettings.OutlineSize;
+		const float LocalOutlineSize = (float)(GetFont().OutlineSettings.OutlineSize);
 
 		// Account for the outline width impacting both size of the text by multiplying by 2
 		// Outline size in Y is accounted for in MaxHeight calculation in Measure()
-		const FVector2D ComputedOutlineSize((float)(LocalOutlineSize * 2), (float)LocalOutlineSize);
-		const FVector2D TextSize = FSlateApplication::Get().GetRenderer()->GetFontMeasureService()->Measure(GetText(), GetFont()) + ComputedOutlineSize + LocalShadowOffset;
+		const FVector2D ComputedOutlineSize(LocalOutlineSize * 2.f, LocalOutlineSize);
+		const FVector2D TextSize = FSlateApplication::Get().GetRenderer()->GetFontMeasureService()->Measure(BoundText.Get(), GetFont()) + ComputedOutlineSize + LocalShadowOffset;
 
-		CachedSimpleDesiredSize = FVector2D(FMath::Max(MinDesiredWidth.Get(0.0f), TextSize.X), TextSize.Y);
+		CachedSimpleDesiredSize = FVector2D(FMath::Max(MinDesiredWidth.Get(), TextSize.X), TextSize.Y);
 		return CachedSimpleDesiredSize.GetValue();
 	}
 	else
 	{
 		// ComputeDesiredSize will also update the text layout cache if required
 		const FVector2D TextSize = TextLayoutCache->ComputeDesiredSize(
-			FSlateTextBlockLayout::FWidgetArgs(BoundText, HighlightText, WrapTextAt, AutoWrapText, WrappingPolicy, GetTransformPolicy(), Margin, LineHeightPercentage, Justification),
-			LayoutScaleMultiplier, GetComputedTextStyle()
-		);
-
-		return FVector2D(FMath::Max(MinDesiredWidth.Get(0.0f), TextSize.X), TextSize.Y);
+			FSlateTextBlockLayout::FWidgetDesiredSizeArgs(
+				BoundText.Get(),
+				HighlightText.Get(),
+				WrapTextAt.Get(),
+				AutoWrapText.Get(),
+				WrappingPolicy.Get(),
+				GetTransformPolicyImpl(),
+				Margin.Get(),
+				LineHeightPercentage.Get(),
+				Justification.Get()),
+			LayoutScaleMultiplier, GetComputedTextStyle());
+
+		return FVector2D(FMath::Max(MinDesiredWidth.Get(), TextSize.X), TextSize.Y);
 	}
 }
 
-bool STextBlock::ComputeVolatility() const
+void STextBlock::SetFont(TAttribute<FSlateFontInfo> InFont)
 {
-	SCOPE_CYCLE_COUNTER(Stat_SlateTextBlockCV);
-	return SLeafWidget::ComputeVolatility() 
-		|| BoundText.IsBound()
-		|| Font.IsBound()
-		|| ColorAndOpacity.IsBound()
-		|| ShadowOffset.IsBound()
-		|| ShadowColorAndOpacity.IsBound()
-		|| HighlightColor.IsBound()
-		|| HighlightShape.IsBound()
-		|| HighlightText.IsBound()
-		|| WrapTextAt.IsBound()
-		|| AutoWrapText.IsBound()
-		|| WrappingPolicy.IsBound()
-		|| TransformPolicy.IsBound()
-		|| Margin.IsBound()
-		|| Justification.IsBound()
-		|| LineHeightPercentage.IsBound()
-		|| MinDesiredWidth.IsBound();
+	bIsAttributeFontSet = InFont.IsSet();
+	Font.Assign(*this, MoveTemp(InFont));
 }
 
-void STextBlock::SetFont(const TAttribute< FSlateFontInfo >& InFont)
+void STextBlock::SetStrikeBrush(TAttribute<const FSlateBrush*> InStrikeBrush)
 {
-	if(!Font.IsSet() || !Font.IdenticalTo(InFont))
-	{
-		Font = InFont;
-		InvalidateText(EInvalidateWidget::LayoutAndVolatility);
-	}
+	bIsAttributeStrikeBrushSet = InStrikeBrush.IsSet();
+	StrikeBrush.Assign(*this, MoveTemp(InStrikeBrush));
 }
 
-void STextBlock::SetStrikeBrush(const TAttribute<const FSlateBrush*>& InStrikeBrush)
+void STextBlock::SetColorAndOpacity(TAttribute<FSlateColor> InColorAndOpacity)
 {
-	if (!StrikeBrush.IsSet() || !StrikeBrush.IdenticalTo(InStrikeBrush))
-	{
-		StrikeBrush = InStrikeBrush;
-		InvalidateText(EInvalidateWidget::LayoutAndVolatility);
-	}
-}
-
-void STextBlock::SetColorAndOpacity(const TAttribute<FSlateColor>& InColorAndOpacity)
-{
-	if ( !ColorAndOpacity.IsSet() || !ColorAndOpacity.IdenticalTo(InColorAndOpacity) )
-	{
-		ColorAndOpacity = InColorAndOpacity;
-		// HACK: Normally this would be Paint only, but textblocks need to recache layout.
-		InvalidateText(EInvalidateWidget::LayoutAndVolatility);
-	}
+	bIsAttributeColorAndOpacitySet = InColorAndOpacity.IsSet();
+	ColorAndOpacity.Assign(*this, MoveTemp(InColorAndOpacity));
 }
 
 void STextBlock::SetTextStyle(const FTextBlockStyle* InTextStyle)
@@ -338,7 +330,7 @@ void STextBlock::SetTextStyle(const FTextBlockStyle* InTextStyle)
 		TextStyle = *Defaults._TextStyle;
 	}
 
-	InvalidateText(EInvalidateWidget::Layout);
+	InvalidateText(EInvalidateWidgetReason::Layout);
 }
 
 void STextBlock::SetTextShapingMethod(const TOptional<ETextShapingMethod>& InTextShapingMethod)
@@ -346,7 +338,7 @@ void STextBlock::SetTextShapingMethod(const TOptional<ETextShapingMethod>& InTex
 	if (!bSimpleTextMode)
 	{
 		TextLayoutCache->SetTextShapingMethod(InTextShapingMethod);
-		InvalidateText(EInvalidateWidget::Layout);
+		InvalidateText(EInvalidateWidgetReason::Layout);
 	}
 }
 
@@ -355,121 +347,82 @@ void STextBlock::SetTextFlowDirection(const TOptional<ETextFlowDirection>& InTex
 	if(!bSimpleTextMode)
 	{
 		TextLayoutCache->SetTextFlowDirection(InTextFlowDirection);
-		InvalidateText(EInvalidateWidget::Layout);
+		InvalidateText(EInvalidateWidgetReason::Layout);
 	}
 }
 
-void STextBlock::SetWrapTextAt(const TAttribute<float>& InWrapTextAt)
+void STextBlock::SetWrapTextAt(TAttribute<float> InWrapTextAt)
 {
-	if(!WrapTextAt.IdenticalTo(InWrapTextAt))
-	{
-		WrapTextAt = InWrapTextAt;
-		InvalidateText(EInvalidateWidget::LayoutAndVolatility);
-	}
+	WrapTextAt.Assign(*this, MoveTemp(InWrapTextAt), 0.f);
 }
 
-void STextBlock::SetAutoWrapText(const TAttribute<bool>& InAutoWrapText)
+void STextBlock::SetAutoWrapText(TAttribute<bool> InAutoWrapText)
 {
-	if(!AutoWrapText.IdenticalTo(InAutoWrapText))
-	{
-		AutoWrapText = InAutoWrapText;
-		InvalidateText(EInvalidateWidget::LayoutAndVolatility);
-	}
+	AutoWrapText.Assign(*this, MoveTemp(InAutoWrapText), 0.f);
 }
 
-void STextBlock::SetWrappingPolicy(const TAttribute<ETextWrappingPolicy>& InWrappingPolicy)
+void STextBlock::SetWrappingPolicy(TAttribute<ETextWrappingPolicy> InWrappingPolicy)
 {
-	if(!WrappingPolicy.IdenticalTo(InWrappingPolicy))
-	{
-		WrappingPolicy = InWrappingPolicy;
-		InvalidateText(EInvalidateWidget::LayoutAndVolatility);
-	}
+	WrappingPolicy.Assign(*this, MoveTemp(InWrappingPolicy));
 }
 
-void STextBlock::SetTransformPolicy(const TAttribute<ETextTransformPolicy>& InTransformPolicy)
+void STextBlock::SetTransformPolicy(TAttribute<ETextTransformPolicy> InTransformPolicy)
 {
-	if(!TransformPolicy.IdenticalTo(InTransformPolicy))
-	{
-		TransformPolicy = InTransformPolicy;
-		InvalidateText(EInvalidateWidget::LayoutAndVolatility);
-	}
+	bIsAttributeTransformPolicySet = InTransformPolicy.IsSet();
+	TransformPolicy.Assign(*this, MoveTemp(InTransformPolicy));
 }
 
-ETextTransformPolicy STextBlock::GetTransformPolicy() const
+void STextBlock::SetShadowOffset(TAttribute<FVector2D> InShadowOffset)
 {
-	if (TransformPolicy.IsSet())
-	{
-		return TransformPolicy.Get();
-	}
-	return TextStyle.TransformPolicy;
+	bIsAttributeShadowOffsetSet = InShadowOffset.IsSet();
+	ShadowOffset.Assign(*this, MoveTemp(InShadowOffset));
 }
 
-void STextBlock::SetShadowOffset(const TAttribute<FVector2D>& InShadowOffset)
+void STextBlock::SetShadowColorAndOpacity(TAttribute<FLinearColor> InShadowColorAndOpacity)
 {
-	if(!ShadowOffset.IdenticalTo(InShadowOffset))
-	{
-		ShadowOffset = InShadowOffset;
-		InvalidateText(EInvalidateWidget::LayoutAndVolatility);
-	}
+	bIsAttributeShadowColorAndOpacitySet = InShadowColorAndOpacity.IsSet();
+	ShadowColorAndOpacity.Assign(*this, MoveTemp(InShadowColorAndOpacity));
 }
 
-void STextBlock::SetShadowColorAndOpacity(const TAttribute<FLinearColor>& InShadowColorAndOpacity)
+void STextBlock::SetHighlightColor(TAttribute<FLinearColor> InHighlightColor)
 {
-	if(!ShadowColorAndOpacity.IdenticalTo(InShadowColorAndOpacity))
-	{
-		ShadowColorAndOpacity = InShadowColorAndOpacity;
-		// HACK: Normally this would be Paint only, but textblocks need to recache layout.
-		InvalidateText(EInvalidateWidget::LayoutAndVolatility);
-	}
+	bIsAttributeHighlightColorSet = InHighlightColor.IsSet();
+	HighlightColor.Assign(*this, MoveTemp(InHighlightColor));
 }
 
-void STextBlock::SetMinDesiredWidth(const TAttribute<float>& InMinDesiredWidth)
+void STextBlock::SetHighlightShape(TAttribute<const FSlateBrush*> InHighlightShape)
 {
-	if(!MinDesiredWidth.IdenticalTo(InMinDesiredWidth))
-	{
-		MinDesiredWidth = InMinDesiredWidth;
-		InvalidateText(EInvalidateWidget::LayoutAndVolatility);
-	}
+	bIsAttributeHighlightShapeSet = InHighlightShape.IsSet();
+	HighlightShape.Assign(*this, MoveTemp(InHighlightShape));
 }
 
-void STextBlock::SetLineHeightPercentage(const TAttribute<float>& InLineHeightPercentage)
+void STextBlock::SetMinDesiredWidth(TAttribute<float> InMinDesiredWidth)
 {
-	if(!LineHeightPercentage.IdenticalTo(InLineHeightPercentage))
-	{
-		LineHeightPercentage = InLineHeightPercentage;
-		InvalidateText(EInvalidateWidget::LayoutAndVolatility);
-	}
+	MinDesiredWidth.Assign(*this, MoveTemp(InMinDesiredWidth), 0.f);
 }
 
-void STextBlock::SetMargin(const TAttribute<FMargin>& InMargin)
+void STextBlock::SetLineHeightPercentage(TAttribute<float> InLineHeightPercentage)
 {
-	if(!Margin.IdenticalTo(InMargin))
-	{
-		Margin = InMargin;
-		InvalidateText(EInvalidateWidget::LayoutAndVolatility);
-	}
+	LineHeightPercentage.Assign(*this, MoveTemp(InLineHeightPercentage));
 }
 
-void STextBlock::SetJustification(const TAttribute<ETextJustify::Type>& InJustification)
+void STextBlock::SetMargin(TAttribute<FMargin> InMargin)
 {
-	if(!Justification.IdenticalTo(InJustification))
-	{
-		Justification = InJustification;
-		InvalidateText(EInvalidateWidget::LayoutAndVolatility);
-	}
+	Margin.Assign(*this, MoveTemp(InMargin));
+}
+
+void STextBlock::SetJustification(TAttribute<ETextJustify::Type> InJustification)
+{
+	Justification.Assign(*this, MoveTemp(InJustification));
 }
 
 FTextBlockStyle STextBlock::GetComputedTextStyle() const
 {
 	FTextBlockStyle ComputedStyle = TextStyle;
 	ComputedStyle.SetFont( GetFont() );
-	if (StrikeBrush.IsSet())
+	if (const FSlateBrush* const ComputedStrikeBrush = GetStrikeBrush())
 	{
-		const FSlateBrush* const ComputedStrikeBrush = StrikeBrush.Get();
-		if (ComputedStrikeBrush)
-		{
-			ComputedStyle.SetStrikeBrush(*ComputedStrikeBrush);
-		}
+		ComputedStyle.SetStrikeBrush(*ComputedStrikeBrush);
 	}
 	ComputedStyle.SetColorAndOpacity( GetColorAndOpacity() );
 	ComputedStyle.SetShadowOffset( GetShadowOffset() );
@@ -487,6 +440,6 @@ TSharedRef<FSlateAccessibleWidget> STextBlock::CreateAccessibleWidget()
 
 TOptional<FText> STextBlock::GetDefaultAccessibleText(EAccessibleType AccessibleType) const
 {
-	return GetText();
+	return BoundText.Get();
 }
 #endif
diff --git a/Engine/Source/Runtime/Slate/Private/Widgets/Text/SlateTextBlockLayout.cpp b/Engine/Source/Runtime/Slate/Private/Widgets/Text/SlateTextBlockLayout.cpp
--- a/Engine/Source/Runtime/Slate/Private/Widgets/Text/SlateTextBlockLayout.cpp
+++ b/Engine/Source/Runtime/Slate/Private/Widgets/Text/SlateTextBlockLayout.cpp
@@ -25,6 +25,80 @@ FSlateTextBlockLayout::FSlateTextBlockLayout(SWidget* InOwner, FTextBlockStyle I
 	TextLayout->SetLineBreakIterator(MoveTemp(InLineBreakPolicy));
 }
 
+FVector2D FSlateTextBlockLayout::ComputeDesiredSize(const FWidgetDesiredSizeArgs& InWidgetArgs, const float InScale, const FTextBlockStyle& InTextStyle)
+{
+	// Cache the wrapping rules so that we can recompute the wrap at width in paint.
+	CachedWrapTextAt = InWidgetArgs.WrapTextAt;
+	bCachedAutoWrapText = InWidgetArgs.AutoWrapText;
+
+	const ETextTransformPolicy PreviousTransformPolicy = TextLayout->GetTransformPolicy();
+
+	// Set the text layout information
+	TextLayout->SetScale(InScale);
+	TextLayout->SetWrappingWidth(CalculateWrappingWidth());
+	TextLayout->SetWrappingPolicy(InWidgetArgs.WrappingPolicy);
+	TextLayout->SetTransformPolicy(InWidgetArgs.TransformPolicy);
+	TextLayout->SetMargin(InWidgetArgs.Margin);
+	TextLayout->SetJustification(InWidgetArgs.Justification);
+	TextLayout->SetLineHeightPercentage(InWidgetArgs.LineHeightPercentage);
+
+	// Has the transform policy changed? If so we need a full refresh as that is destructive to the model text
+	if (PreviousTransformPolicy != TextLayout->GetTransformPolicy())
+	{
+		Marshaller->MakeDirty();
+	}
+
+	// Has the style used for this text block changed?
+	if (!IsStyleUpToDate(InTextStyle))
+	{
+		TextLayout->SetDefaultTextStyle(InTextStyle);
+		Marshaller->MakeDirty(); // will regenerate the text using the new default style
+	}
+
+	{
+		bool bRequiresTextUpdate = false;
+		const FText& TextToSet = InWidgetArgs.Text;
+		if (!TextLastUpdate.IdenticalTo(TextToSet))
+		{
+			// The pointer used by the bound text has changed, however the text may still be the same - check that now
+			if (!TextLastUpdate.IsDisplayStringEqualTo(TextToSet))
+			{
+				// The source text has changed, so update the internal text
+				bRequiresTextUpdate = true;
+			}
+
+			// Update this even if the text is lexically identical, as it will update the pointer compared by IdenticalTo for the next Tick
+			TextLastUpdate = FTextSnapshot(TextToSet);
+		}
+
+		if (bRequiresTextUpdate || Marshaller->IsDirty())
+		{
+			UpdateTextLayout(TextToSet);
+		}
+	}
+
+	{
+		const FText& HighlightTextToSet = InWidgetArgs.HighlightText;
+		if (!HighlightTextLastUpdate.IdenticalTo(HighlightTextToSet))
+		{
+			// The pointer used by the bound text has changed, however the text may still be the same - check that now
+			if (!HighlightTextLastUpdate.IsDisplayStringEqualTo(HighlightTextToSet))
+			{
+				UpdateTextHighlights(HighlightTextToSet);
+			}
+
+			// Update this even if the text is lexically identical, as it will update the pointer compared by IdenticalTo for the next Tick
+			HighlightTextLastUpdate = FTextSnapshot(HighlightTextToSet);
+		}
+	}
+
+	// We need to update our size if the text layout has become dirty
+	TextLayout->UpdateIfNeeded();
+
+	return TextLayout->GetSize();
+}
+
+PRAGMA_DISABLE_DEPRECATION_WARNINGS
 FVector2D FSlateTextBlockLayout::ComputeDesiredSize(const FWidgetArgs& InWidgetArgs, const float InScale, const FTextBlockStyle& InTextStyle)
 {
 	// Cache the wrapping rules so that we can recompute the wrap at width in paint.
@@ -97,6 +171,7 @@ FVector2D FSlateTextBlockLayout::ComputeDesiredSize(const FWidgetArgs& InWidgetA
 
 	return TextLayout->GetSize();
 }
+PRAGMA_DISABLE_DEPRECATION_WARNINGS
 
 FVector2D FSlateTextBlockLayout::GetDesiredSize() const
 {
diff --git a/Engine/Source/Runtime/Slate/Public/Widgets/Text/STextBlock.h b/Engine/Source/Runtime/Slate/Public/Widgets/Text/STextBlock.h
--- a/Engine/Source/Runtime/Slate/Public/Widgets/Text/STextBlock.h
+++ b/Engine/Source/Runtime/Slate/Public/Widgets/Text/STextBlock.h
@@ -38,6 +38,7 @@ namespace ETextRole
  */
 class SLATE_API STextBlock : public SLeafWidget
 {
+	SLATE_DECLARE_WIDGET(STextBlock, SLeafWidget)
 
 public:
 
@@ -92,7 +93,7 @@ public:
 		/** The color used to highlight the specified text */
 		SLATE_ATTRIBUTE( FLinearColor, HighlightColor )
 
-		/** The brush used to highlight the specified text*/
+		/** The brush used to highlight the specified text */
 		SLATE_ATTRIBUTE( const FSlateBrush*, HighlightShape )
 
 		/** Highlight this text in the text block */
@@ -165,42 +166,29 @@ public:
 	 */
 	const FText& GetText() const
 	{
+		if (bIsAttributeBoundTextBound)
+		{
+			STextBlock& MutableSelf = const_cast<STextBlock&>(*this);
+			MutableSelf.BoundText.UpdateNow(MutableSelf);
+		}
 		return BoundText.Get();
 	}
 	
 public:
-	UE_DEPRECATED(4.24, "SetText taking FString is deprecated. Use the FText version instead")
-	void SetText( const TAttribute< FString >& InText );
-	UE_DEPRECATED(4.24, "SetText taking FString is deprecated. Use the FText version instead")
-	void SetText( const FString& InText );
-
-	/**
-	 * Sets the text for this text block
-	 *
-	 * @param	InText	The new text to display
-	 */
-	void SetText( const TAttribute< FText >& InText );
-	void SetText( const FText& InText );
+	/** Sets the text for this text block */
+	void SetText(TAttribute<FText> InText);
 
-	/**
-	* Sets the highlight text for this text block 
-	*
-	* @param	InText	The new text to highlight
-	*/
+	/** Sets the highlight text for this text block */
 	void SetHighlightText(TAttribute<FText> InText);
 
-	/**
-	 * Sets the font used to draw the text
-	 *
-	 * @param	InFont	The new font to use
-	 */
-	void SetFont(const TAttribute< FSlateFontInfo >& InFont);
+	/** Sets the font used to draw the text	*/
+	void SetFont(TAttribute<FSlateFontInfo> InFont);
 
 	/** Sets the brush used to strike through the text */
-	void SetStrikeBrush(const TAttribute<const FSlateBrush*>& InStrikeBrush);
+	void SetStrikeBrush(TAttribute<const FSlateBrush*> InStrikeBrush);
 
 	/** See ColorAndOpacity attribute */
-	void SetColorAndOpacity(const TAttribute<FSlateColor>& InColorAndOpacity);
+	void SetColorAndOpacity(TAttribute<FSlateColor> InColorAndOpacity);
 
 	/** See TextStyle argument */
 	void SetTextStyle(const FTextBlockStyle* InTextStyle);
@@ -212,35 +200,44 @@ public:
 	void SetTextFlowDirection(const TOptional<ETextFlowDirection>& InTextFlowDirection);
 
 	/** See WrapTextAt attribute */
-	void SetWrapTextAt(const TAttribute<float>& InWrapTextAt);
+	void SetWrapTextAt(TAttribute<float> InWrapTextAt);
 
 	/** See AutoWrapText attribute */
-	void SetAutoWrapText(const TAttribute<bool>& InAutoWrapText);
+	void SetAutoWrapText(TAttribute<bool> InAutoWrapText);
 
 	/** Set WrappingPolicy attribute */
-	void SetWrappingPolicy(const TAttribute<ETextWrappingPolicy>& InWrappingPolicy);
+	void SetWrappingPolicy(TAttribute<ETextWrappingPolicy> InWrappingPolicy);
 
 	/** Set TransformPolicy attribute */
-	void SetTransformPolicy(const TAttribute<ETextTransformPolicy>& InTransformPolicy);
-	ETextTransformPolicy GetTransformPolicy() const;
+	void SetTransformPolicy(TAttribute<ETextTransformPolicy> InTransformPolicy);
+
+	/** Get TransformPolicy attribute */
+	UE_DEPRECATED(5.0, "GetTransformPolicy is not accessible anymore since it's attribute value may not have been updated yet.")
+	ETextTransformPolicy GetTransformPolicy() const { return GetTransformPolicyImpl(); }
 
 	/** See ShadowOffset attribute */
-	void SetShadowOffset(const TAttribute<FVector2D>& InShadowOffset);
+	void SetShadowOffset(TAttribute<FVector2D> InShadowOffset);
 
 	/** See ShadowColorAndOpacity attribute */
-	void SetShadowColorAndOpacity(const TAttribute<FLinearColor>& InShadowColorAndOpacity);
+	void SetShadowColorAndOpacity(TAttribute<FLinearColor> InShadowColorAndOpacity);
+
+	/** See HighlightColor attribute */
+	void SetHighlightColor(TAttribute<FLinearColor> InHighlightColor);
+	
+	/** See HighlightShape attribute */
+	void SetHighlightShape(TAttribute<const FSlateBrush*> InHighlightShape);
 
 	/** See MinDesiredWidth attribute */
-	void SetMinDesiredWidth(const TAttribute<float>& InMinDesiredWidth);
+	void SetMinDesiredWidth(TAttribute<float> InMinDesiredWidth);
 
 	/** See LineHeightPercentage attribute */
-	void SetLineHeightPercentage(const TAttribute<float>& InLineHeightPercentage);
+	void SetLineHeightPercentage(TAttribute<float> InLineHeightPercentage);
 
 	/** See Margin attribute */
-	void SetMargin(const TAttribute<FMargin>& InMargin);
+	void SetMargin(TAttribute<FMargin> InMargin);
 
 	/** See Justification attribute */
-	void SetJustification(const TAttribute<ETextJustify::Type>& InJustification);
+	void SetJustification(TAttribute<ETextJustify::Type> InJustification);
 
 	// SWidget interface
 	virtual int32 OnPaint( const FPaintArgs& Args, const FGeometry& AllottedGeometry, const FSlateRect& MyCullingRect, FSlateWindowElementList& OutDrawElements, int32 LayerId, const FWidgetStyle& InWidgetStyle, bool bParentEnabled ) const override;
@@ -251,11 +248,6 @@ public:
 #endif
 	// End of SWidget interface
 
-protected:
-	// Begin SWidget overrides.
-	virtual bool ComputeVolatility() const override;
-	// End of SWidget interface
-
 private:
 	/** Get the computed text style to use with the text marshaller */
 	FTextBlockStyle GetComputedTextStyle() const;
@@ -268,6 +260,9 @@ private:
 
 	/** Gets the current strike brush */
 	const FSlateBrush* GetStrikeBrush() const;
+
+	/** Get TransformPolicy attribute */
+	ETextTransformPolicy GetTransformPolicyImpl() const;
 	
 	/** Gets the current shadow offset */
 	FVector2D GetShadowOffset() const;
@@ -282,11 +277,14 @@ private:
 	const FSlateBrush* GetHighlightShape() const;
 
 	/** Call to invalidate this text block */
-	void InvalidateText(EInvalidateWidget InvalidateReason);
+	void InvalidateText(EInvalidateWidgetReason InvalidateReason);
+
+private:
+
 
 private:
 	/** The text displayed in this text block */
-	TAttribute< FText > BoundText;
+	TSlateAttribute<FText> BoundText;
 
 	/** The wrapped layout for this text block */
 	TUniquePtr< FSlateTextBlockLayout > TextLayoutCache;
@@ -295,52 +293,74 @@ private:
 	FTextBlockStyle TextStyle;
 
 	/** Sets the font used to draw the text */
-	TAttribute< FSlateFontInfo > Font;
+	TSlateAttribute<FSlateFontInfo> Font;
 
 	/** Sets the brush used to strike through the text */
-	TAttribute< const FSlateBrush* > StrikeBrush;
+	TSlateAttribute<const FSlateBrush*> StrikeBrush;
 
 	/** Text color and opacity */
-	TAttribute<FSlateColor> ColorAndOpacity;
+	TSlateAttribute<FSlateColor> ColorAndOpacity;
 
 	/** Drop shadow offset in pixels */
-	TAttribute< FVector2D > ShadowOffset;
+	TSlateAttribute<FVector2D> ShadowOffset;
 
 	/** Shadow color and opacity */
-	TAttribute<FLinearColor> ShadowColorAndOpacity;
+	TSlateAttribute<FLinearColor> ShadowColorAndOpacity;
 
-	TAttribute<FLinearColor> HighlightColor;
+	/** The color used to highlight the specified text */
+	TSlateAttribute<FLinearColor> HighlightColor;
 
-	TAttribute< const FSlateBrush* > HighlightShape;
+	/** The brush used to highlight the specified text */
+	TSlateAttribute<const FSlateBrush*> HighlightShape;
 
-	/** Highlight this text in the textblock */
-	TAttribute<FText> HighlightText;
+	/** Highlight this text in the TextBlock */
+	TSlateAttribute<FText> HighlightText;
 
 	/** Whether text wraps onto a new line when it's length exceeds this width; if this value is zero or negative, no wrapping occurs. */
-	TAttribute<float> WrapTextAt;
+	TSlateAttribute<float> WrapTextAt;
 
 	/** True if we're wrapping text automatically based on the computed horizontal space for this widget */
-	TAttribute<bool> AutoWrapText;
+	TSlateAttribute<bool> AutoWrapText;
 
 	/** The wrapping policy we're using */
-	TAttribute<ETextWrappingPolicy> WrappingPolicy;
+	TSlateAttribute<ETextWrappingPolicy> WrappingPolicy;
 
 	/** The transform policy we're using */
-	TAttribute<ETextTransformPolicy> TransformPolicy;
+	TSlateAttribute<ETextTransformPolicy> TransformPolicy;
 
 	/** The amount of blank space left around the edges of text area. */
-	TAttribute< FMargin > Margin;
+	TSlateAttribute<FMargin> Margin;
 
 	/** The amount to scale each lines height by. */
-	TAttribute< ETextJustify::Type > Justification; 
+	TSlateAttribute<ETextJustify::Type> Justification;
 
 	/** How the text should be aligned with the margin. */
-	TAttribute< float > LineHeightPercentage;
+	TSlateAttribute<float> LineHeightPercentage;
 
 	/** Prevents the text block from being smaller than desired in certain cases (e.g. when it is empty) */
-	TAttribute<float> MinDesiredWidth;
+	TSlateAttribute<float> MinDesiredWidth;
 
 	/** If this is enabled, text shaping, wrapping, justification are disabled in favor of much faster text layout and measurement. */
 	mutable TOptional<FVector2D> CachedSimpleDesiredSize;
+
+	/** Flags used to check if the SlateAttribute is set. */
+	union
+	{
+		struct 
+		{
+			uint16 bIsAttributeBoundTextBound : 1;
+			uint16 bIsAttributeFontSet : 1;
+			uint16 bIsAttributeStrikeBrushSet : 1;
+			uint16 bIsAttributeColorAndOpacitySet : 1;
+			uint16 bIsAttributeShadowOffsetSet : 1;
+			uint16 bIsAttributeShadowColorAndOpacitySet : 1;
+			uint16 bIsAttributeHighlightColorSet : 1;
+			uint16 bIsAttributeHighlightShapeSet : 1;
+			uint16 bIsAttributeWrapTextAtSet : 1;
+			uint16 bIsAttributeTransformPolicySet : 1;
+		};
+		uint16 Union_Flags;
+	};
+
 	bool bSimpleTextMode;
 };
diff --git a/Engine/Source/Runtime/Slate/Public/Widgets/Text/SlateTextBlockLayout.h b/Engine/Source/Runtime/Slate/Public/Widgets/Text/SlateTextBlockLayout.h
--- a/Engine/Source/Runtime/Slate/Public/Widgets/Text/SlateTextBlockLayout.h
+++ b/Engine/Source/Runtime/Slate/Public/Widgets/Text/SlateTextBlockLayout.h
@@ -20,7 +20,45 @@ enum class ETextShapingMethod : uint8;
 class SLATE_API FSlateTextBlockLayout
 {
 public:
-	struct FWidgetArgs
+	struct FWidgetDesiredSizeArgs
+	{
+		FORCEINLINE FWidgetDesiredSizeArgs(
+			const FText& InText,
+			const FText& InHighlightText,
+			const float InWrapTextAt,
+			const bool InAutoWrapText,
+			const ETextWrappingPolicy InWrappingPolicy,
+			const ETextTransformPolicy InTransformPolicy,
+			const FMargin& InMargin,
+			const float InLineHeightPercentage,
+			const ETextJustify::Type InJustification
+		)
+			: Text(InText)
+			, HighlightText(InHighlightText)
+			, Margin(InMargin)
+			, WrapTextAt(InWrapTextAt)
+			, LineHeightPercentage(InLineHeightPercentage)
+			, WrappingPolicy(InWrappingPolicy)
+			, TransformPolicy(InTransformPolicy)
+			, Justification(InJustification)
+			, AutoWrapText(InAutoWrapText)
+		{
+		}
+
+		const FText Text = FText::GetEmpty();
+		const FText HighlightText = FText::GetEmpty();
+		const FMargin Margin;
+		const float WrapTextAt = 0.f;
+		const float LineHeightPercentage;
+		const ETextWrappingPolicy WrappingPolicy;
+		const ETextTransformPolicy TransformPolicy;
+		const ETextJustify::Type Justification;
+		const bool AutoWrapText = false;
+	};
+
+	struct
+		UE_DEPRECATED(5.0, "FWidgetArgs is deprecated. Upgrade to FWidgetDesiredSizeArgs instead.")
+		FWidgetArgs
 	{
 		FORCEINLINE FWidgetArgs(
 			const TAttribute<FText>& InText, 
@@ -64,7 +102,15 @@ public:
 	/**
 	 * Get the computed desired size for this layout, updating the internal cache as required
 	 */
+	FVector2D ComputeDesiredSize(const FWidgetDesiredSizeArgs& InWidgetArgs, const float InScale, const FTextBlockStyle& InTextStyle);
+
+	PRAGMA_DISABLE_DEPRECATION_WARNINGS
+	/**
+	 * Get the computed desired size for this layout, updating the internal cache as required
+	 */
+	UE_DEPRECATED(5.0, "FWidgetArgs is deprecated. Upgrade to FWidgetDesiredSizeArgs instead.")
 	FVector2D ComputeDesiredSize(const FWidgetArgs& InWidgetArgs, const float InScale, const FTextBlockStyle& InTextStyle);
+	PRAGMA_DISABLE_DEPRECATION_WARNINGS
 
 	/**
 	 * Gets the last computed desired size.
