Fix Slate deprecation errors

diff --git a/Engine/Source/Runtime/Slate/Private/Widgets/Text/SRichTextBlock.cpp b/Engine/Source/Runtime/Slate/Private/Widgets/Text/SRichTextBlock.cpp
--- a/Engine/Source/Runtime/Slate/Private/Widgets/Text/SRichTextBlock.cpp
+++ b/Engine/Source/Runtime/Slate/Private/Widgets/Text/SRichTextBlock.cpp
@@ -86,7 +86,7 @@ FVector2D SRichTextBlock::ComputeDesiredSize(float LayoutScaleMultiplier) const
 {
 	// ComputeDesiredSize will also update the text layout cache if required
 	const FVector2D TextSize = TextLayoutCache->ComputeDesiredSize(
-		FSlateTextBlockLayout::FWidgetArgs(BoundText, HighlightText, WrapTextAt, AutoWrapText, WrappingPolicy, TransformPolicy, Margin, LineHeightPercentage, Justification),
+		FSlateTextBlockLayout::FWidgetDesiredSizeArgs(BoundText.Get(), HighlightText.Get(), WrapTextAt.Get(), AutoWrapText.Get(), WrappingPolicy.Get(), TransformPolicy.Get(), Margin.Get(), LineHeightPercentage.Get(), Justification.Get()),
 		LayoutScaleMultiplier * TextBlockScale, TextStyle) * TextBlockScale;
 
 	return FVector2D(FMath::Max(TextSize.X, MinDesiredWidth.Get()), TextSize.Y);
diff --git a/Engine/Source/Runtime/Slate/Private/Widgets/Text/SlateEditableTextLayout.cpp b/Engine/Source/Runtime/Slate/Private/Widgets/Text/SlateEditableTextLayout.cpp
--- a/Engine/Source/Runtime/Slate/Private/Widgets/Text/SlateEditableTextLayout.cpp
+++ b/Engine/Source/Runtime/Slate/Private/Widgets/Text/SlateEditableTextLayout.cpp
@@ -3428,7 +3428,7 @@ FVector2D FSlateEditableTextLayout::ComputeDesiredSize(float LayoutScaleMultipli
 		MarginValue.Right += CaretWidth;
 
 		const FVector2D HintTextSize = HintTextLayout->ComputeDesiredSize(
-			FSlateTextBlockLayout::FWidgetArgs(HintText, FText::GetEmpty(), WrapTextAt, AutoWrapText, WrappingPolicy, ETextTransformPolicy::None, MarginValue, LineHeightPercentage, Justification),
+			FSlateTextBlockLayout::FWidgetDesiredSizeArgs(HintText.Get(), FText::GetEmpty(), WrapTextAt.Get(), AutoWrapText.Get(), WrappingPolicy.Get(), ETextTransformPolicy::None, MarginValue, LineHeightPercentage.Get(), Justification.Get()),
 			LayoutScaleMultiplier, HintTextStyle
 			);
 
