Slate: Modification to FOneSimpleMemberChild to use the SLATE_ADD_MEMBER_ATTRIBUTE_DEFINITION macro.

diff --git a/Engine/Source/Runtime/Slate/Private/Widgets/Layout/SBox.cpp b/Engine/Source/Runtime/Slate/Private/Widgets/Layout/SBox.cpp
--- a/Engine/Source/Runtime/Slate/Private/Widgets/Layout/SBox.cpp
+++ b/Engine/Source/Runtime/Slate/Private/Widgets/Layout/SBox.cpp
@@ -8,7 +8,7 @@
 SLATE_IMPLEMENT_WIDGET(SBox)
 void SBox::PrivateRegisterAttributes(FSlateAttributeInitializer& AttributeInitializer)
 {
-	AttributeInitializer.AddMemberAttribute("SlotPadding", STRUCT_OFFSET(PrivateThisType, ChildSlot) + FBoxSlot::GetSlotPaddingAttributeOffset(), FSlateAttributeDescriptor::FInvalidateWidgetReasonAttribute{ EInvalidateWidgetReason::Layout });
+	SLATE_ADD_MEMBER_ATTRIBUTE_DEFINITION_WITH_NAME(AttributeInitializer, "SlotPadding", ChildSlot.SlotPaddingAttribute, EInvalidateWidgetReason::Layout);
 	SLATE_ADD_MEMBER_ATTRIBUTE_DEFINITION(AttributeInitializer, WidthOverride, EInvalidateWidgetReason::Layout);
 	SLATE_ADD_MEMBER_ATTRIBUTE_DEFINITION(AttributeInitializer, HeightOverride, EInvalidateWidgetReason::Layout);
 	SLATE_ADD_MEMBER_ATTRIBUTE_DEFINITION(AttributeInitializer, MinDesiredWidth, EInvalidateWidgetReason::Layout);
diff --git a/Engine/Source/Runtime/Slate/Private/Widgets/Layout/SDPIScaler.cpp b/Engine/Source/Runtime/Slate/Private/Widgets/Layout/SDPIScaler.cpp
--- a/Engine/Source/Runtime/Slate/Private/Widgets/Layout/SDPIScaler.cpp
+++ b/Engine/Source/Runtime/Slate/Private/Widgets/Layout/SDPIScaler.cpp
@@ -6,7 +6,7 @@
 SLATE_IMPLEMENT_WIDGET(SDPIScaler)
 void SDPIScaler::PrivateRegisterAttributes(FSlateAttributeInitializer& AttributeInitializer)
 {
-	AttributeInitializer.AddMemberAttribute("SlotPadding", STRUCT_OFFSET(PrivateThisType, ChildSlot) + FOneSimpleMemberChild::GetSlotPaddingAttributeOffset(), FSlateAttributeDescriptor::FInvalidateWidgetReasonAttribute{ EInvalidateWidgetReason::Layout });
+	SLATE_ADD_MEMBER_ATTRIBUTE_DEFINITION_WITH_NAME(AttributeInitializer, "SlotPadding", ChildSlot.SlotPaddingAttribute, EInvalidateWidgetReason::Layout);
 	SLATE_ADD_MEMBER_ATTRIBUTE_DEFINITION_WITH_NAME(AttributeInitializer, "DPIScale", DPIScaleAttribute, EInvalidateWidgetReason::Prepass);
 }
 
@@ -22,7 +22,6 @@ SDPIScaler::SDPIScaler()
 void SDPIScaler::Construct( const FArguments& InArgs )
 {
 	SetDPIScale(InArgs._DPIScale);
-	DPIScaleAttribute.UpdateNow(*this);
 
 	ChildSlot
 	[
@@ -72,6 +71,7 @@ void SDPIScaler::SetContent(TSharedRef<SWidget> InContent)
 void SDPIScaler::SetDPIScale(TAttribute<float> InDPIScale)
 {
 	DPIScaleAttribute.Assign(*this, MoveTemp(InDPIScale), 1.f);
+	DPIScaleAttribute.UpdateNow(*this);
 }
 
 float SDPIScaler::GetRelativeLayoutScale(int32 ChildIndex, float LayoutScaleMultiplier) const
diff --git a/Engine/Source/Runtime/Slate/Public/Widgets/Layout/SBox.h b/Engine/Source/Runtime/Slate/Public/Widgets/Layout/SBox.h
--- a/Engine/Source/Runtime/Slate/Public/Widgets/Layout/SBox.h
+++ b/Engine/Source/Runtime/Slate/Public/Widgets/Layout/SBox.h
@@ -24,9 +24,11 @@ class SLATE_API SBox : public SPanel
 	SLATE_DECLARE_WIDGET(SBox, SPanel)
 
 public:
-	class FBoxSlot : public FOneSimpleMemberChild
+
+	class UE_DEPRECATED(5.0, "FBoxSlot is deprecated. Use FSimpleSlot or FOneSimpleMemberChild")
+	FBoxSlot : public FSimpleSlot
 	{
-		using FOneSimpleMemberChild::FOneSimpleMemberChild;
+		using FSimpleSlot::FSimpleSlot;
 	};
 
 	SLATE_BEGIN_ARGS(SBox)
@@ -133,7 +135,12 @@ protected:
 
 protected:
 
-	FBoxSlot ChildSlot;
+	struct FBoxOneChildSlot : ::FOneSimpleMemberChild
+	{
+		friend SBox;
+		using ::FOneSimpleMemberChild::FOneSimpleMemberChild;
+	};
+	FBoxOneChildSlot ChildSlot;
 
 private:
 	/** When specified, ignore the content's desired size and report the.WidthOverride as the Box's desired width. */
diff --git a/Engine/Source/Runtime/Slate/Public/Widgets/Layout/SDPIScaler.h b/Engine/Source/Runtime/Slate/Public/Widgets/Layout/SDPIScaler.h
--- a/Engine/Source/Runtime/Slate/Public/Widgets/Layout/SDPIScaler.h
+++ b/Engine/Source/Runtime/Slate/Public/Widgets/Layout/SDPIScaler.h
@@ -65,8 +65,14 @@ protected:
 	FSlateDeprecatedTAttribute<float> DPIScale;
 #endif
 
+	struct FDPIScalerOneChildSlot : ::FOneSimpleMemberChild
+	{
+		friend SDPIScaler;
+		using ::FOneSimpleMemberChild::FOneSimpleMemberChild;
+	};
+
 	/** The content being scaled. */
-	FOneSimpleMemberChild ChildSlot;
+	FDPIScalerOneChildSlot ChildSlot;
 
 private:
 	/**
diff --git a/Engine/Source/Runtime/SlateCore/Private/Widgets/SCompoundWidget.cpp b/Engine/Source/Runtime/SlateCore/Private/Widgets/SCompoundWidget.cpp
--- a/Engine/Source/Runtime/SlateCore/Private/Widgets/SCompoundWidget.cpp
+++ b/Engine/Source/Runtime/SlateCore/Private/Widgets/SCompoundWidget.cpp
@@ -11,7 +11,7 @@ DECLARE_CYCLE_STAT(TEXT("Child Paint"), STAT_ChildPaint, STATGROUP_SlateVeryVerb
 SLATE_IMPLEMENT_WIDGET(SCompoundWidget)
 void SCompoundWidget::PrivateRegisterAttributes(FSlateAttributeInitializer& AttributeInitializer)
 {
-	AttributeInitializer.AddMemberAttribute("SlotPadding", STRUCT_OFFSET(PrivateThisType, ChildSlot) + FOneSimpleMemberChild::GetSlotPaddingAttributeOffset(), FSlateAttributeDescriptor::FInvalidateWidgetReasonAttribute{ EInvalidateWidgetReason::Layout });
+	SLATE_ADD_MEMBER_ATTRIBUTE_DEFINITION_WITH_NAME(AttributeInitializer, "SlotPadding", ChildSlot.SlotPaddingAttribute, EInvalidateWidgetReason::Layout);
 	SLATE_ADD_MEMBER_ATTRIBUTE_DEFINITION_WITH_NAME(AttributeInitializer, "ContentScale", ContentScaleAttribute, EInvalidateWidgetReason::Layout);
 	SLATE_ADD_MEMBER_ATTRIBUTE_DEFINITION_WITH_NAME(AttributeInitializer, "ColorAndOpacity", ColorAndOpacityAttribute, EInvalidateWidgetReason::Paint);
 	SLATE_ADD_MEMBER_ATTRIBUTE_DEFINITION_WITH_NAME(AttributeInitializer, "ForegroundColor", ForegroundColorAttribute, EInvalidateWidgetReason::Paint);
diff --git a/Engine/Source/Runtime/SlateCore/Public/Layout/Children.h b/Engine/Source/Runtime/SlateCore/Public/Layout/Children.h
--- a/Engine/Source/Runtime/SlateCore/Public/Layout/Children.h
+++ b/Engine/Source/Runtime/SlateCore/Public/Layout/Children.h
@@ -865,12 +865,10 @@ public:
 	using SlotPaddingAttributeType = SlateAttributePrivate::TSlateMemberAttribute<FMargin, ::SlateAttributePrivate::FSlateAttributeNoInvalidationReason, TSlateAttributeComparePredicate<>>;
 	using SlotPaddingAttributeRefType = SlateAttributePrivate::TSlateMemberAttributeRef<SlotPaddingAttributeType>;
 
-	static SIZE_T GetSlotPaddingAttributeOffset() { return STRUCT_OFFSET(FOneSimpleMemberChild, SlotPaddingAttribute); }
-
 	template<typename WidgetType, typename V = typename std::enable_if<std::is_base_of<SWidget, WidgetType>::value>::type>
 	SlotPaddingAttributeRefType GetSlotPaddingAttribute() const { return SlotPaddingAttributeRefType(*(static_cast<const WidgetType*>(FChildren::GetOwner())), SlotPaddingAttribute); }
 
-private:
+protected:
 	SlotPaddingAttributeType SlotPaddingAttribute;
 };
 
diff --git a/Engine/Source/Runtime/SlateCore/Public/Widgets/SCompoundWidget.h b/Engine/Source/Runtime/SlateCore/Public/Widgets/SCompoundWidget.h
--- a/Engine/Source/Runtime/SlateCore/Public/Widgets/SCompoundWidget.h
+++ b/Engine/Source/Runtime/SlateCore/Public/Widgets/SCompoundWidget.h
@@ -102,8 +102,14 @@ protected:
 	/** Disallow public construction */
 	SCompoundWidget();
 
+	struct FCompoundWidgetOneChildSlot : ::FOneSimpleMemberChild
+	{
+		friend SCompoundWidget;
+		using ::FOneSimpleMemberChild::FOneSimpleMemberChild;
+	};
+
 	/** The slot that contains this widget's descendants.*/
-	FOneSimpleMemberChild ChildSlot;
+	FCompoundWidgetOneChildSlot ChildSlot;
 
 #if WITH_EDITORONLY_DATA
 	UE_DEPRECATED(5.0, "Direct access to ContentScale is now deprecated. Use the setter or getter.")
