Slate: Convert SBackgroundBlur TAttributes to TSlateAttributes.

diff --git a/Engine/Source/Runtime/Slate/Private/Widgets/Layout/SBackgroundBlur.cpp b/Engine/Source/Runtime/Slate/Private/Widgets/Layout/SBackgroundBlur.cpp
--- a/Engine/Source/Runtime/Slate/Private/Widgets/Layout/SBackgroundBlur.cpp
+++ b/Engine/Source/Runtime/Slate/Private/Widgets/Layout/SBackgroundBlur.cpp
@@ -15,13 +15,26 @@ static FAutoConsoleVariableRef CVarDownsampleForBlur(TEXT("Slate.BackgroundBlurD
 static int32 bForceLowQualityBrushFallback = 0;
 static FAutoConsoleVariableRef CVarForceLowQualityBackgroundBlurOverride(TEXT("Slate.ForceBackgroundBlurLowQualityOverride"), bForceLowQualityBrushFallback, TEXT("Whether or not to force a slate brush to be used instead of actually blurring the background"), ECVF_Scalability);
 
+SLATE_IMPLEMENT_WIDGET(SBackgroundBlur)
+void SBackgroundBlur::PrivateRegisterAttributes(FSlateAttributeInitializer& AttributeInitializer)
+{
+	SLATE_ADD_MEMBER_ATTRIBUTE_DEFINITION_WITH_NAME(AttributeInitializer, "BlurStrength", BlurStrengthAttribute, EInvalidateWidgetReason::Paint);
+	SLATE_ADD_MEMBER_ATTRIBUTE_DEFINITION_WITH_NAME(AttributeInitializer, "BlurRadius", BlurRadiusAttribute, EInvalidateWidgetReason::Paint);
+}
+
+SBackgroundBlur::SBackgroundBlur()
+	: BlurStrengthAttribute(*this)
+	, BlurRadiusAttribute(*this)
+{
+}
+
 
 void SBackgroundBlur::Construct(const FArguments& InArgs)
 {
 	bApplyAlphaToBlur = InArgs._bApplyAlphaToBlur;
 	LowQualityFallbackBrush = InArgs._LowQualityFallbackBrush;
-	BlurStrength = InArgs._BlurStrength;
-	BlurRadius = InArgs._BlurRadius;
+	SetBlurStrength(InArgs._BlurStrength);
+	SetBlurRadius(InArgs._BlurRadius);
 
 	ChildSlot
 		.HAlign(InArgs._HAlign)
@@ -48,22 +61,14 @@ void SBackgroundBlur::SetApplyAlphaToBlur(bool bInApplyAlphaToBlur)
 	}
 }
 
-void SBackgroundBlur::SetBlurRadius(const TAttribute<TOptional<int32>>& InBlurRadius)
+void SBackgroundBlur::SetBlurRadius(TAttribute<TOptional<int32>> InBlurRadius)
 {
-	if(!BlurRadius.IdenticalTo(InBlurRadius))
-	{
-		BlurRadius = InBlurRadius;
-		Invalidate(EInvalidateWidget::Paint);
-	}
+	BlurRadiusAttribute.Assign(*this, InBlurRadius);
 }
 
-void SBackgroundBlur::SetBlurStrength(const TAttribute<float>& InStrength)
+void SBackgroundBlur::SetBlurStrength(TAttribute<float> InStrength)
 {
-	if(!BlurStrength.IdenticalTo(InStrength))
-	{
-		BlurStrength = InStrength;
-		Invalidate(EInvalidateWidget::Paint);
-	}
+	BlurStrengthAttribute.Assign(*this, MoveTemp(InStrength));
 }
 
 void SBackgroundBlur::SetLowQualityBackgroundBrush(const FSlateBrush* InBrush)
@@ -113,7 +118,7 @@ int32 SBackgroundBlur::OnPaint(const FPaintArgs& Args, const FGeometry& Allotted
 		if (!bUsingLowQualityFallbackBrush)
 		{
 			// Modulate blur strength by the widget alpha
-			const float Strength = BlurStrength.Get() * (bApplyAlphaToBlur ? InWidgetStyle.GetColorAndOpacityTint().A : 1.f);
+			const float Strength = BlurStrengthAttribute.Get() * (bApplyAlphaToBlur ? InWidgetStyle.GetColorAndOpacityTint().A : 1.f);
 			if (Strength > 0.f)
 			{
 				FSlateRect RenderBoundingRect = AllottedGeometry.GetRenderBoundingRect();
@@ -165,7 +170,7 @@ int32 SBackgroundBlur::OnPaint(const FPaintArgs& Args, const FGeometry& Allotted
 void SBackgroundBlur::ComputeEffectiveKernelSize(float Strength, int32& OutKernelSize, int32& OutDownsampleAmount) const
 {
 	// If the radius isn't set, auto-compute it based on the strength
-	OutKernelSize = BlurRadius.Get().Get(FMath::RoundToInt(Strength * 3.f));
+	OutKernelSize = BlurRadiusAttribute.Get().Get(FMath::RoundToInt(Strength * 3.f));
 
 	// Downsample if needed
 	if (bDownsampleForBlur && OutKernelSize > 9)
diff --git a/Engine/Source/Runtime/Slate/Public/Widgets/Layout/SBackgroundBlur.h b/Engine/Source/Runtime/Slate/Public/Widgets/Layout/SBackgroundBlur.h
--- a/Engine/Source/Runtime/Slate/Public/Widgets/Layout/SBackgroundBlur.h
+++ b/Engine/Source/Runtime/Slate/Public/Widgets/Layout/SBackgroundBlur.h
@@ -17,6 +17,8 @@
 */
 class SLATE_API SBackgroundBlur : public SCompoundWidget
 {
+	SLATE_DECLARE_WIDGET(SBackgroundBlur, SCompoundWidget)
+
 public:
 	SLATE_BEGIN_ARGS(SBackgroundBlur)
 		: _HAlign(HAlign_Fill)
@@ -42,13 +44,14 @@ public:
 	SLATE_END_ARGS()
 
 public:
+	SBackgroundBlur();
 	void Construct(const FArguments& InArgs);
 	virtual int32 OnPaint(const FPaintArgs& Args, const FGeometry& AllottedGeometry, const FSlateRect& MyCullingRect, FSlateWindowElementList& OutDrawElements, int32 LayerId, const FWidgetStyle& InWidgetStyle, bool bParentEnabled) const override;
 
 	void SetContent(const TSharedRef<SWidget>& InContent);
 	void SetApplyAlphaToBlur(bool bInApplyAlphaToBlur);
-	void SetBlurRadius(const TAttribute<TOptional<int32>>& InBlurRadius);
-	void SetBlurStrength(const TAttribute<float>& InStrength);
+	void SetBlurRadius(TAttribute<TOptional<int32>> InBlurRadius);
+	void SetBlurStrength(TAttribute<float> InStrength);
 	void SetLowQualityBackgroundBrush(const FSlateBrush* InBrush);
 	
 	void SetHAlign(EHorizontalAlignment HAlign);
@@ -60,8 +63,22 @@ public:
 protected:
 	void ComputeEffectiveKernelSize(float Strength, int32& OutKernelSize, int32& OutDownsampleAmount) const;
 
+
+	/** @return an attribute reference of ColorAndOpacity */
+	TSlateAttributeRef<float> GetBlurStrengthAttribute() const { return TSlateAttributeRef<float>{*this, BlurStrengthAttribute}; }
+
+	/** @return an attribute reference of ForegroundColor */
+	TSlateAttributeRef<TOptional<int32>> GetBlurRadiusAttribute() const { return TSlateAttributeRef<TOptional<int32>>{*this, BlurRadiusAttribute}; }
+
+#if WITH_EDITORONLY_DATA
+	FSlateDeprecatedTAttribute<float> BlurStrength;
+	FSlateDeprecatedTAttribute<TOptional<int32>> BlurRadius;
+#endif
+
 	bool bApplyAlphaToBlur;
-	TAttribute<float> BlurStrength;
-	TAttribute<TOptional<int32>> BlurRadius;
 	const FSlateBrush* LowQualityFallbackBrush;
+
+private:
+	TSlateAttribute<float> BlurStrengthAttribute;
+	TSlateAttribute<TOptional<int32>> BlurRadiusAttribute;
 };
