Slate: Convert SSpacer TAttributes to TSlateAttributes.

diff --git a/Engine/Source/Runtime/Slate/Private/Widgets/Layout/SSpacer.cpp b/Engine/Source/Runtime/Slate/Private/Widgets/Layout/SSpacer.cpp
--- a/Engine/Source/Runtime/Slate/Private/Widgets/Layout/SSpacer.cpp
+++ b/Engine/Source/Runtime/Slate/Private/Widgets/Layout/SSpacer.cpp
@@ -2,15 +2,23 @@
 
 #include "Widgets/Layout/SSpacer.h"
 
+SLATE_IMPLEMENT_WIDGET(SSpacer)
+void SSpacer::PrivateRegisterAttributes(FSlateAttributeInitializer& AttributeInitializer)
+{
+	SLATE_ADD_MEMBER_ATTRIBUTE_DEFINITION(AttributeInitializer, SpacerSize, EInvalidateWidgetReason::Layout);
+}
+
+SSpacer::SSpacer()
+	: SpacerSize(*this)
+	, bIsSpacerSizeBound(false)
+{
+	SetCanTick(false);
+	bCanSupportFocus = false;
+}
 
-/**
- * Construct this widget
- *
- * @param	InArgs	The declaration data for this widget
- */
 void SSpacer::Construct( const FArguments& InArgs )
 {
-	SpacerSize = InArgs._Size;
+	SetSize(InArgs._Size);
 }
 
 
diff --git a/Engine/Source/Runtime/Slate/Public/Widgets/Layout/SSpacer.h b/Engine/Source/Runtime/Slate/Public/Widgets/Layout/SSpacer.h
--- a/Engine/Source/Runtime/Slate/Public/Widgets/Layout/SSpacer.h
+++ b/Engine/Source/Runtime/Slate/Public/Widgets/Layout/SSpacer.h
@@ -13,6 +13,8 @@ class FSlateWindowElementList;
 
 class SLATE_API SSpacer : public SLeafWidget
 {
+	SLATE_DECLARE_WIDGET(SSpacer, SLeafWidget)
+
 public:
 
 	SLATE_BEGIN_ARGS( SSpacer )
@@ -24,11 +26,7 @@ public:
 		SLATE_ATTRIBUTE( FVector2D, Size )
 	SLATE_END_ARGS()
 
-	SSpacer()
-	{
-		SetCanTick(false);
-		bCanSupportFocus = false;
-	}
+	SSpacer();
 
 	/**
 	 * Construct this widget
@@ -43,20 +41,26 @@ public:
 
 	FVector2D GetSize() const
 	{
+		if (bIsSpacerSizeBound)
+		{
+			SSpacer& MutableSelf = const_cast<SSpacer&>(*this);
+			MutableSelf.SpacerSize.UpdateNow(MutableSelf);
+		}
 		return SpacerSize.Get();
 	}
 
 	void SetSize( TAttribute<FVector2D> InSpacerSize )
 	{
-		SpacerSize = InSpacerSize;
-		Invalidate(EInvalidateWidgetReason::Layout);
+		bIsSpacerSizeBound = InSpacerSize.IsBound();
+		SpacerSize.Assign(*this, InSpacerSize);
 	}
 
 protected:
-	// Begin SWidget overrides.
+	//~ Begin SWidget overrides.
 	virtual FVector2D ComputeDesiredSize(float) const override;
-	// End SWidget overrides.
+	//~ End SWidget overrides.
 
 private:
-	TAttribute<FVector2D> SpacerSize;
+	TSlateAttribute<FVector2D> SpacerSize;
+	bool bIsSpacerSizeBound;
 };
