Slate: Convert SViewport TAttributes to TSlateAttributes.

diff --git a/Engine/Source/Runtime/Slate/Private/Widgets/SViewport.cpp b/Engine/Source/Runtime/Slate/Private/Widgets/SViewport.cpp
--- a/Engine/Source/Runtime/Slate/Private/Widgets/SViewport.cpp
+++ b/Engine/Source/Runtime/Slate/Private/Widgets/SViewport.cpp
@@ -26,7 +26,9 @@ FVector2D SViewport::FArguments::GetDefaultViewportSize()
  *****************************************************************************/
 
 SViewport::SViewport()
-	: bRenderDirectlyToWindow(false)
+	: ShowDisabledEffect(*this)
+	, ViewportSize(*this)
+	, bRenderDirectlyToWindow(false)
 	, bEnableGammaCorrection(true)
 	, bReverseGammaCorrection(false)
 	, bEnableStereoRendering(false)
@@ -38,7 +40,7 @@ SViewport::SViewport()
 
 void SViewport::Construct( const FArguments& InArgs )
 {
-	ShowDisabledEffect = InArgs._ShowEffectWhenDisabled;
+	ShowDisabledEffect.Assign(*this, InArgs._ShowEffectWhenDisabled);
 	bRenderDirectlyToWindow = InArgs._RenderDirectlyToWindow;
 	bEnableGammaCorrection = InArgs._EnableGammaCorrection;
 	bReverseGammaCorrection = InArgs._ReverseGammaCorrection;
@@ -47,7 +49,7 @@ void SViewport::Construct( const FArguments& InArgs )
 	bIgnoreTextureAlpha = InArgs._IgnoreTextureAlpha;
 	bPreMultipliedAlpha = InArgs._PreMultipliedAlpha;
 	ViewportInterface = InArgs._ViewportInterface;
-	ViewportSize = InArgs._ViewportSize;
+	ViewportSize.Assign(*this, InArgs._ViewportSize);
 
 #if UE_WITH_SLATE_SIMULATEDNAVIGATIONMETADATA
 	AddMetadata(MakeShared<FSimulatedNavigationMetaData>(EUINavigationRule::Stop));
@@ -61,7 +63,7 @@ void SViewport::Construct( const FArguments& InArgs )
 
 void SViewport::SetActive(bool bActive)
 {
-	// In game enviroments the viewport is always active
+	// In game environments the viewport is always active
 	if(GIsEditor || IS_PROGRAM)
 	{
 		if (bActive && !ActiveTimerHandle.IsValid())
@@ -297,7 +299,11 @@ void SViewport::SetContent( TSharedPtr<SWidget> InContent )
 
 void SViewport::SetCustomHitTestPath( TSharedPtr<ICustomHitTestPath> InCustomHitTestPath )
 {
-	CustomHitTestPath = InCustomHitTestPath;
+	if (CustomHitTestPath != InCustomHitTestPath)
+	{
+		CustomHitTestPath = InCustomHitTestPath;
+		Invalidate(EInvalidateWidgetReason::Paint);
+	}
 }
 
 TSharedPtr<ICustomHitTestPath> SViewport::GetCustomHitTestPath()
diff --git a/Engine/Source/Runtime/Slate/Public/Widgets/SViewport.h b/Engine/Source/Runtime/Slate/Public/Widgets/SViewport.h
--- a/Engine/Source/Runtime/Slate/Public/Widgets/SViewport.h
+++ b/Engine/Source/Runtime/Slate/Public/Widgets/SViewport.h
@@ -114,7 +114,11 @@ public:
 	 */
 	void SetViewportInterface( TSharedRef<ISlateViewport> InViewportInterface )
 	{
-		ViewportInterface = InViewportInterface;
+		if (ViewportInterface != InViewportInterface)
+		{
+			ViewportInterface = InViewportInterface;
+			Invalidate(EInvalidateWidgetReason::Paint);
+		}
 	}
 
 	/**
@@ -170,7 +174,11 @@ public:
 	 */
 	void SetRenderDirectlyToWindow( const bool bInRenderDirectlyToWindow )
 	{
-		bRenderDirectlyToWindow = bInRenderDirectlyToWindow;
+		if (bRenderDirectlyToWindow != bInRenderDirectlyToWindow)
+		{
+			bRenderDirectlyToWindow = bInRenderDirectlyToWindow;
+			Invalidate(EInvalidateWidgetReason::Paint);
+		}
 	}
 
 	/**
@@ -181,7 +189,11 @@ public:
 	 */
 	void SetIgnoreTextureAlpha(const bool bInIgnoreTextureAlpha)
 	{
-		bIgnoreTextureAlpha = bInIgnoreTextureAlpha;
+		if (bIgnoreTextureAlpha != bInIgnoreTextureAlpha)
+		{
+			bIgnoreTextureAlpha = bInIgnoreTextureAlpha;
+			Invalidate(EInvalidateWidgetReason::Paint);
+		}
 	}
 
 	/** @return Whether or not to ignore texture alpha when blending */
@@ -265,10 +277,10 @@ private:
 	TWeakPtr<FActiveTimerHandle> ActiveTimerHandle;
 
 	/** Whether or not to show the disabled effect when this viewport is disabled. */
-	TAttribute<bool> ShowDisabledEffect;
+	TSlateAttribute<bool, EInvalidateWidgetReason::Paint> ShowDisabledEffect;
 
 	/** Size of the viewport. */
-	TAttribute<FVector2D> ViewportSize;
+	TSlateAttribute<FVector2D, EInvalidateWidgetReason::Layout> ViewportSize;
 
 	TSharedPtr<ICustomHitTestPath> CustomHitTestPath;
 
