Slate: Change color widget attributes to SlateAttributes.

diff --git a/Engine/Source/Runtime/AppFramework/Private/Widgets/Colors/SColorPicker.cpp b/Engine/Source/Runtime/AppFramework/Private/Widgets/Colors/SColorPicker.cpp
--- a/Engine/Source/Runtime/AppFramework/Private/Widgets/Colors/SColorPicker.cpp
+++ b/Engine/Source/Runtime/AppFramework/Private/Widgets/Colors/SColorPicker.cpp
@@ -881,7 +881,7 @@ TSharedRef<SWidget> SColorPicker::MakeColorSlider( EColorPickerChannels Channel
 				.EndColor(this, &SColorPicker::HandleColorSliderEndColor, Channel)
 				.StartColor(this, &SColorPicker::HandleColorSliderStartColor, Channel)
 				.Orientation(Orient_Horizontal)
-				.UseSRGB(SharedThis(this), &SColorPicker::HandleColorPickerUseSRGB)
+				//.UseSRGB(SharedThis(this), &SColorPicker::HandleColorPickerUseSRGB)
 		]
 
 	+ SOverlay::Slot()
@@ -930,8 +930,8 @@ TSharedRef<SWidget> SColorPicker::MakeColorSpinBox( EColorPickerChannels Channel
 		GradientWidget = SNew(SSimpleGradient)
 			.StartColor(this, &SColorPicker::GetGradientStartColor, Channel)
 			.EndColor(this, &SColorPicker::GetGradientEndColor, Channel)
-			.HasAlphaBackground(Channel == EColorPickerChannels::Alpha)
-			.UseSRGB(SharedThis(this), &SColorPicker::HandleColorPickerUseSRGB);
+			.HasAlphaBackground(Channel == EColorPickerChannels::Alpha);
+			//.UseSRGB(SharedThis(this), &SColorPicker::HandleColorPickerUseSRGB);
 	}
 	
 	// create spin box
diff --git a/Engine/Source/Runtime/Slate/Private/Widgets/Colors/SColorBlock.cpp b/Engine/Source/Runtime/Slate/Private/Widgets/Colors/SColorBlock.cpp
--- a/Engine/Source/Runtime/Slate/Private/Widgets/Colors/SColorBlock.cpp
+++ b/Engine/Source/Runtime/Slate/Private/Widgets/Colors/SColorBlock.cpp
@@ -5,6 +5,31 @@
 #include "Styling/CoreStyle.h"
 
 
+SLATE_IMPLEMENT_WIDGET(SColorBlock)
+void SColorBlock::PrivateRegisterAttributes(FSlateAttributeInitializer& AttributeInitializer)
+{
+	SLATE_ADD_MEMBER_ATTRIBUTE_DEFINITION(AttributeInitializer, Color, EInvalidateWidgetReason::Paint);
+	SLATE_ADD_MEMBER_ATTRIBUTE_DEFINITION(AttributeInitializer, AlphaBackgroundBrush, EInvalidateWidgetReason::Paint);
+	SLATE_ADD_MEMBER_ATTRIBUTE_DEFINITION(AttributeInitializer, GradientCornerRadius, EInvalidateWidgetReason::Paint);
+	SLATE_ADD_MEMBER_ATTRIBUTE_DEFINITION(AttributeInitializer, ColorBlockSize, EInvalidateWidgetReason::Layout);
+	SLATE_ADD_MEMBER_ATTRIBUTE_DEFINITION(AttributeInitializer, AlphaDisplayMode, EInvalidateWidgetReason::Paint);
+	SLATE_ADD_MEMBER_ATTRIBUTE_DEFINITION(AttributeInitializer, ColorIsHSV, EInvalidateWidgetReason::Paint);
+	SLATE_ADD_MEMBER_ATTRIBUTE_DEFINITION(AttributeInitializer, ShowBackgroundForAlpha, EInvalidateWidgetReason::Paint);
+	SLATE_ADD_MEMBER_ATTRIBUTE_DEFINITION(AttributeInitializer, bUseSRGB, EInvalidateWidgetReason::Paint);
+}
+
+SColorBlock::SColorBlock()
+	: Color(*this)
+	, AlphaBackgroundBrush(*this)
+	, GradientCornerRadius(*this)
+	, ColorBlockSize(*this)
+	, AlphaDisplayMode(*this)
+	, ColorIsHSV(*this)
+	, ShowBackgroundForAlpha(*this)
+	, bUseSRGB(*this)
+{
+}
+
 /**
  * Construct this widget
  *
@@ -12,15 +37,15 @@
  */
 void SColorBlock::Construct(const FArguments& InArgs)
 {
-	Color = InArgs._Color;
-	AlphaBackgroundBrush = InArgs._AlphaBackgroundBrush;
-	GradientCornerRadius = InArgs._CornerRadius;
-	ColorIsHSV = InArgs._ColorIsHSV;
-	AlphaDisplayMode = InArgs._AlphaDisplayMode;
-	ShowBackgroundForAlpha = InArgs._ShowBackgroundForAlpha;
+	Color.Assign(*this, InArgs._Color);
+	AlphaBackgroundBrush.Assign(*this, InArgs._AlphaBackgroundBrush);
+	GradientCornerRadius.Assign(*this, InArgs._CornerRadius, FVector4(0.0f));
+	ColorBlockSize.Assign(*this, InArgs._Size);
 	MouseButtonDownHandler = InArgs._OnMouseButtonDown;
-	bUseSRGB = InArgs._UseSRGB;
-	ColorBlockSize = InArgs._Size;
+	AlphaDisplayMode.Assign(*this, InArgs._AlphaDisplayMode);
+	ColorIsHSV.Assign(*this, InArgs._ColorIsHSV);
+	ShowBackgroundForAlpha.Assign(*this, InArgs._ShowBackgroundForAlpha);
+	bUseSRGB.Assign(*this, InArgs._UseSRGB);
 }
 
 int32 SColorBlock::OnPaint(const FPaintArgs& Args, const FGeometry& AllottedGeometry, const FSlateRect& MyCullingRect, FSlateWindowElementList& OutDrawElements, int32 LayerId, const FWidgetStyle& InWidgetStyle, bool bParentEnabled) const
@@ -64,7 +89,7 @@ int32 SColorBlock::OnPaint(const FPaintArgs& Args, const FGeometry& AllottedGeom
 		MoveTemp(GradientStops),
 		(AllottedGeometry.GetLocalSize().X > AllottedGeometry.GetLocalSize().Y) ? Orient_Vertical : Orient_Horizontal,
 		DrawEffects,
-		GradientCornerRadius.Get(FVector4(0.0f))
+		GradientCornerRadius.Get()
 	);
 
 
diff --git a/Engine/Source/Runtime/Slate/Private/Widgets/Colors/SColorGradingWheel.cpp b/Engine/Source/Runtime/Slate/Private/Widgets/Colors/SColorGradingWheel.cpp
--- a/Engine/Source/Runtime/Slate/Private/Widgets/Colors/SColorGradingWheel.cpp
+++ b/Engine/Source/Runtime/Slate/Private/Widgets/Colors/SColorGradingWheel.cpp
@@ -4,7 +4,21 @@
 #include "Rendering/RenderingCommon.h"
 #include "Rendering/DrawElements.h"
 
+SLATE_IMPLEMENT_WIDGET(SColorGradingWheel)
+void SColorGradingWheel::PrivateRegisterAttributes(FSlateAttributeInitializer& AttributeInitializer)
+{
+	SLATE_ADD_MEMBER_ATTRIBUTE_DEFINITION_WITH_NAME(AttributeInitializer, "SelectedColor", SelectedColorAttribute, EInvalidateWidgetReason::Paint);
+	SLATE_ADD_MEMBER_ATTRIBUTE_DEFINITION_WITH_NAME(AttributeInitializer, "DesiredWheelSize", DesiredWheelSizeAttribute, EInvalidateWidgetReason::Layout);
+	SLATE_ADD_MEMBER_ATTRIBUTE_DEFINITION_WITH_NAME(AttributeInitializer, "ExponentDisplacement", ExponentDisplacementAttribute, EInvalidateWidgetReason::Paint);
+}
 
+SColorGradingWheel::SColorGradingWheel()
+	: SelectedColorAttribute(*this)
+	, DesiredWheelSizeAttribute(*this)
+	, ExponentDisplacementAttribute(*this)
+	, Union_IsAttributeSet(0)
+{
+}
 
 /* SColorGradingWheel methods
  *****************************************************************************/
@@ -13,23 +27,54 @@ void SColorGradingWheel::Construct(const FArguments& InArgs)
 {
 	Image = FCoreStyle::Get().GetBrush("ColorGradingWheel.HueValueCircle");
 	SelectorImage = FCoreStyle::Get().GetBrush("ColorGradingWheel.Selector");
-	SelectedColor = InArgs._SelectedColor;
-	DesiredWheelSize = InArgs._DesiredWheelSize;
-	ExponentDisplacement = InArgs._ExponentDisplacement;
+	SetSelectedColorAttribute(InArgs._SelectedColor);
+	SetDesiredWheelSizeAttribute(InArgs._DesiredWheelSize);
+	SetExponentDisplacementAttribute(InArgs._ExponentDisplacement);
 	OnMouseCaptureBegin = InArgs._OnMouseCaptureBegin;
 	OnMouseCaptureEnd = InArgs._OnMouseCaptureEnd;
 	OnValueChanged = InArgs._OnValueChanged;
 }
 
+void SColorGradingWheel::SetSelectedColorAttribute(TAttribute<FLinearColor> InSelectedColor)
+{
+	SelectedColorAttribute.Assign(*this, MoveTemp(InSelectedColor));
+}
+
+void SColorGradingWheel::SetDesiredWheelSizeAttribute(TAttribute<int32> InDesiredWheelSize)
+{
+	const bool bDesiredWheelSizeSetChanged = (bIsAttributeDesiredWheelSizeSet != InDesiredWheelSize.IsSet());
+	bIsAttributeDesiredWheelSizeSet = InDesiredWheelSize.IsSet();
+	const bool bAttributeWasAssgined = DesiredWheelSizeAttribute.Assign(*this, MoveTemp(InDesiredWheelSize));
+
+	// If the assign didn't invalidate the widget but the attribute set changed, then invalidate the widget.
+	if (bDesiredWheelSizeSetChanged && !bAttributeWasAssgined)
+	{
+		Invalidate(EInvalidateWidgetReason::Layout);
+	}
+}
+
+void SColorGradingWheel::SetExponentDisplacementAttribute(TAttribute<float> InExponentDisplacement)
+{
+	const bool bExponentDisplacementSetChanged = (bIsAttributeExponentDisplacementSet != InExponentDisplacement.IsSet());
+	bIsAttributeExponentDisplacementSet = InExponentDisplacement.IsSet();
+	const bool bAttributeWasAssgined = ExponentDisplacementAttribute.Assign(*this, MoveTemp(InExponentDisplacement), 1.f);
+
+	// If the assign didn't invalidate the widget but the attribute set changed, then invalidate the widget.
+	if (bExponentDisplacementSetChanged && !bAttributeWasAssgined)
+	{
+		Invalidate(EInvalidateWidgetReason::Paint);
+	}
+}
+
 
 /* SWidget overrides
  *****************************************************************************/
 
 FVector2D SColorGradingWheel::ComputeDesiredSize(float) const
 {
-	if (DesiredWheelSize.IsSet())
+	if (bIsAttributeDesiredWheelSizeSet)
 	{
-		int32 CachedDesiredWheelSize = DesiredWheelSize.Get();
+		int32 CachedDesiredWheelSize = DesiredWheelSizeAttribute.Get();
 		return FVector2D(CachedDesiredWheelSize, CachedDesiredWheelSize);
 	}
 	return Image->ImageSize;
@@ -46,11 +91,19 @@ FReply SColorGradingWheel::OnMouseButtonDown(const FGeometry& MyGeometry, const
 {
 	if (MouseEvent.GetEffectingButton() == EKeys::LeftMouseButton)
 	{
-		OnMouseCaptureBegin.ExecuteIfBound(SelectedColor.Get());
+		if (OnMouseCaptureBegin.IsBound())
+		{
+			SelectedColorAttribute.UpdateNow(*this);
+			OnMouseCaptureBegin.Execute(SelectedColorAttribute.Get());
+		}
 
 		if (!ProcessMouseAction(MyGeometry, MouseEvent, false))
 		{
-			OnMouseCaptureEnd.ExecuteIfBound(SelectedColor.Get());
+			if (OnMouseCaptureEnd.IsBound())
+			{
+				SelectedColorAttribute.UpdateNow(*this);
+				OnMouseCaptureEnd.Execute(SelectedColorAttribute.Get());
+			}
 			return FReply::Unhandled();
 		}
 
@@ -65,7 +118,11 @@ FReply SColorGradingWheel::OnMouseButtonUp(const FGeometry& MyGeometry, const FP
 {
 	if (MouseEvent.GetEffectingButton() == EKeys::LeftMouseButton && HasMouseCapture())
 	{
-		OnMouseCaptureEnd.ExecuteIfBound(SelectedColor.Get());
+		if (OnMouseCaptureEnd.IsBound())
+		{
+			SelectedColorAttribute.UpdateNow(*this);
+			OnMouseCaptureEnd.Execute(SelectedColorAttribute.Get());
+		}
 
 		return FReply::Handled().ReleaseMouseCapture();
 	}
@@ -94,9 +151,9 @@ int32 SColorGradingWheel::OnPaint(const FPaintArgs& Args, const FGeometry& Allot
 	const FVector2D& SelectorSize = SelectorImage->ImageSize;
 	FVector2D CircleSize = AllottedGeometry.GetLocalSize() - SelectorSize;
 	FVector2D AllottedGeometrySize = AllottedGeometry.GetLocalSize();
-	if (DesiredWheelSize.IsSet())
+	if (bIsAttributeDesiredWheelSizeSet)
 	{
-		int32 CachedDesiredWheelSize = DesiredWheelSize.Get();
+		int32 CachedDesiredWheelSize = DesiredWheelSizeAttribute.Get();
 		CircleSize.X = (float)CachedDesiredWheelSize - SelectorSize.X;
 		CircleSize.Y = (float)CachedDesiredWheelSize - SelectorSize.Y;
 		AllottedGeometrySize.X = CachedDesiredWheelSize;
@@ -130,12 +187,12 @@ int32 SColorGradingWheel::OnPaint(const FPaintArgs& Args, const FGeometry& Allot
 
 FVector2D SColorGradingWheel::CalcRelativePositionFromCenter() const
 {
-	float Hue = SelectedColor.Get().R;
-	float Saturation = SelectedColor.Get().G;
-	if (ExponentDisplacement.IsSet() && ExponentDisplacement.Get() != 1.0f && !FMath::IsNearlyEqual(ExponentDisplacement.Get(), 0.0f, 0.00001f))
+	float Hue = SelectedColorAttribute.Get().R;
+	float Saturation = SelectedColorAttribute.Get().G;
+	if (bIsAttributeExponentDisplacementSet && ExponentDisplacementAttribute.Get() != 1.0f && !FMath::IsNearlyEqual(ExponentDisplacementAttribute.Get(), 0.0f, 0.00001f))
 	{
 		//Use log curve to set the distance G value
-		Saturation = FMath::Pow(Saturation, 1.0f / ExponentDisplacement.Get());
+		Saturation = FMath::Pow(Saturation, 1.0f / ExponentDisplacementAttribute.Get());
 	}
 	float Angle = Hue / 180.0f * PI;
 	float Radius = Saturation;
@@ -147,9 +204,10 @@ FVector2D SColorGradingWheel::CalcRelativePositionFromCenter() const
 bool SColorGradingWheel::ProcessMouseAction(const FGeometry& MyGeometry, const FPointerEvent& MouseEvent, bool bProcessWhenOutsideColorWheel)
 {
 	FVector2D GeometrySize = MyGeometry.GetLocalSize();
-	if (DesiredWheelSize.IsSet())
+	if (bIsAttributeDesiredWheelSizeSet)
 	{
-		int32 CachedDesiredWheelSize = DesiredWheelSize.Get();
+		DesiredWheelSizeAttribute.UpdateNow(*this);
+		int32 CachedDesiredWheelSize = DesiredWheelSizeAttribute.Get();
 		GeometrySize.X = CachedDesiredWheelSize;
 		GeometrySize.Y = CachedDesiredWheelSize;
 	}
@@ -167,14 +225,15 @@ bool SColorGradingWheel::ProcessMouseAction(const FGeometry& MyGeometry, const F
 			Angle += 2.0f * PI;
 		}
 
-		FLinearColor NewColor = SelectedColor.Get();
+		SelectedColorAttribute.UpdateNow(*this);
+		FLinearColor NewColor = SelectedColorAttribute.Get();
 		{
 			NewColor.R = Angle * 180.0f * INV_PI;
 			float LinearRadius = FMath::Min(RelativeRadius, 1.0f);
-			if (ExponentDisplacement.IsSet() && ExponentDisplacement.Get() != 1.0f)
+			if (bIsAttributeExponentDisplacementSet && ExponentDisplacementAttribute.Get() != 1.0f)
 			{
 				//Use log curve to set the distance G value
-				float LogDistance = FMath::Pow(LinearRadius, ExponentDisplacement.Get());
+				float LogDistance = FMath::Pow(LinearRadius, ExponentDisplacementAttribute.Get());
 				NewColor.G = LogDistance;
 			}
 			else
diff --git a/Engine/Source/Runtime/Slate/Private/Widgets/Colors/SColorSpectrum.cpp b/Engine/Source/Runtime/Slate/Private/Widgets/Colors/SColorSpectrum.cpp
--- a/Engine/Source/Runtime/Slate/Private/Widgets/Colors/SColorSpectrum.cpp
+++ b/Engine/Source/Runtime/Slate/Private/Widgets/Colors/SColorSpectrum.cpp
@@ -5,6 +5,11 @@
 #include "Styling/CoreStyle.h"
 
 
+SColorSpectrum::SColorSpectrum()
+	: SelectedColor(*this)
+{
+}
+
 /* SColorSpectrum methods
  *****************************************************************************/
 
@@ -12,7 +17,7 @@ void SColorSpectrum::Construct( const FArguments& InArgs )
 {
 	Image = FCoreStyle::Get().GetBrush("ColorSpectrum.Spectrum");
 	SelectorImage = FCoreStyle::Get().GetBrush("ColorSpectrum.Selector");
-	SelectedColor = InArgs._SelectedColor;
+	SelectedColor.Assign(*this, InArgs._SelectedColor);
 
 	OnMouseCaptureBegin = InArgs._OnMouseCaptureBegin;
 	OnMouseCaptureEnd = InArgs._OnMouseCaptureEnd;
@@ -131,6 +136,7 @@ void SColorSpectrum::ProcessMouseAction(const FGeometry& MyGeometry, const FPoin
 	FVector2D NormalizedMousePosition = MyGeometry.AbsoluteToLocal(MouseEvent.GetScreenSpacePosition()) / MyGeometry.GetLocalSize();
 	NormalizedMousePosition = NormalizedMousePosition.ClampAxes(0.0f, 1.0f);
 
+	SelectedColor.UpdateNow(*this);
 	FLinearColor NewColor = SelectedColor.Get();
 	NewColor.R = 360.0f * NormalizedMousePosition.X;
 
diff --git a/Engine/Source/Runtime/Slate/Private/Widgets/Colors/SColorWheel.cpp b/Engine/Source/Runtime/Slate/Private/Widgets/Colors/SColorWheel.cpp
--- a/Engine/Source/Runtime/Slate/Private/Widgets/Colors/SColorWheel.cpp
+++ b/Engine/Source/Runtime/Slate/Private/Widgets/Colors/SColorWheel.cpp
@@ -4,6 +4,9 @@
 #include "Rendering/DrawElements.h"
 #include "Styling/CoreStyle.h"
 
+SColorWheel::SColorWheel()
+	: SelectedColor(*this)
+{}
 
 /* SColorWheel methods
  *****************************************************************************/
@@ -12,7 +15,7 @@ void SColorWheel::Construct(const FArguments& InArgs)
 {
 	Image = FCoreStyle::Get().GetBrush("ColorWheel.HueValueCircle");
 	SelectorImage = FCoreStyle::Get().GetBrush("ColorWheel.Selector");
-	SelectedColor = InArgs._SelectedColor;
+	SelectedColor.Assign(*this, InArgs._SelectedColor);
 
 	OnMouseCaptureBegin = InArgs._OnMouseCaptureBegin;
 	OnMouseCaptureEnd = InArgs._OnMouseCaptureEnd;
@@ -138,6 +141,7 @@ bool SColorWheel::ProcessMouseAction(const FGeometry& MyGeometry, const FPointer
 			Angle += 2.0f * PI;
 		}
 
+		SelectedColor.UpdateNow(*this);
 		FLinearColor NewColor = SelectedColor.Get();
 		{
 			NewColor.R = Angle * 180.0f * INV_PI;
diff --git a/Engine/Source/Runtime/Slate/Private/Widgets/Colors/SComplexGradient.cpp b/Engine/Source/Runtime/Slate/Private/Widgets/Colors/SComplexGradient.cpp
--- a/Engine/Source/Runtime/Slate/Private/Widgets/Colors/SComplexGradient.cpp
+++ b/Engine/Source/Runtime/Slate/Private/Widgets/Colors/SComplexGradient.cpp
@@ -7,6 +7,8 @@
 
 
 SComplexGradient::SComplexGradient()
+	: GradientColors(*this)
+	, DesiredSizeOverride(*this)
 {
 	SetCanTick(false);
 }
@@ -16,10 +18,10 @@ SComplexGradient::SComplexGradient()
 
 void SComplexGradient::Construct( const FArguments& InArgs )
 {
-	DesiredSizeOverride = InArgs._DesiredSizeOverride;
-	GradientColors = InArgs._GradientColors;
-	bHasAlphaBackground = InArgs._HasAlphaBackground.Get();
-	Orientation = InArgs._Orientation.Get();
+	DesiredSizeOverride.Assign(*this, InArgs._DesiredSizeOverride);
+	GradientColors.Assign(*this, InArgs._GradientColors);
+	bHasAlphaBackground = InArgs._HasAlphaBackground;
+	Orientation = InArgs._Orientation;
 }
 
 
diff --git a/Engine/Source/Runtime/Slate/Private/Widgets/Colors/SSimpleGradient.cpp b/Engine/Source/Runtime/Slate/Private/Widgets/Colors/SSimpleGradient.cpp
--- a/Engine/Source/Runtime/Slate/Private/Widgets/Colors/SSimpleGradient.cpp
+++ b/Engine/Source/Runtime/Slate/Private/Widgets/Colors/SSimpleGradient.cpp
@@ -5,14 +5,22 @@
 #include "Styling/CoreStyle.h"
 #include "Application/SlateWindowHelper.h"
 
+/* SSimpleGradient implementation
+ *****************************************************************************/
+SSimpleGradient::SSimpleGradient()
+	: StartColor(*this)
+	, EndColor(*this)
+{
+	SetCanTick(false);
+}
 
 /* SSimpleGradient interface
  *****************************************************************************/
 
 void SSimpleGradient::Construct( const FArguments& InArgs )
 {
-	StartColor = InArgs._StartColor;
-	EndColor = InArgs._EndColor;
+	StartColor.Assign(*this, InArgs._StartColor);
+	EndColor.Assign(*this, InArgs._EndColor);
 	bHasAlphaBackground = InArgs._HasAlphaBackground.Get();
 	Orientation = InArgs._Orientation.Get();
 }
diff --git a/Engine/Source/Runtime/Slate/Public/Widgets/Colors/SColorBlock.h b/Engine/Source/Runtime/Slate/Public/Widgets/Colors/SColorBlock.h
--- a/Engine/Source/Runtime/Slate/Public/Widgets/Colors/SColorBlock.h
+++ b/Engine/Source/Runtime/Slate/Public/Widgets/Colors/SColorBlock.h
@@ -24,6 +24,8 @@ enum class EColorBlockAlphaDisplayMode : uint8
 
 class SLATE_API SColorBlock : public SLeafWidget
 {
+	SLATE_DECLARE_WIDGET(SColorBlock, SLeafWidget)
+
 public:
 
 	SLATE_BEGIN_ARGS(SColorBlock)
@@ -77,6 +79,9 @@ public:
 
 	SLATE_END_ARGS()
 
+public:
+	SColorBlock();
+
 	/**
 	 * Construct this widget
 	 *
@@ -91,29 +96,29 @@ private:
 	virtual FVector2D ComputeDesiredSize(float) const override;
 
 	void MakeSection(TArray<FSlateGradientStop>& OutGradientStops, FVector2D StartPt, FVector2D EndPt, FLinearColor Color, const FWidgetStyle& InWidgetStyle, bool bIgnoreAlpha) const;
-private:
 
+private:
 	/** The color to display for this color block */
-	TAttribute<FLinearColor> Color;
+	TSlateAttribute<FLinearColor> Color;
 
-	TAttribute<const FSlateBrush*> AlphaBackgroundBrush;
+	TSlateAttribute<const FSlateBrush*> AlphaBackgroundBrush;
 
-	TAttribute<FVector4> GradientCornerRadius;
+	TSlateAttribute<FVector4> GradientCornerRadius;
 
-	/** Whether the color displayed is HSV or not */
-	TAttribute<bool> ColorIsHSV;
+	TSlateAttribute<FVector2D> ColorBlockSize;
+
+	/** A handler to activate when the mouse is pressed. */
+	FPointerEventHandler MouseButtonDownHandler;
 
 	/** Whether to ignore alpha entirely from the input color */
-	TAttribute<EColorBlockAlphaDisplayMode> AlphaDisplayMode;
+	TSlateAttribute<EColorBlockAlphaDisplayMode> AlphaDisplayMode;
+
+	/** Whether the color displayed is HSV or not */
+	TSlateAttribute<bool> ColorIsHSV;
 
 	/** Whether to display a background for viewing opacity. Irrelevant if ignoring alpha */
-	TAttribute<bool> ShowBackgroundForAlpha;
+	TSlateAttribute<bool> ShowBackgroundForAlpha;
 
 	/** Whether to display sRGB color */
-	TAttribute<bool> bUseSRGB;
-
-	/** A handler to activate when the mouse is pressed. */
-	FPointerEventHandler MouseButtonDownHandler;
-
-	TAttribute<FVector2D> ColorBlockSize;
+	TSlateAttribute<bool> bUseSRGB;
 };
diff --git a/Engine/Source/Runtime/Slate/Public/Widgets/Colors/SColorGradingWheel.h b/Engine/Source/Runtime/Slate/Public/Widgets/Colors/SColorGradingWheel.h
--- a/Engine/Source/Runtime/Slate/Public/Widgets/Colors/SColorGradingWheel.h
+++ b/Engine/Source/Runtime/Slate/Public/Widgets/Colors/SColorGradingWheel.h
@@ -15,6 +15,8 @@
 class SLATE_API SColorGradingWheel
 	: public SLeafWidget
 {
+	SLATE_DECLARE_WIDGET(SColorGradingWheel, SLeafWidget)
+
 public:
 
 	DECLARE_DELEGATE_OneParam(FOnColorGradingWheelMouseCapture, const FLinearColor&);
@@ -32,10 +34,8 @@ public:
 		/** The current color selected by the user. */
 		SLATE_ATTRIBUTE(FLinearColor, SelectedColor)
 		
-		/** The current color selected by the user. */
 		SLATE_ATTRIBUTE(int32, DesiredWheelSize)
 
-		/** The current color selected by the user. */
 		SLATE_ATTRIBUTE(float, ExponentDisplacement)
 		
 		/** Invoked when the mouse is pressed and a capture begins. */
@@ -50,6 +50,7 @@ public:
 	SLATE_END_ARGS()
 	
 public:
+	SColorGradingWheel();
 
 	/**
 	 * Construct this widget.
@@ -85,15 +86,26 @@ protected:
 	 */
 	bool ProcessMouseAction(const FGeometry& MyGeometry, const FPointerEvent& MouseEvent, bool bProcessWhenOutsideColorWheel);
 
+	/** */
+	void SetSelectedColorAttribute(TAttribute<FLinearColor> InSelectedColor);
+
+	/** */
+	void SetDesiredWheelSizeAttribute(TAttribute<int32> InDesiredWheelSize);
+
+	/** */
+	void SetExponentDisplacementAttribute(TAttribute<float> InExponentDisplacement);
+
+	/** @return an attribute reference of SelectedColor */
+	TSlateAttributeRef<FLinearColor> GetSelectedColorAttribute() const { return TSlateAttributeRef<FLinearColor>(*this, SelectedColorAttribute); }
+
+	/** @return an attribute reference of DesiredWheelSize */
+	TSlateAttributeRef<int32> GetDesiredWheelSizeAttribute() const { return TSlateAttributeRef<int32>(*this, DesiredWheelSizeAttribute); }
+
+	/** @return an attribute reference of ExponentDisplacement */
+	TSlateAttributeRef<float> GetExponentDisplacementAttribute() const { return TSlateAttributeRef<float>(*this, ExponentDisplacementAttribute); }
+
 	/** The color wheel image to show. */
 	const FSlateBrush* Image;
-	
-	/** The current color selected by the user. */
-	TAttribute< FLinearColor > SelectedColor;
-	
-	TAttribute< int32 > DesiredWheelSize;
-
-	TAttribute<float> ExponentDisplacement;
 
 	/** The color selector image to show. */
 	const FSlateBrush* SelectorImage;
@@ -106,4 +118,31 @@ protected:
 
 	/** Invoked when a new value is selected on the color wheel. */
 	FOnColorGradingWheelValueChanged OnValueChanged;
+
+#if WITH_EDITORONLY_DATA
+	UE_DEPRECATED(5.0, "Direct access to SelectedColor is now deprecated. Use the setter or getter.")
+	FSlateDeprecatedTAttribute<FLinearColor> SelectedColor;
+	UE_DEPRECATED(5.0, "Direct access to DesiredWheelSize is now deprecated. Use the setter or getter.")
+	FSlateDeprecatedTAttribute<int32> DesiredWheelSize;
+	UE_DEPRECATED(5.0, "Direct access to DesiredSizeOverride is now deprecated. Use the setter or getter.")
+	FSlateDeprecatedTAttribute<float> ExponentDisplacement;
+#endif
+
+private:
+	/** The current color selected by the user. */
+	TSlateAttribute<FLinearColor> SelectedColorAttribute;
+
+	TSlateAttribute<int32> DesiredWheelSizeAttribute;
+	TSlateAttribute<float> ExponentDisplacementAttribute;
+
+	/** Flags used to check if the SlateAttribute is set. */
+	union
+	{
+		struct
+		{
+			uint8 bIsAttributeDesiredWheelSizeSet : 1;
+			uint8 bIsAttributeExponentDisplacementSet : 1;
+		};
+		uint8 Union_IsAttributeSet;
+	};
 };
diff --git a/Engine/Source/Runtime/Slate/Public/Widgets/Colors/SColorSpectrum.h b/Engine/Source/Runtime/Slate/Public/Widgets/Colors/SColorSpectrum.h
--- a/Engine/Source/Runtime/Slate/Public/Widgets/Colors/SColorSpectrum.h
+++ b/Engine/Source/Runtime/Slate/Public/Widgets/Colors/SColorSpectrum.h
@@ -43,6 +43,7 @@ public:
 	SLATE_END_ARGS()
 	
 public:
+	SColorSpectrum();
 
 	/**
 	 * Construct this widget.
@@ -82,7 +83,7 @@ private:
 	const FSlateBrush* Image;
 	
 	// The current color selected by the user.
-	TAttribute< FLinearColor > SelectedColor;
+	TSlateAttribute<FLinearColor, EInvalidateWidgetReason::Paint> SelectedColor;
 
 	// The color selector image to show.
 	const FSlateBrush* SelectorImage;
diff --git a/Engine/Source/Runtime/Slate/Public/Widgets/Colors/SColorWheel.h b/Engine/Source/Runtime/Slate/Public/Widgets/Colors/SColorWheel.h
--- a/Engine/Source/Runtime/Slate/Public/Widgets/Colors/SColorWheel.h
+++ b/Engine/Source/Runtime/Slate/Public/Widgets/Colors/SColorWheel.h
@@ -43,6 +43,7 @@ public:
 	SLATE_END_ARGS()
 	
 public:
+	SColorWheel();
 
 	/**
 	 * Construct this widget.
@@ -84,7 +85,7 @@ private:
 	const FSlateBrush* Image;
 	
 	/** The current color selected by the user. */
-	TAttribute< FLinearColor > SelectedColor;
+	TSlateAttribute<FLinearColor, EInvalidateWidgetReason::Paint> SelectedColor;
 
 	/** The color selector image to show. */
 	const FSlateBrush* SelectorImage;
diff --git a/Engine/Source/Runtime/Slate/Public/Widgets/Colors/SComplexGradient.h b/Engine/Source/Runtime/Slate/Public/Widgets/Colors/SComplexGradient.h
--- a/Engine/Source/Runtime/Slate/Public/Widgets/Colors/SComplexGradient.h
+++ b/Engine/Source/Runtime/Slate/Public/Widgets/Colors/SComplexGradient.h
@@ -30,10 +30,10 @@ public:
 		SLATE_ATTRIBUTE(TArray<FLinearColor>, GradientColors)
 
 		/** Whether a checker background is displayed for alpha viewing */
-		SLATE_ATTRIBUTE(bool, HasAlphaBackground)
+		SLATE_ARGUMENT(bool, HasAlphaBackground)
 		
 		/** Horizontal or vertical gradient */
-		SLATE_ATTRIBUTE(EOrientation, Orientation)
+		SLATE_ARGUMENT(EOrientation, Orientation)
 
 		/** When specified use this as the gradients desired size */
 		SLATE_ATTRIBUTE(TOptional<FVector2D>, DesiredSizeOverride)
@@ -58,10 +58,10 @@ protected:
 private:
 
 	/** The colors used in the gradient. */
-	TAttribute<TArray<FLinearColor>> GradientColors;
+	TSlateAttribute<TArray<FLinearColor>, EInvalidateWidgetReason::Paint> GradientColors;
 
-	/** Optiona override for desired size */
-	TAttribute<TOptional<FVector2D>> DesiredSizeOverride;
+	/** Optional override for desired size */
+	TSlateAttribute<TOptional<FVector2D>, EInvalidateWidgetReason::Layout> DesiredSizeOverride;
 
 	/** Whether a checker background is displayed for alpha viewing. */
 	bool bHasAlphaBackground;
diff --git a/Engine/Source/Runtime/Slate/Public/Widgets/Colors/SSimpleGradient.h b/Engine/Source/Runtime/Slate/Public/Widgets/Colors/SSimpleGradient.h
--- a/Engine/Source/Runtime/Slate/Public/Widgets/Colors/SSimpleGradient.h
+++ b/Engine/Source/Runtime/Slate/Public/Widgets/Colors/SSimpleGradient.h
@@ -23,7 +23,6 @@ public:
 			, _EndColor(FLinearColor(1.0f, 1.0f, 1.0f))
 			, _HasAlphaBackground(false)
 			, _Orientation(Orient_Vertical)
-			, _UseSRGB(true)
 		{ }
 
 		/** The leftmost gradient color */
@@ -39,10 +38,16 @@ public:
 		SLATE_ATTRIBUTE(EOrientation, Orientation)
 
 		/** Whether to display sRGB color */
-		SLATE_ATTRIBUTE(bool, UseSRGB)
+		PRAGMA_DISABLE_DEPRECATION_WARNINGS
+		UE_DEPRECATED(5.0, "UseSRGB is unused")
+		SLATE_ARGUMENT_DEFAULT(bool, UseSRGB) = true;
+		PRAGMA_ENABLE_DEPRECATION_WARNINGS
 
 	SLATE_END_ARGS()
 
+public:
+	SSimpleGradient();
+
 	/**
 	 * Constructs the widget
 	 *
@@ -59,17 +64,14 @@ protected:
 private:
 
 	/** The leftmost gradient color */
-	TAttribute<FLinearColor> StartColor;
+	TSlateAttribute<FLinearColor, EInvalidateWidgetReason::Paint> StartColor;
 
 	/** The rightmost gradient color */
-	TAttribute<FLinearColor> EndColor;
+	TSlateAttribute<FLinearColor, EInvalidateWidgetReason::Paint> EndColor;
 
 	/** Whether a checker background is displayed for alpha viewing */
 	bool bHasAlphaBackground;
 	
 	/** Horizontal or vertical gradient */
 	EOrientation Orientation;
-
-	/** Whether to display sRGB color */
-	TAttribute<bool> UseSRGB;
 };
