Fix CIS. Move implementation to cpp to prevent missing include and to be backward compatible.
#rnx

diff --git a/Engine/Source/Runtime/Slate/Private/Widgets/SViewport.cpp b/Engine/Source/Runtime/Slate/Private/Widgets/SViewport.cpp
--- a/Engine/Source/Runtime/Slate/Private/Widgets/SViewport.cpp
+++ b/Engine/Source/Runtime/Slate/Private/Widgets/SViewport.cpp
@@ -61,6 +61,33 @@ void SViewport::Construct( const FArguments& InArgs )
 	];
 }
 
+void SViewport::SetViewportInterface(TSharedRef<ISlateViewport> InViewportInterface)
+{
+	if (ViewportInterface != InViewportInterface)
+	{
+		ViewportInterface = InViewportInterface;
+		Invalidate(EInvalidateWidgetReason::Paint);
+	}
+}
+
+void SViewport::SetRenderDirectlyToWindow(const bool bInRenderDirectlyToWindow)
+{
+	if (bRenderDirectlyToWindow != bInRenderDirectlyToWindow)
+	{
+		bRenderDirectlyToWindow = bInRenderDirectlyToWindow;
+		Invalidate(EInvalidateWidgetReason::Paint);
+	}
+}
+
+void SViewport::SetIgnoreTextureAlpha(const bool bInIgnoreTextureAlpha)
+{
+	if (bIgnoreTextureAlpha != bInIgnoreTextureAlpha)
+	{
+		bIgnoreTextureAlpha = bInIgnoreTextureAlpha;
+		Invalidate(EInvalidateWidgetReason::Paint);
+	}
+}
+
 void SViewport::SetActive(bool bActive)
 {
 	// In game environments the viewport is always active
diff --git a/Engine/Source/Runtime/Slate/Public/Widgets/SViewport.h b/Engine/Source/Runtime/Slate/Public/Widgets/SViewport.h
--- a/Engine/Source/Runtime/Slate/Public/Widgets/SViewport.h
+++ b/Engine/Source/Runtime/Slate/Public/Widgets/SViewport.h
@@ -112,14 +112,7 @@ public:
 	 *
 	 * @param InViewportInterface The interface to use
 	 */
-	void SetViewportInterface( TSharedRef<ISlateViewport> InViewportInterface )
-	{
-		if (ViewportInterface != InViewportInterface)
-		{
-			ViewportInterface = InViewportInterface;
-			Invalidate(EInvalidateWidgetReason::Paint);
-		}
-	}
+	void SetViewportInterface( TSharedRef<ISlateViewport> InViewportInterface );
 
 	/**
 	 * Sets the interface to be used by this viewport for rendering and I/O
@@ -172,14 +165,7 @@ public:
 	 * 
 	 * @param	bInRenderDirectlyToWindow	Whether we should be able to render to the back buffer
 	 */
-	void SetRenderDirectlyToWindow( const bool bInRenderDirectlyToWindow )
-	{
-		if (bRenderDirectlyToWindow != bInRenderDirectlyToWindow)
-		{
-			bRenderDirectlyToWindow = bInRenderDirectlyToWindow;
-			Invalidate(EInvalidateWidgetReason::Paint);
-		}
-	}
+	void SetRenderDirectlyToWindow( const bool bInRenderDirectlyToWindow );
 
 	/**
 	 * If true, the viewport's texture alpha is ignored when performing blending.  In this case only the viewport tint opacity is used
@@ -187,14 +173,7 @@ public:
 	 * 
 	 * @param bIgnoreTextureAlpha If texture alpha should be ignored when blending.
 	 */
-	void SetIgnoreTextureAlpha(const bool bInIgnoreTextureAlpha)
-	{
-		if (bIgnoreTextureAlpha != bInIgnoreTextureAlpha)
-		{
-			bIgnoreTextureAlpha = bInIgnoreTextureAlpha;
-			Invalidate(EInvalidateWidgetReason::Paint);
-		}
-	}
+	void SetIgnoreTextureAlpha(const bool bInIgnoreTextureAlpha);
 
 	/** @return Whether or not to ignore texture alpha when blending */
 	bool GetIgnoreTextureAlpha(void) const
