The standard duplicate command now defaults to Ctrl+D instead of Ctrl+W to bring it inline with other dcc tools

diff --git a/Engine/Source/Runtime/Slate/Private/Framework/Commands/GenericCommands.cpp b/Engine/Source/Runtime/Slate/Private/Framework/Commands/GenericCommands.cpp
--- a/Engine/Source/Runtime/Slate/Private/Framework/Commands/GenericCommands.cpp
+++ b/Engine/Source/Runtime/Slate/Private/Framework/Commands/GenericCommands.cpp
@@ -6,21 +6,21 @@
 
 void FGenericCommands::RegisterCommands()
 {
-	UI_COMMAND( Cut, "Cut", "Cut selection", EUserInterfaceActionType::Button, FInputChord( EModifierKey::Control, EKeys::X ) )
-	UI_COMMAND( Copy, "Copy", "Copy selection", EUserInterfaceActionType::Button, FInputChord( EModifierKey::Control, EKeys::C ) )
-	UI_COMMAND( Paste, "Paste", "Paste clipboard contents", EUserInterfaceActionType::Button, FInputChord( EModifierKey::Control, EKeys::V ) )
+	UI_COMMAND(Cut, "Cut", "Cut selection", EUserInterfaceActionType::Button, FInputChord(EModifierKey::Control, EKeys::X))
+	UI_COMMAND(Copy, "Copy", "Copy selection", EUserInterfaceActionType::Button, FInputChord(EModifierKey::Control, EKeys::C))
+	UI_COMMAND(Paste, "Paste", "Paste clipboard contents", EUserInterfaceActionType::Button, FInputChord(EModifierKey::Control, EKeys::V))
 #if PLATFORM_MAC
-	UI_COMMAND( Duplicate, "Duplicate", "Duplicate selection", EUserInterfaceActionType::Button, FInputChord( EModifierKey::Command, EKeys::W ) )
+	UI_COMMAND(Duplicate, "Duplicate", "Duplicate selection", EUserInterfaceActionType::Button, FInputChord(EModifierKey::Command, EKeys::D))
 #else
-	UI_COMMAND( Duplicate, "Duplicate", "Duplicate selection", EUserInterfaceActionType::Button, FInputChord( EModifierKey::Control, EKeys::W ) )
+	UI_COMMAND(Duplicate, "Duplicate", "Duplicate selection", EUserInterfaceActionType::Button, FInputChord(EModifierKey::Control, EKeys::D))
 #endif
-	UI_COMMAND( Undo, "Undo", "Undo last action", EUserInterfaceActionType::Button, FInputChord( EModifierKey::Control, EKeys::Z ) )
-	UI_COMMAND( Redo, "Redo", "Redo last action", EUserInterfaceActionType::Button, FInputChord( EModifierKey::Control, EKeys::Y ) )
+	UI_COMMAND(Undo, "Undo", "Undo last action", EUserInterfaceActionType::Button, FInputChord(EModifierKey::Control, EKeys::Z))
+	UI_COMMAND(Redo, "Redo", "Redo last action", EUserInterfaceActionType::Button, FInputChord(EModifierKey::Control, EKeys::Y))
 
 	UI_COMMAND(Delete, "Delete", "Delete current selection", EUserInterfaceActionType::Button, FInputChord(EKeys::Delete), FInputChord(EKeys::BackSpace))
-	
-	UI_COMMAND( Rename, "Rename", "Rename current selection", EUserInterfaceActionType::Button, FInputChord( EKeys::F2 ) )
-	UI_COMMAND( SelectAll, "Select All", "Select everything in the current scope", EUserInterfaceActionType::Button, FInputChord( EModifierKey::Control, EKeys::A ) )
+
+	UI_COMMAND(Rename, "Rename", "Rename current selection", EUserInterfaceActionType::Button, FInputChord(EKeys::F2))
+	UI_COMMAND(SelectAll, "Select All", "Select everything in the current scope", EUserInterfaceActionType::Button, FInputChord(EModifierKey::Control, EKeys::A))
 }
 
 #undef LOCTEXT_NAMESPACE
