Add docking/undocking listeners to enable docking support for native rendering browser windows for Bridge plugin.

diff --git a/Engine/Source/Runtime/Slate/Private/Framework/Docking/TabManager.cpp b/Engine/Source/Runtime/Slate/Private/Framework/Docking/TabManager.cpp
--- a/Engine/Source/Runtime/Slate/Private/Framework/Docking/TabManager.cpp
+++ b/Engine/Source/Runtime/Slate/Private/Framework/Docking/TabManager.cpp
@@ -2166,6 +2166,8 @@ void FTabManager::OnTabForegrounded( const TSharedPtr<SDockTab>& NewForegroundTa
 
 void FTabManager::OnTabRelocated( const TSharedRef<SDockTab>& RelocatedTab, const TSharedPtr<SWindow>& NewOwnerWindow )
 {
+	RelocatedTab->NotifyTabRelocated();
+
 	CleanupPointerArray(DockAreas);
 	RemoveTabFromCollapsedAreas( FTabMatcher( RelocatedTab->GetLayoutIdentifier() ) );
 	for (int32 DockAreaIndex=0; DockAreaIndex < DockAreas.Num(); ++DockAreaIndex)
diff --git a/Engine/Source/Runtime/Slate/Private/Widgets/Docking/SDockTab.cpp b/Engine/Source/Runtime/Slate/Private/Widgets/Docking/SDockTab.cpp
--- a/Engine/Source/Runtime/Slate/Private/Widgets/Docking/SDockTab.cpp
+++ b/Engine/Source/Runtime/Slate/Private/Widgets/Docking/SDockTab.cpp
@@ -446,6 +446,16 @@ void SDockTab::SetOnTabDrawerClosed(const FSimpleDelegate InDelegate)
 	OnTabDrawerClosedEvent = InDelegate;
 }
 
+void SDockTab::SetOnTabRelocated(const FSimpleDelegate InDelegate)
+{
+	OnTabRelocated = InDelegate;
+}
+
+void SDockTab::SetOnTabDraggedOverDockArea(const FSimpleDelegate InDelegate)
+{
+	OnTabDraggedOverDockArea = InDelegate;
+}
+
 TSharedRef<FTabManager> SDockTab::GetTabManager() const
 {
 	return MyTabManager.Pin().ToSharedRef();
@@ -505,6 +515,7 @@ float SDockTab::GetFlashValue() const
 void SDockTab::SetDraggedOverDockArea( const TSharedPtr<SDockingArea>& Area )
 {
 	DraggedOverDockingArea = Area;
+	OnTabDraggedOverDockArea.ExecuteIfBound();
 }
 
 bool SDockTab::HasSiblingTab(const FTabId& SiblingTabId, const bool TreatIndexNoneAsWildcard) const
@@ -533,6 +544,8 @@ void SDockTab::Construct( const FArguments& InArgs )
 	this->TabColorScale = InArgs._TabColorScale;
 	
 	OnTabDrawerClosedEvent = InArgs._OnTabDrawerClosed;
+	OnTabRelocated = InArgs._OnTabRelocated;
+	OnTabDraggedOverDockArea = InArgs._OnTabDraggedOverDockArea;
 
 	MajorTabStyle = &FAppStyle::Get().GetWidgetStyle<FDockTabStyle>("Docking.MajorTab");
 	GenericTabStyle = &FAppStyle::Get().GetWidgetStyle<FDockTabStyle>("Docking.Tab");
@@ -713,6 +726,11 @@ void SDockTab::OnTabDrawerClosed()
 	OnTabDrawerClosedEvent.ExecuteIfBound();
 }
 
+void SDockTab::NotifyTabRelocated()
+{
+	OnTabRelocated.ExecuteIfBound();
+}
+
 const FDockTabStyle& SDockTab::GetCurrentStyle() const
 {
 	if ( GetVisualTabRole() == ETabRole::MajorTab )
diff --git a/Engine/Source/Runtime/Slate/Public/Widgets/Docking/SDockTab.h b/Engine/Source/Runtime/Slate/Public/Widgets/Docking/SDockTab.h
--- a/Engine/Source/Runtime/Slate/Public/Widgets/Docking/SDockTab.h
+++ b/Engine/Source/Runtime/Slate/Public/Widgets/Docking/SDockTab.h
@@ -76,6 +76,8 @@ public:
 		, _Icon( FStyleDefaults::GetNoBrush() )
 		, _OnTabClosed()
 		, _OnTabActivated()
+		, _OnTabRelocated()
+		, _OnTabDraggedOverDockArea()
 		, _ShouldAutosize(false)
 		, _OnCanCloseTab()
 		, _OnPersistVisualState()
@@ -93,6 +95,8 @@ public:
 		SLATE_ATTRIBUTE( const FSlateBrush*, Icon )
 		SLATE_EVENT( FOnTabClosedCallback, OnTabClosed )
 		SLATE_EVENT( FOnTabActivatedCallback, OnTabActivated )
+		SLATE_EVENT( FSimpleDelegate, OnTabRelocated )
+		SLATE_EVENT( FSimpleDelegate, OnTabDraggedOverDockArea )
 		SLATE_ARGUMENT( bool, ShouldAutosize )
 		SLATE_EVENT( FCanCloseTab, OnCanCloseTab )
 		SLATE_EVENT( FOnPersistVisualState, OnPersistVisualState )
@@ -248,6 +252,12 @@ public:
 	/** Set the handler that will be invoked when the tab is activated */
 	void SetOnTabActivated( const FOnTabActivatedCallback& InDelegate );
 
+	/** Set the handler that will be invoked when the tab is relocated to a new tab well */
+	void SetOnTabRelocated(const FSimpleDelegate InDelegate);
+
+	/** Set the handler that will be invoked when the tab is dragged over dock area */
+	void SetOnTabDraggedOverDockArea(const FSimpleDelegate InDelegate);
+
 	/** Set the handler that will be invoked when the tab is renamed */
 	void SetOnTabRenamed(const FOnTabRenamed& InDelegate);
 
@@ -351,6 +361,8 @@ private:
 	void OnTabDrawerOpened();
 	void OnTabDrawerClosed();
 
+	void NotifyTabRelocated();
+
 	/** The handle to the active tab activation tick */
 	TWeakPtr<FActiveTimerHandle> DragDropTimerHandle;
 	TWeakPtr<FActiveTimerHandle> UpdateStyleTimerHandle;
@@ -396,6 +408,10 @@ protected:
 
 	FSimpleDelegate OnTabDrawerOpenedEvent;
 
+	FSimpleDelegate OnTabRelocated;
+
+	FSimpleDelegate OnTabDraggedOverDockArea;
+
 	/**
 	 * Invoked during the Save Visual State pass; gives this tab a chance to save misc info about visual state.
 	 * e.g. Content Browser might save the current filters, current folder, whether some panel is collapsed, etc.
