Make logging for regstering and unregistering tab spanwers verbose to cut down on log spam

diff --git a/Engine/Source/Runtime/Slate/Private/Framework/Docking/TabManager.cpp b/Engine/Source/Runtime/Slate/Private/Framework/Docking/TabManager.cpp
--- a/Engine/Source/Runtime/Slate/Private/Framework/Docking/TabManager.cpp
+++ b/Engine/Source/Runtime/Slate/Private/Framework/Docking/TabManager.cpp
@@ -918,7 +918,7 @@ FTabSpawnerEntry& FTabManager::RegisterTabSpawner(const FName TabId, const FOnSp
 	ensure(!TabSpawner.Contains(TabId));
 	ensure(!FGlobalTabmanager::Get()->IsLegacyTabType(TabId));
 
-	UE_LOG(LogSlate, Log, TEXT("Registering tab spawner for '%s' in layout '%s' on TabManager 0x%p"), *TabId.ToString(), *ActiveLayoutName.ToString(), this);
+	UE_LOG(LogSlate, Verbose, TEXT("Registering tab spawner for '%s' in layout '%s' on TabManager 0x%p"), *TabId.ToString(), *ActiveLayoutName.ToString(), this);
 
 	TSharedRef<FTabSpawnerEntry> NewSpawnerEntry = MakeShareable(new FTabSpawnerEntry(TabId, OnSpawnTab, CanSpawnTab));
 	TabSpawner.Add(TabId, NewSpawnerEntry);
@@ -927,7 +927,7 @@ FTabSpawnerEntry& FTabManager::RegisterTabSpawner(const FName TabId, const FOnSp
 
 bool FTabManager::UnregisterTabSpawner( const FName TabId )
 {
-	UE_LOG(LogSlate, Log, TEXT("Unregistering tab spawner for '%s' in layout '%s' on TabManager 0x%p"), *TabId.ToString(), *ActiveLayoutName.ToString(), this);
+	UE_LOG(LogSlate, Verbose, TEXT("Unregistering tab spawner for '%s' in layout '%s' on TabManager 0x%p"), *TabId.ToString(), *ActiveLayoutName.ToString(), this);
 
 	return TabSpawner.Remove( TabId ) > 0;
 }
@@ -937,7 +937,7 @@ void FTabManager::UnregisterAllTabSpawners()
 	if (TabSpawner.Num())
 	{
 		FString AllTabIdsString = FString::JoinBy(TabSpawner, TEXT(", "), [](const TPair<FName, TSharedRef<FTabSpawnerEntry>> Entry) { return Entry.Key.ToString(); });
-		UE_LOG(LogSlate, Log, TEXT("Unregistering all tab spawners on TabManager 0x%p for layout '%s': %s"), this, *ActiveLayoutName.ToString(), *AllTabIdsString);
+		UE_LOG(LogSlate, Verbose, TEXT("Unregistering all tab spawners on TabManager 0x%p for layout '%s': %s"), this, *ActiveLayoutName.ToString(), *AllTabIdsString);
 	}
 
 	TabSpawner.Empty();
@@ -2435,11 +2435,11 @@ FTabSpawnerEntry& FGlobalTabmanager::RegisterNomadTabSpawner(const FName TabId,
 	// Remove TabId if it was previously loaded. This allows re-loading the Editor UI layout without restarting the whole Editor (Window->Load Layout)
 	if (NomadTabSpawner->Contains(TabId))
 	{
-		UE_LOG(LogSlate, Log, TEXT("Unregistering nomad tab spawner for layout '%s': '%s'"), *ActiveLayoutName.ToString(), *TabId.ToString());
+		UE_LOG(LogSlate, Verbose, TEXT("Unregistering nomad tab spawner for layout '%s': '%s'"), *ActiveLayoutName.ToString(), *TabId.ToString());
 		UnregisterNomadTabSpawner(TabId);
 	}
 
-	UE_LOG(LogSlate, Log, TEXT("Registering nomad tab spawner for layout '%s': '%s'"), *ActiveLayoutName.ToString(), *TabId.ToString());
+	UE_LOG(LogSlate, Verbose, TEXT("Registering nomad tab spawner for layout '%s': '%s'"), *ActiveLayoutName.ToString(), *TabId.ToString());
 
 	// (Re)create and return NewSpawnerEntry
 	TSharedRef<FTabSpawnerEntry> NewSpawnerEntry = MakeShareable(new FTabSpawnerEntry(TabId, OnSpawnTab, CanSpawnTab));
@@ -2449,7 +2449,7 @@ FTabSpawnerEntry& FGlobalTabmanager::RegisterNomadTabSpawner(const FName TabId,
 
 void FGlobalTabmanager::UnregisterNomadTabSpawner( const FName TabId )
 {
-	UE_LOG(LogSlate, Log, TEXT("Unregistering nomad tab spawner for layout '%s': '%s'"), *ActiveLayoutName.ToString(), *TabId.ToString());
+	UE_LOG(LogSlate, Verbose, TEXT("Unregistering nomad tab spawner for layout '%s': '%s'"), *ActiveLayoutName.ToString(), *TabId.ToString());
 	const int32 NumRemoved = NomadTabSpawner->Remove(TabId);
 }
 
