Layouts are now also stored as JSON in the AppData/Local directory.

Added helpers to FTabManager::FLayout to work directly with JSON rather than roundtripping to strings.

diff --git a/Engine/Source/Runtime/Slate/Private/Framework/Docking/LayoutService.cpp b/Engine/Source/Runtime/Slate/Private/Framework/Docking/LayoutService.cpp
--- a/Engine/Source/Runtime/Slate/Private/Framework/Docking/LayoutService.cpp
+++ b/Engine/Source/Runtime/Slate/Private/Framework/Docking/LayoutService.cpp
@@ -1,7 +1,12 @@
 // Copyright Epic Games, Inc. All Rights Reserved.
 
 #include "Framework/Docking/LayoutService.h"
+#include "Misc/App.h"
 #include "Misc/ConfigCacheIni.h"
+#include "Misc/FileHelper.h"
+#include "Dom/JsonObject.h"
+#include "Serialization/JsonReader.h"
+#include "Serialization/JsonSerializer.h"
 
 DEFINE_LOG_CATEGORY_STATIC(LogLayoutService, Log, All);
 
@@ -14,6 +19,62 @@ const FString& FLayoutSaveRestore::GetAdditionalLayoutConfigIni()
 	return IniSectionAdditionalConfig;
 }
 
+static void SaveLayoutToJson(const FString& InConfigFileName, const TSharedRef<FTabManager::FLayout>& InLayoutToSave)
+{
+	const FString JsonFileName = FPaths::GetBaseFilename(InConfigFileName) + TEXT(".json");
+	const FString UserSettingsPath = FPaths::Combine(FPlatformProcess::UserSettingsDir(), FApp::GetEpicProductIdentifier(), TEXT("Editor"), JsonFileName);
+
+	TSharedPtr<FJsonObject> AllLayoutsObject;
+
+	FString ExistingJsonContents;
+	if (FFileHelper::LoadFileToString(ExistingJsonContents, *UserSettingsPath))
+	{
+		TSharedRef<TJsonReader<>> Reader = TJsonReaderFactory<>::Create(ExistingJsonContents);
+		FJsonSerializer::Deserialize(Reader, AllLayoutsObject);
+	}
+	
+	if (!AllLayoutsObject.IsValid())
+	{
+		// doesn't exist
+		AllLayoutsObject = MakeShared<FJsonObject>();
+	}
+
+	TSharedRef<FJsonObject> LayoutJson = InLayoutToSave->ToJson();
+	AllLayoutsObject->SetObjectField(InLayoutToSave->GetLayoutName().ToString(), LayoutJson);
+
+	FString NewJsonContents;
+	NewJsonContents.Reserve(ExistingJsonContents.Len());
+	TSharedRef<TJsonWriter<>> Writer = TJsonWriterFactory<>::Create(&NewJsonContents);
+	if (FJsonSerializer::Serialize(AllLayoutsObject.ToSharedRef(), Writer))
+	{
+		FFileHelper::SaveStringToFile(NewJsonContents, *UserSettingsPath);
+	}
+}
+
+static bool LoadLayoutFromJson(const FString& InConfigFileName, const FString& InLayoutName, TSharedPtr<FTabManager::FLayout>& OutLayout)
+{
+	const FString JsonFileName = FPaths::GetBaseFilename(InConfigFileName) + TEXT(".json");
+	const FString UserSettingsPath = FPaths::Combine(FPlatformProcess::UserSettingsDir(), FApp::GetEpicProductIdentifier(), TEXT("Editor"), JsonFileName);
+
+	FString ExistingJsonContents;
+	if (FFileHelper::LoadFileToString(ExistingJsonContents, *UserSettingsPath))
+	{
+		TSharedPtr<FJsonObject> AllLayoutsObject;
+
+		TSharedRef<TJsonReader<>> Reader = TJsonReaderFactory<>::Create(ExistingJsonContents);
+		if (FJsonSerializer::Deserialize(Reader, AllLayoutsObject))
+		{
+			const TSharedPtr<FJsonObject>* LayoutJson = nullptr;
+			if (AllLayoutsObject->TryGetObjectField(InLayoutName, LayoutJson))
+			{
+				OutLayout = FTabManager::FLayout::NewFromJson(*LayoutJson);
+				return true;
+			}
+		}
+	}
+
+	return false;
+}
 
 void FLayoutSaveRestore::SaveToConfig( const FString& InConfigFileName, const TSharedRef<FTabManager::FLayout>& InLayoutToSave )
 {
@@ -22,10 +83,11 @@ void FLayoutSaveRestore::SaveToConfig( const FString& InConfigFileName, const TS
 	{
 		const FString LayoutAsString = FLayoutSaveRestore::PrepareLayoutStringForIni(InLayoutToSave->ToString());
 		GConfig->SetString(EditorLayoutsSectionName, *InLayoutToSave->GetLayoutName().ToString(), *LayoutAsString, InConfigFileName);
+
+		SaveLayoutToJson(InConfigFileName, InLayoutToSave);
 	}
 }
 
-
 TSharedRef<FTabManager::FLayout> FLayoutSaveRestore::LoadFromConfig(const FString& InConfigFileName, const TSharedRef<FTabManager::FLayout>& InDefaultLayout,
 	const EOutputCanBeNullptr InPrimaryAreaOutputCanBeNullptr)
 {
@@ -45,21 +107,28 @@ TSharedRef<FTabManager::FLayout> FLayoutSaveRestore::LoadFromConfigPrivate(const
 	const EOutputCanBeNullptr InPrimaryAreaOutputCanBeNullptr, const bool bInRemoveOlderLayoutVersions, TArray<FString>& OutRemovedOlderLayoutVersions)
 {
 	const FString LayoutNameString = InDefaultLayout->GetLayoutName().ToString();
-	// If the Key (InDefaultLayout->GetLayoutName()) already exists in the section EditorLayoutsSectionName of the file InConfigFileName, try to load the layout from that file
-	FString UserLayoutString;
-	if (GConfig->GetString(EditorLayoutsSectionName, *LayoutNameString, UserLayoutString, InConfigFileName) && !UserLayoutString.IsEmpty())
+
+	TSharedPtr<FTabManager::FLayout> UserLayout;
+
+	// First try to load from JSON, then INI if that does not exist
+	if (!LoadLayoutFromJson(InConfigFileName, LayoutNameString, UserLayout))
+	{
+		FString IniLayoutString;
+		// If the Key (InDefaultLayout->GetLayoutName()) already exists in the section EditorLayoutsSectionName of the file InConfigFileName, try to load the layout from that file
+		GConfig->GetString(EditorLayoutsSectionName, *LayoutNameString, IniLayoutString, InConfigFileName);
+		UserLayout = FTabManager::FLayout::NewFromString( FLayoutSaveRestore::GetLayoutStringFromIni( IniLayoutString ) );
+	}
+
+	if ( UserLayout.IsValid() )
 	{
-		TSharedPtr<FTabManager::FLayout> UserLayout = FTabManager::FLayout::NewFromString( FLayoutSaveRestore::GetLayoutStringFromIni( UserLayoutString ));
-		if ( UserLayout.IsValid() && UserLayout->GetPrimaryArea().IsValid() )
+		// Return UserLayout in the following 2 cases:
+		// - By default (PrimaryAreaOutputCanBeNullptr = Never or IfNoTabValid)
+		// - For the case of PrimaryAreaOutputCanBeNullptr = IfNoOpenTabValid, only if the primary area has at least a valid open tab
+		if (UserLayout->GetPrimaryArea().IsValid() &&
+			InPrimaryAreaOutputCanBeNullptr != EOutputCanBeNullptr::IfNoOpenTabValid
+			|| FGlobalTabmanager::Get()->HasValidOpenTabs(UserLayout->GetPrimaryArea().Pin().ToSharedRef()))
 		{
-			// Return UserLayout in the following 2 cases:
-			// - By default (PrimaryAreaOutputCanBeNullptr = Never or IfNoTabValid)
-			// - For the case of PrimaryAreaOutputCanBeNullptr = IfNoOpenTabValid, only if the primary area has at least a valid open tab
-			if (InPrimaryAreaOutputCanBeNullptr != EOutputCanBeNullptr::IfNoOpenTabValid
-				|| FGlobalTabmanager::Get()->HasValidOpenTabs(UserLayout->GetPrimaryArea().Pin().ToSharedRef()))
-			{
-				return UserLayout.ToSharedRef();
-			}
+			return UserLayout.ToSharedRef();
 		}
 	}
 	// If the file layout could not be loaded and the caller wants to remove old fields
diff --git a/Engine/Source/Runtime/Slate/Private/Framework/Docking/TabManager.cpp b/Engine/Source/Runtime/Slate/Private/Framework/Docking/TabManager.cpp
--- a/Engine/Source/Runtime/Slate/Private/Framework/Docking/TabManager.cpp
+++ b/Engine/Source/Runtime/Slate/Private/Framework/Docking/TabManager.cpp
@@ -316,33 +316,40 @@ TSharedRef<FTabManager::FLayoutNode> FTabManager::FLayout::NewFromString_Helper(
 
 TSharedPtr<FTabManager::FLayout> FTabManager::FLayout::NewFromString( const FString& LayoutAsText )
 {
-	TSharedPtr<FTabManager::FLayout> Layout;
 	TSharedPtr<FJsonObject> JsonObject;
-	FString OutError;
 
 	TSharedRef< TJsonReader<> > Reader = TJsonReaderFactory<>::Create( LayoutAsText );
-	
 	if (FJsonSerializer::Deserialize( Reader, JsonObject ))
 	{
-		const FString LayoutName = JsonObject->GetStringField(TEXT("Name"));
-		TSharedRef<FTabManager::FLayout> NewLayout = FTabManager::NewLayout( *LayoutName );
-		int32 PrimaryAreaIndex = FMath::TruncToInt((float)JsonObject->GetNumberField(TEXT("PrimaryAreaIndex")) );
+		return NewFromJson( JsonObject );
+	}
+	
+	return TSharedPtr<FTabManager::FLayout>();
+}
 
-		TArray< TSharedPtr<FJsonValue> > Areas = JsonObject->GetArrayField(TEXT("Areas"));
-		for(int32 AreaIndex=0; AreaIndex < Areas.Num(); ++AreaIndex)
+TSharedPtr<FTabManager::FLayout> FTabManager::FLayout::NewFromJson( const TSharedPtr<FJsonObject>& LayoutAsJson )
+{
+	if (!LayoutAsJson.IsValid())
+	{
+		return TSharedPtr<FTabManager::FLayout>();
+	}
+
+	const FString LayoutName = LayoutAsJson->GetStringField(TEXT("Name"));
+	TSharedRef<FTabManager::FLayout> NewLayout = FTabManager::NewLayout( *LayoutName );
+	int32 PrimaryAreaIndex = FMath::TruncToInt((float)LayoutAsJson->GetNumberField(TEXT("PrimaryAreaIndex")) );
+
+	TArray< TSharedPtr<FJsonValue> > Areas = LayoutAsJson->GetArrayField(TEXT("Areas"));
+	for(int32 AreaIndex=0; AreaIndex < Areas.Num(); ++AreaIndex)
+	{
+		TSharedRef<FTabManager::FArea> NewArea = StaticCastSharedRef<FTabManager::FArea>( NewFromString_Helper( Areas[AreaIndex]->AsObject() ) );
+		NewLayout->AddArea( NewArea );
+		if (AreaIndex == PrimaryAreaIndex)
 		{
-			TSharedRef<FTabManager::FArea> NewArea = StaticCastSharedRef<FTabManager::FArea>( NewFromString_Helper( Areas[AreaIndex]->AsObject() ) );
-			NewLayout->AddArea( NewArea );
-			if (AreaIndex == PrimaryAreaIndex)
-			{
-				NewLayout->PrimaryArea = NewArea;
-			}
+			NewLayout->PrimaryArea = NewArea;
 		}
-		
-		Layout = NewLayout;
 	}
-	
-	return Layout;
+		
+	return NewLayout;
 }
 
 FName FTabManager::FLayout::GetLayoutName() const
@@ -350,7 +357,7 @@ FName FTabManager::FLayout::GetLayoutName() const
 	return LayoutName;
 }
 
-FString FTabManager::FLayout::ToString() const
+TSharedRef<FJsonObject> FTabManager::FLayout::ToJson() const
 {
 	TSharedRef<FJsonObject> LayoutJson = MakeShareable( new FJsonObject() );
 	LayoutJson->SetStringField( TEXT("Type"), TEXT("Layout") );
@@ -369,6 +376,13 @@ FString FTabManager::FLayout::ToString() const
 	}
 	LayoutJson->SetArrayField( TEXT("Areas"), AreasAsJson );
 
+	return LayoutJson;
+}
+
+FString FTabManager::FLayout::ToString() const
+{
+	TSharedRef<FJsonObject> LayoutJson = this->ToJson();
+
 	FString LayoutAsString;
 	TSharedRef< TJsonWriter<> > Writer = TJsonWriterFactory<>::Create( &LayoutAsString );
 	check( FJsonSerializer::Serialize( LayoutJson, Writer ) );
diff --git a/Engine/Source/Runtime/Slate/Public/Framework/Docking/TabManager.h b/Engine/Source/Runtime/Slate/Public/Framework/Docking/TabManager.h
--- a/Engine/Source/Runtime/Slate/Public/Framework/Docking/TabManager.h
+++ b/Engine/Source/Runtime/Slate/Public/Framework/Docking/TabManager.h
@@ -625,7 +625,9 @@ class SLATE_API FTabManager : public TSharedFromThis<FTabManager>
 				
 			public:
 				static TSharedPtr<FTabManager::FLayout> NewFromString( const FString& LayoutAsText );
+				static TSharedPtr<FTabManager::FLayout> NewFromJson( const TSharedPtr<FJsonObject>& LayoutAsJson );
 				FName GetLayoutName() const;
+				TSharedRef<FJsonObject> ToJson() const;
 				FString ToString() const;
 
 				void ProcessExtensions(const FLayoutExtender& Extender);
