Fixed Clang build error due to ambiguous || and && precedence.

diff --git a/Engine/Source/Runtime/Slate/Private/Framework/Docking/LayoutService.cpp b/Engine/Source/Runtime/Slate/Private/Framework/Docking/LayoutService.cpp
--- a/Engine/Source/Runtime/Slate/Private/Framework/Docking/LayoutService.cpp
+++ b/Engine/Source/Runtime/Slate/Private/Framework/Docking/LayoutService.cpp
@@ -119,14 +119,13 @@ TSharedRef<FTabManager::FLayout> FLayoutSaveRestore::LoadFromConfigPrivate(const
 		UserLayout = FTabManager::FLayout::NewFromString( FLayoutSaveRestore::GetLayoutStringFromIni( IniLayoutString ) );
 	}
 
-	if ( UserLayout.IsValid() )
+	if (UserLayout.IsValid())
 	{
 		// Return UserLayout in the following 2 cases:
 		// - By default (PrimaryAreaOutputCanBeNullptr = Never or IfNoTabValid)
 		// - For the case of PrimaryAreaOutputCanBeNullptr = IfNoOpenTabValid, only if the primary area has at least a valid open tab
-		if (UserLayout->GetPrimaryArea().IsValid() &&
-			InPrimaryAreaOutputCanBeNullptr != EOutputCanBeNullptr::IfNoOpenTabValid
-			|| FGlobalTabmanager::Get()->HasValidOpenTabs(UserLayout->GetPrimaryArea().Pin().ToSharedRef()))
+		TSharedPtr<FTabManager::FArea> PrimaryArea = UserLayout->GetPrimaryArea().Pin();
+		if (PrimaryArea.IsValid() && (InPrimaryAreaOutputCanBeNullptr != EOutputCanBeNullptr::IfNoOpenTabValid || FGlobalTabmanager::Get()->HasValidOpenTabs(PrimaryArea.ToSharedRef())))
 		{
 			return UserLayout.ToSharedRef();
 		}
