Add "Mask" and "Scale" styles for progress bar fill modes, implement them for every mode. Also adds "FillFromCenterVertical" and "FillFromCenterHorizontal". Breaking change: Default behavior for "FillFromCenter" changed from "Scale" to "Mask".

diff --git a/Engine/Source/Runtime/Slate/Private/Widgets/Notifications/SProgressBar.cpp b/Engine/Source/Runtime/Slate/Private/Widgets/Notifications/SProgressBar.cpp
--- a/Engine/Source/Runtime/Slate/Private/Widgets/Notifications/SProgressBar.cpp
+++ b/Engine/Source/Runtime/Slate/Private/Widgets/Notifications/SProgressBar.cpp
@@ -14,6 +14,7 @@ void SProgressBar::Construct( const FArguments& InArgs )
 
 	SetPercent(InArgs._Percent);
 	BarFillType = InArgs._BarFillType;
+	BarFillStyle = InArgs._BarFillStyle;
 	
 	BackgroundImage = InArgs._BackgroundImage;
 	FillImage = InArgs._FillImage;
@@ -65,6 +66,15 @@ void SProgressBar::SetBarFillType(EProgressBarFillType::Type InBarFillType)
 	}
 }
 
+void SProgressBar::SetBarFillStyle(EProgressBarFillStyle::Type InBarFillStyle)
+{
+	if (BarFillStyle != InBarFillStyle)
+	{
+		BarFillStyle = InBarFillStyle;
+		Invalidate(EInvalidateWidget::Paint);
+	}
+}
+
 void SProgressBar::SetFillColorAndOpacity(TAttribute< FSlateColor > InFillColorAndOpacity)
 {
 	if(!FillColorAndOpacity.IdenticalTo(InFillColorAndOpacity))
@@ -188,29 +198,49 @@ int32 SProgressBar::OnPaint( const FPaintArgs& Args, const FGeometry& AllottedGe
 			}
 		}
 
+		float MarqueeAnimOffsetX = 0.0f;
+		float MarqueeImageSizeX = 0.0f;
+		float MarqueeAnimOffsetY = 0.0f;
+		float MarqueeImageSizeY = 0.0f;
+
+		if (Style->EnableFillAnimation)
+		{
+			MarqueeAnimOffsetX = CurrentFillImage->ImageSize.X * MarqueeOffset;
+			MarqueeImageSizeX = CurrentFillImage->ImageSize.X;
+
+			MarqueeAnimOffsetY = CurrentFillImage->ImageSize.Y * MarqueeOffset;
+			MarqueeImageSizeY = CurrentFillImage->ImageSize.Y;
+		}
+
+		bool bScaleWithFillPerc = BarFillStyle == EProgressBarFillStyle::Scale;
+
 		const float ClampedFraction = FMath::Clamp(ProgressFraction.GetValue(), 0.0f, 1.0f);
 		switch (ComputedBarFillType)
 		{
-			case EProgressBarFillType::RightToLeft:
+			case EProgressBarFillType::FillFromCenter:
 			{
-				if (PushTransformedClip(OutDrawElements, AllottedGeometry, BorderPaddingRef, FVector2D(1, 0), FSlateRect(ClampedFraction, 0, 0, 1)))
+				float HalfFrac = ClampedFraction / 2.0f;
+				if (PushTransformedClip(OutDrawElements, AllottedGeometry, BorderPaddingRef, FVector2D(0.5, 0.5), FSlateRect(HalfFrac, HalfFrac, HalfFrac, HalfFrac)))
 				{
-					float MarqueeAnimOffset = 0;
-					float MarqueeImageSize = 0;
-
-					if (Style->EnableFillAnimation)
+					FPaintGeometry PaintRect;
+					if (bScaleWithFillPerc)
+					{
+						PaintRect = AllottedGeometry.ToPaintGeometry(
+							FVector2D((AllottedGeometry.GetLocalSize().X * 0.5f) - ((AllottedGeometry.GetLocalSize().X * (ClampedFraction)) * 0.5f), (AllottedGeometry.GetLocalSize().Y * 0.5f) - ((AllottedGeometry.GetLocalSize().Y * (ClampedFraction)) * 0.5f)),
+							FVector2D(AllottedGeometry.GetLocalSize().X * ClampedFraction, AllottedGeometry.GetLocalSize().Y * ClampedFraction));
+					}
+					else
 					{
-						MarqueeAnimOffset = CurrentFillImage->ImageSize.X * MarqueeOffset;
-						MarqueeImageSize = CurrentFillImage->ImageSize.X;
+						PaintRect = AllottedGeometry.ToPaintGeometry(
+							FVector2D(0.0f, 0.0f),
+							FVector2D(AllottedGeometry.GetLocalSize().X, AllottedGeometry.GetLocalSize().Y));
 					}
 
 					// Draw Fill
 					FSlateDrawElement::MakeBox(
 						OutDrawElements,
 						RetLayerId++,
-						AllottedGeometry.ToPaintGeometry(
-							FVector2D(MarqueeAnimOffset - MarqueeImageSize, 0.0f),
-							FVector2D(AllottedGeometry.GetLocalSize().X+MarqueeImageSize, AllottedGeometry.GetLocalSize().Y)),
+						PaintRect,
 						CurrentFillImage,
 						DrawEffects,
 						FillColorAndOpacitySRGB
@@ -220,22 +250,67 @@ int32 SProgressBar::OnPaint( const FPaintArgs& Args, const FGeometry& AllottedGe
 				}
 				break;
 			}
-			case EProgressBarFillType::FillFromCenter:
+			case EProgressBarFillType::FillFromCenterHorizontal:
 			{
 				float HalfFrac = ClampedFraction / 2.0f;
-				if (PushTransformedClip(OutDrawElements, AllottedGeometry, BorderPaddingRef, FVector2D(0.5, 0.5), FSlateRect(HalfFrac, HalfFrac, HalfFrac, HalfFrac)))
+				if (PushTransformedClip(OutDrawElements, AllottedGeometry, BorderPaddingRef, FVector2D(0.5, 0), FSlateRect(HalfFrac, 0, HalfFrac, 1)))
 				{
+					FPaintGeometry PaintRect;
+					if (bScaleWithFillPerc)
+					{
+						PaintRect = AllottedGeometry.ToPaintGeometry(
+							FVector2D((AllottedGeometry.GetLocalSize().X * 0.5f) - ((AllottedGeometry.GetLocalSize().X * (ClampedFraction)) * 0.5f), 0.0f),
+							FVector2D(AllottedGeometry.GetLocalSize().X * ClampedFraction, AllottedGeometry.GetLocalSize().Y));
+					}
+					else
+					{
+						PaintRect = AllottedGeometry.ToPaintGeometry(
+							FVector2D(0.0f, 0.0f),
+							FVector2D(AllottedGeometry.GetLocalSize().X, AllottedGeometry.GetLocalSize().Y));
+					}
+
 					// Draw Fill
 					FSlateDrawElement::MakeBox(
 						OutDrawElements,
 						RetLayerId++,
-						AllottedGeometry.ToPaintGeometry(
-							FVector2D((AllottedGeometry.GetLocalSize().X * 0.5f) - ((AllottedGeometry.GetLocalSize().X * ( ClampedFraction ))*0.5f), 0.0f),
-							FVector2D(AllottedGeometry.GetLocalSize().X * ( ClampedFraction ) , AllottedGeometry.GetLocalSize().Y )),
+						PaintRect,
 						CurrentFillImage,
 						DrawEffects,
 						FillColorAndOpacitySRGB
-						);
+					);
+
+					OutDrawElements.PopClip();
+				}
+				break;
+			}
+			case EProgressBarFillType::FillFromCenterVertical:
+			{
+				float HalfFrac = ClampedFraction / 2.0f;
+				if (PushTransformedClip(OutDrawElements, AllottedGeometry, BorderPaddingRef, FVector2D(0, 0.5), FSlateRect(0, HalfFrac, 1, HalfFrac)))
+				{
+					FPaintGeometry PaintRect;
+					if (bScaleWithFillPerc)
+					{
+						PaintRect = AllottedGeometry.ToPaintGeometry(
+							FVector2D(0.0f, (AllottedGeometry.GetLocalSize().Y * 0.5f) - ((AllottedGeometry.GetLocalSize().Y * (ClampedFraction)) * 0.5f)),
+							FVector2D(AllottedGeometry.GetLocalSize().X, AllottedGeometry.GetLocalSize().Y * ClampedFraction));
+					}
+					else
+					{
+						PaintRect = AllottedGeometry.ToPaintGeometry(
+							FVector2D(0.0f, 0.0f),
+							FVector2D(AllottedGeometry.GetLocalSize().X, AllottedGeometry.GetLocalSize().Y));
+					}
+
+					// Draw Fill
+					FSlateDrawElement::MakeBox(
+						OutDrawElements,
+						RetLayerId++,
+						PaintRect,
+						CurrentFillImage,
+						DrawEffects,
+						FillColorAndOpacitySRGB
+					);
 
 					OutDrawElements.PopClip();
 				}
@@ -245,22 +320,25 @@ int32 SProgressBar::OnPaint( const FPaintArgs& Args, const FGeometry& AllottedGe
 			{
 				if (PushTransformedClip(OutDrawElements, AllottedGeometry, BorderPaddingRef, FVector2D(0, 0), FSlateRect(0, 0, 1, ClampedFraction)))
 				{
-					float MarqueeAnimOffset = 0;
-					float MarqueeImageSize = 0;
-
-					if (Style->EnableFillAnimation)
+					FPaintGeometry PaintRect;
+					if (bScaleWithFillPerc)
 					{
-						MarqueeAnimOffset = CurrentFillImage->ImageSize.Y * MarqueeOffset;
-						MarqueeImageSize = CurrentFillImage->ImageSize.Y;
+						PaintRect = AllottedGeometry.ToPaintGeometry(
+							FVector2D(0.0f, 0.0f),
+							FVector2D(AllottedGeometry.GetLocalSize().X, AllottedGeometry.GetLocalSize().Y * ClampedFraction));
+					}
+					else
+					{
+						PaintRect = AllottedGeometry.ToPaintGeometry(
+							FVector2D(0.0f, MarqueeAnimOffsetY - MarqueeImageSizeY),
+							FVector2D(AllottedGeometry.GetLocalSize().X, AllottedGeometry.GetLocalSize().Y + MarqueeImageSizeY));
 					}
 
 					// Draw Fill
 					FSlateDrawElement::MakeBox(
 						OutDrawElements,
 						RetLayerId++,
-						AllottedGeometry.ToPaintGeometry(
-							FVector2D(0.0f, MarqueeAnimOffset - MarqueeImageSize),
-							FVector2D(AllottedGeometry.GetLocalSize().X, AllottedGeometry.GetLocalSize().Y + MarqueeImageSize)),
+						PaintRect,
 						CurrentFillImage,
 						DrawEffects,
 						FillColorAndOpacitySRGB
@@ -274,26 +352,61 @@ int32 SProgressBar::OnPaint( const FPaintArgs& Args, const FGeometry& AllottedGe
 			{
 				if (PushTransformedClip(OutDrawElements, AllottedGeometry, BorderPaddingRef, FVector2D(0, 1), FSlateRect(0, ClampedFraction, 1, 0)))
 				{
-					float MarqueeAnimOffset = 0;
-					float MarqueeImageSize = 0;
+					FPaintGeometry PaintRect;
+					if (bScaleWithFillPerc)
+					{
+						PaintRect = AllottedGeometry.ToPaintGeometry(
+							FVector2D(0.0f, AllottedGeometry.GetLocalSize().Y * (1.0f - ClampedFraction)),
+							FVector2D(AllottedGeometry.GetLocalSize().X, AllottedGeometry.GetLocalSize().Y * ClampedFraction));
+					}
+					else
+					{
+						PaintRect = AllottedGeometry.ToPaintGeometry(
+							FVector2D(0.0f, MarqueeAnimOffsetY - MarqueeImageSizeY),
+							FVector2D(AllottedGeometry.GetLocalSize().X, AllottedGeometry.GetLocalSize().Y + MarqueeImageSizeY));
+					}
 
-					if (Style->EnableFillAnimation)
+					// Draw Fill
+					FSlateDrawElement::MakeBox(
+						OutDrawElements,
+						RetLayerId++,
+						PaintRect,
+						CurrentFillImage,
+						DrawEffects,
+						FillColorAndOpacitySRGB
+					);
+
+					OutDrawElements.PopClip();
+				}
+				break;
+			}
+			case EProgressBarFillType::RightToLeft:
+			{
+				if (PushTransformedClip(OutDrawElements, AllottedGeometry, BorderPaddingRef, FVector2D(1, 0), FSlateRect(ClampedFraction, 0, 0, 1)))
+				{
+					FPaintGeometry PaintRect;
+					if (bScaleWithFillPerc)
+					{
+						PaintRect = AllottedGeometry.ToPaintGeometry(
+							FVector2D(AllottedGeometry.GetLocalSize().X * (1.0f - ClampedFraction), 0.0f),
+							FVector2D(AllottedGeometry.GetLocalSize().X * ClampedFraction, AllottedGeometry.GetLocalSize().Y));
+					}
+					else
 					{
-						MarqueeAnimOffset = CurrentFillImage->ImageSize.Y * MarqueeOffset;
-						MarqueeImageSize = CurrentFillImage->ImageSize.Y;
+						PaintRect = AllottedGeometry.ToPaintGeometry(
+							FVector2D(MarqueeAnimOffsetX - MarqueeImageSizeX, 0.0f),
+							FVector2D(AllottedGeometry.GetLocalSize().X + MarqueeImageSizeX, AllottedGeometry.GetLocalSize().Y));
 					}
 
 					// Draw Fill
 					FSlateDrawElement::MakeBox(
 						OutDrawElements,
 						RetLayerId++,
-						AllottedGeometry.ToPaintGeometry(
-							FVector2D(0.0f, MarqueeAnimOffset - MarqueeImageSize),
-							FVector2D( AllottedGeometry.GetLocalSize().X, AllottedGeometry.GetLocalSize().Y + MarqueeImageSize)),
+						PaintRect,
 						CurrentFillImage,
 						DrawEffects,
 						FillColorAndOpacitySRGB
-						);
+					);
 
 					OutDrawElements.PopClip();
 				}
@@ -302,25 +415,27 @@ int32 SProgressBar::OnPaint( const FPaintArgs& Args, const FGeometry& AllottedGe
 			case EProgressBarFillType::LeftToRight:
 			default:
 			{
-				float MarqueeAnimOffset = 0;
-				float MarqueeImageSize = 0;
-
-				if(Style->EnableFillAnimation)
-				{
-					MarqueeAnimOffset = CurrentFillImage->ImageSize.X * MarqueeOffset;
-					MarqueeImageSize = CurrentFillImage->ImageSize.X;
-				}
-
 				if (PushTransformedClip(OutDrawElements, AllottedGeometry, BorderPaddingRef, FVector2D(0, 0), FSlateRect(0, 0, ClampedFraction, 1)))
 				{
+					FPaintGeometry PaintRect;
+					if (bScaleWithFillPerc)
+					{
+						PaintRect = AllottedGeometry.ToPaintGeometry(
+							FVector2D(0.0f, 0.0f),
+							FVector2D(AllottedGeometry.GetLocalSize().X * ClampedFraction, AllottedGeometry.GetLocalSize().Y));
+					}
+					else
+					{
+						PaintRect = AllottedGeometry.ToPaintGeometry(
+							FVector2D(MarqueeAnimOffsetX - MarqueeImageSizeX, 0.0f),
+							FVector2D(AllottedGeometry.GetLocalSize().X + MarqueeImageSizeX, AllottedGeometry.GetLocalSize().Y));
+					}
 
 					// Draw fill
 					FSlateDrawElement::MakeBox(
 						OutDrawElements,
 						RetLayerId++,
-						AllottedGeometry.ToPaintGeometry(
-							FVector2D(MarqueeAnimOffset - MarqueeImageSize, 0.0f),
-							FVector2D(AllottedGeometry.GetLocalSize().X+ MarqueeImageSize, AllottedGeometry.GetLocalSize().Y)),
+						PaintRect,
 						CurrentFillImage,
 						DrawEffects,
 						FillColorAndOpacitySRGB
diff --git a/Engine/Source/Runtime/Slate/Public/Widgets/Notifications/SProgressBar.h b/Engine/Source/Runtime/Slate/Public/Widgets/Notifications/SProgressBar.h
--- a/Engine/Source/Runtime/Slate/Public/Widgets/Notifications/SProgressBar.h
+++ b/Engine/Source/Runtime/Slate/Public/Widgets/Notifications/SProgressBar.h
@@ -29,8 +29,12 @@ namespace EProgressBarFillType
 		LeftToRight,
 		// will fill up from the right side to the left side
 		RightToLeft,
-		// will fill up from the center to the outer edges
+		// will scale up from the midpoint to the outer edges both vertically and horizontally
 		FillFromCenter,
+		// will fill up from the centerline to the outer edges horizontally
+		FillFromCenterHorizontal,
+		// will fill up from the centerline to the outer edges vertically
+		FillFromCenterVertical,
 		// will fill up from the top to the the bottom
 		TopToBottom,
 		// will fill up from the bottom to the the top
@@ -38,6 +42,21 @@ namespace EProgressBarFillType
 	};
 }
 
+/**
+ * SProgressBar Fill Style
+ */
+UENUM(BlueprintType)
+namespace EProgressBarFillStyle
+{
+	enum Type
+	{
+		// a mask is used to paint the fill image
+		Mask,
+		// the fill image is scaled to the fill percentage
+		Scale,
+	};
+}
+
 /** A progress bar widget.*/
 class SLATE_API SProgressBar : public SLeafWidget
 {
@@ -46,6 +65,7 @@ public:
 	SLATE_BEGIN_ARGS(SProgressBar)
 		: _Style( &FAppStyle::Get().GetWidgetStyle<FProgressBarStyle>("ProgressBar") )
 		, _BarFillType(EProgressBarFillType::LeftToRight)
+		, _BarFillStyle(EProgressBarFillStyle::Mask)
 		, _Percent( TOptional<float>() )
 		, _FillColorAndOpacity( FLinearColor::White )
 		, _BorderPadding( FVector2D(0,0) )
@@ -60,9 +80,12 @@ public:
 		/** Style used for the progress bar */
 		SLATE_STYLE_ARGUMENT( FProgressBarStyle, Style )
 
-		/** Defines if this progress bar fills Left to right or right to left*/
+		/** Defines the direction in which the progress bar fills */
 		SLATE_ARGUMENT( EProgressBarFillType::Type, BarFillType )
 
+		/** Defines the visual style of the progress bar fill - scale or mask */
+		SLATE_ARGUMENT( EProgressBarFillStyle::Type, BarFillStyle )
+
 		/** Used to determine the fill position of the progress bar ranging 0..1 */
 		SLATE_ATTRIBUTE( TOptional<float>, Percent )
 
@@ -105,6 +128,9 @@ public:
 	
 	/** See attribute BarFillType */
 	void SetBarFillType(EProgressBarFillType::Type InBarFillType);
+
+	/** See attribute BarFillStyle */
+	void SetBarFillStyle(EProgressBarFillStyle::Type InBarFillStyle);
 	
 	/** See attribute SetFillColorAndOpacity */
 	void SetFillColorAndOpacity(TAttribute< FSlateColor > InFillColorAndOpacity);
@@ -148,6 +174,8 @@ private:
 
 	EProgressBarFillType::Type BarFillType;
 
+	EProgressBarFillStyle::Type BarFillStyle;
+
 	/** Background image to use for the progress bar */
 	const FSlateBrush* BackgroundImage;
 
diff --git a/Engine/Source/Runtime/UMG/Private/Components/ProgressBar.cpp b/Engine/Source/Runtime/UMG/Private/Components/ProgressBar.cpp
--- a/Engine/Source/Runtime/UMG/Private/Components/ProgressBar.cpp
+++ b/Engine/Source/Runtime/UMG/Private/Components/ProgressBar.cpp
@@ -26,6 +26,7 @@ UProgressBar::UProgressBar(const FObjectInitializer& ObjectInitializer)
 	WidgetStyle.FillImage.TintColor = FLinearColor::White;
 
 	BarFillType = EProgressBarFillType::LeftToRight;
+	BarFillStyle = EProgressBarFillStyle::Mask;
 	bIsMarquee = false;
 	Percent = 0;
 	FillColorAndOpacity = FLinearColor::White;
@@ -56,6 +57,7 @@ void UProgressBar::SynchronizeProperties()
 	MyProgressBar->SetStyle(&WidgetStyle);
 
 	MyProgressBar->SetBarFillType(BarFillType);
+	MyProgressBar->SetBarFillStyle(BarFillStyle);
 	MyProgressBar->SetPercent(bIsMarquee ? TOptional<float>() : PercentBinding);
 	MyProgressBar->SetFillColorAndOpacity(FillColorAndOpacityBinding);
 	MyProgressBar->SetBorderPadding(BorderPadding);
diff --git a/Engine/Source/Runtime/UMG/Public/Components/ProgressBar.h b/Engine/Source/Runtime/UMG/Public/Components/ProgressBar.h
--- a/Engine/Source/Runtime/UMG/Public/Components/ProgressBar.h
+++ b/Engine/Source/Runtime/UMG/Public/Components/ProgressBar.h
@@ -48,9 +48,13 @@ public:
 	UPROPERTY(EditAnywhere, BlueprintReadOnly, Category=Progress, meta=( UIMin = "0", UIMax = "1" ))
 	float Percent;
 
-	/** Defines if this progress bar fills Left to right or right to left */
+	/** Defines the direction in which the progress bar fills */
 	UPROPERTY(EditAnywhere, BlueprintReadOnly, Category=Progress)
 	TEnumAsByte<EProgressBarFillType::Type> BarFillType;
+
+	/** Defines the visual style of the progress bar fill - scale or mask */
+	UPROPERTY(EditAnywhere, BlueprintReadOnly, Category=Progress)
+	TEnumAsByte<EProgressBarFillStyle::Type> BarFillStyle;
 	
 	UPROPERTY(EditAnywhere, BlueprintReadOnly, Category=Progress)
 	bool bIsMarquee;
