Fix behavior of MinSize on SSplitter Slots when the first slot has a min size, and a Value not large enough to cover its min size

diff --git a/Engine/Source/Runtime/Slate/Private/Widgets/Layout/SSplitter.cpp b/Engine/Source/Runtime/Slate/Private/Widgets/Layout/SSplitter.cpp
--- a/Engine/Source/Runtime/Slate/Private/Widgets/Layout/SSplitter.cpp
+++ b/Engine/Source/Runtime/Slate/Private/Widgets/Layout/SSplitter.cpp
@@ -139,6 +139,8 @@ TArray<FLayoutGeometry> SSplitter::ArrangeChildrenForLayout(const FGeometry& All
 	TArray<float, TMemStackAllocator<>> SlotSizes;
 	SlotSizes.Empty(Children.Num());
 
+	float ExtraRequiredSpace = 0;
+
 	// calculate slot sizes
 	for (int32 ChildIndex = 0; ChildIndex < Children.Num(); ++ChildIndex)
 	{
@@ -151,9 +153,18 @@ TArray<FLayoutGeometry> SSplitter::ArrangeChildrenForLayout(const FGeometry& All
 			continue;
 		}
 
-		float ChildSpace = (CurSlot.SizingRule.Get() == SSplitter::SizeToContent)
-			? CurSlot.GetWidget()->GetDesiredSize()[AxisIndex]
-			: ResizableSpace * CurSlot.SizeValue.Get() / CoefficientTotal;
+		float ChildSpace;
+		
+		if (CurSlot.SizingRule.Get() == SSplitter::SizeToContent)
+		{
+			ChildSpace = CurSlot.GetWidget()->GetDesiredSize()[AxisIndex];
+		}
+		else
+		{
+			ChildSpace = ResizableSpace * CurSlot.SizeValue.Get() / CoefficientTotal  - ExtraRequiredSpace;
+			ExtraRequiredSpace = 0;
+		}
+
 
 		// we have enough space to clamp this and steal space from earlier slots
 		if (ResizableSpace >= MinResizableSpace)
@@ -175,6 +186,12 @@ TArray<FLayoutGeometry> SSplitter::ArrangeChildrenForLayout(const FGeometry& All
 					}
 				}
 			}
+			
+			if (CurrentRequiredSpace > 0)
+			{
+				// we weren't able to steal enough space from previous children, store this value to try and take it from following children
+				ExtraRequiredSpace = CurrentRequiredSpace;
+			}
 
 			ChildSpace = ClampedSpace;
 		}
