Window > Load Layout now works with JSON layout files as well.

FLayoutSaveRestore
 - DuplicateConfig added to handle JSON files when loading an existing layout as well.
 - MigrateConfig updated to handle JSON files.
 - GetSectionFromConfig and SetSectionFromConfig now also handle JSON files.

diff --git a/Engine/Source/Editor/MainFrame/Private/Menus/LayoutsMenu.cpp b/Engine/Source/Editor/MainFrame/Private/Menus/LayoutsMenu.cpp
--- a/Engine/Source/Editor/MainFrame/Private/Menus/LayoutsMenu.cpp
+++ b/Engine/Source/Editor/MainFrame/Private/Menus/LayoutsMenu.cpp
@@ -147,17 +147,23 @@ TArray<FString> FPrivateLayoutsMenu::GetIniFilesInFolderInternal(const FString&
 bool FPrivateLayoutsMenu::TrySaveLayoutOrWarnInternal(const FString& InSourceFilePath, const FString& InTargetFilePath, const FText& InWhatIsThis,
 	const bool bCleanLayoutNameAndDescriptionFieldsIfNoSameValues, const bool bShouldAskBeforeCleaningLayoutNameAndDescriptionFields, const bool bShowSaveToast)
 {
-	// If desired, ask user whether to keep the LayoutName and LayoutDescription fields
+	// We must re-read configs to avoid the Editor using a previously cached version
+	GConfig->UnloadFile(InSourceFilePath); 
+	GConfig->UnloadFile(InTargetFilePath);
+
 	bool bCleanLayoutNameAndDescriptionFields = false;
+
 	// If we are checking whether to clean the fields, we only want to maintain them if we are saving the file into an existing file that already has the same field values
 	if (bCleanLayoutNameAndDescriptionFieldsIfNoSameValues)
 	{
-		GConfig->UnloadFile(InSourceFilePath); // We must re-read it to avoid the Editor to use a previously cached name and description
+		// read the source name and description
 		const FText LayoutNameSource = FLayoutSaveRestore::LoadSectionFromConfig(InSourceFilePath, "LayoutName");
 		const FText LayoutDescriptionSource = FLayoutSaveRestore::LoadSectionFromConfig(InSourceFilePath, "LayoutDescription");
-		GConfig->UnloadFile(InTargetFilePath); // We must re-read it to avoid the Editor to use a previously cached name and description
+
+		// read the target name and description
 		const FText LayoutNameTarget = FLayoutSaveRestore::LoadSectionFromConfig(InTargetFilePath, "LayoutName");
 		const FText LayoutDescriptionTarget = FLayoutSaveRestore::LoadSectionFromConfig(InTargetFilePath, "LayoutDescription");
+
 		// The output target exists (overriding)
 		// These fields are not empty in source
 		if (!LayoutNameSource.IsEmpty() || !LayoutDescriptionSource.IsEmpty())
@@ -174,6 +180,7 @@ bool FPrivateLayoutsMenu::TrySaveLayoutOrWarnInternal(const FString& InSourceFil
 					const FText TextBody = FText::Format(
 						LOCTEXT("OverrideLayoutNameAndDescriptionFieldBody", "You are saving a layout that contains a custom layout name and/or description. Do you also want to copy these 2 properties?\n - Current layout name: {0}\n - Current layout description: {1}\n\nIf you select \"Preserve Values\", the displayed name and description of the original layout customization will also be copied into the new configuration file.\n\nIf you select \"Clear Values\", these fields will be emptied.\n\nIf you are not sure, select \"Preserve Values\" if you are exporting the layout configuration without making any changes, or \"Clear Values\" if you"" have made or plan to make changes to the layout.\n\n"),
 						LayoutNameSource, LayoutDescriptionSource);
+
 					// Dialog SWidget
 					TSharedRef<SVerticalBox> DialogContents = SNew(SVerticalBox);
 					DialogContents->AddSlot()
@@ -182,14 +189,15 @@ bool FPrivateLayoutsMenu::TrySaveLayoutOrWarnInternal(const FString& InSourceFil
 							SNew(STextBlock)
 							.Text(TextBody)
 						];
+
 					const FText PreserveValuesText = LOCTEXT("PreserveValuesText", "Preserve Values");
 					const FText ClearValuesText = LOCTEXT("ClearValuesText", "Clear Values");
 					const FText CancelText = NSLOCTEXT("Dialogs", "EAppReturnTypeCancel", "Cancel");
 					TSharedRef<SCustomDialog> CustomDialog = SNew(SCustomDialog)
 						.Title(TextTitle)
 						.DialogContent(DialogContents)
-						.Buttons({ SCustomDialog::FButton(PreserveValuesText), SCustomDialog::FButton(ClearValuesText), SCustomDialog::FButton(CancelText) })
-					;
+						.Buttons({ SCustomDialog::FButton(PreserveValuesText), SCustomDialog::FButton(ClearValuesText), SCustomDialog::FButton(CancelText) });
+
 					// Returns 0 when "Preserve Values" is pressed, 1 when "Clear Values" is pressed, or 2 when Cancel/Esc is pressed
 					const int32 ButtonPressed = CustomDialog->ShowModal();
 					// Preserve Values
@@ -217,51 +225,54 @@ bool FPrivateLayoutsMenu::TrySaveLayoutOrWarnInternal(const FString& InSourceFil
 			}
 		}
 	}
-	// Copy: Replace main layout with desired one
+
 	const FString TargetAbsoluteFilePath = FPaths::ConvertRelativePathToFull(InTargetFilePath);
-	const bool bShouldReplace = true;
-	const bool bCopyEvenIfReadOnly = true;
-	const bool bCopyAttributes = false; // If true, it could e.g., copy the read-only flag of DefaultLayout.ini and make all the save/load stuff stop working
-	if (COPY_Fail == IFileManager::Get().Copy(*InTargetFilePath, *InSourceFilePath, bShouldReplace, bCopyEvenIfReadOnly, bCopyAttributes))
+	if (!FLayoutSaveRestore::DuplicateConfig(InSourceFilePath, InTargetFilePath))
 	{
+		const FText SourceAbsoluteFilePathText = FText::FromString(FPaths::ConvertRelativePathToFull(InSourceFilePath));		
+		const FText TargetAbsoluteFilePathText = FText::FromString(TargetAbsoluteFilePath);
+
 		FMessageLog EditorErrors("EditorErrors");
 		FText TextBody;
 		FFormatNamedArguments Arguments;
 		Arguments.Add(TEXT("WhatIs"), InWhatIsThis);
+
 		// Source does not exist
 		if (!FPaths::FileExists(InSourceFilePath))
 		{
-			Arguments.Add(TEXT("FileName"), FText::FromString(FPaths::ConvertRelativePathToFull(InSourceFilePath)));
+			Arguments.Add(TEXT("FileName"), SourceAbsoluteFilePathText);
 			TextBody = FText::Format(LOCTEXT("UnsuccessfulSave_NoExist_Notification", "The requested operation ({WhatIs}) was unsuccessful, the desired file does not exist. File path:\n{FileName}"), Arguments);
 			EditorErrors.Warning(TextBody);
 		}
 		// Target is read-only
 		else if (IFileManager::Get().IsReadOnly(*InTargetFilePath))
 		{
-			Arguments.Add(TEXT("FileName"), FText::FromString(TargetAbsoluteFilePath));
+			Arguments.Add(TEXT("FileName"), TargetAbsoluteFilePathText);
 			TextBody = FText::Format(LOCTEXT("UnsuccessfulSave_ReadOnly_Notification", "The requested operation ({WhatIs}) was unsuccessful, the target file path is read-only. File path:\n{FileName}"), Arguments);
 			EditorErrors.Warning(TextBody);
 		}
 		// Target and source are the same
-		else if (TargetAbsoluteFilePath == FPaths::ConvertRelativePathToFull(InSourceFilePath))
+		else if (TargetAbsoluteFilePathText.EqualTo(SourceAbsoluteFilePathText))
 		{
-			Arguments.Add(TEXT("SourceFileName"), FText::FromString(FPaths::ConvertRelativePathToFull(InSourceFilePath)));
-			Arguments.Add(TEXT("FinalFileName"), FText::FromString(TargetAbsoluteFilePath));
+			Arguments.Add(TEXT("SourceFileName"), SourceAbsoluteFilePathText);
+			Arguments.Add(TEXT("FinalFileName"), TargetAbsoluteFilePathText);
 			TextBody = FText::Format(LOCTEXT("UnsuccessfulSave_Identical_Notification", "The requested operation ({WhatIs}) was unsuccessful, target and source layout file paths are the same ({SourceFileName})!\nAre you trying to import or replace a file that is already in the layouts folder? If so, remove the current file first."), Arguments);
 			EditorErrors.Warning(TextBody);
 		}
 		// We don't specifically know why it failed, this is a fallback
 		else
 		{
-			Arguments.Add(TEXT("SourceFileName"), FText::FromString(FPaths::ConvertRelativePathToFull(InSourceFilePath)));
-			Arguments.Add(TEXT("FinalFileName"), FText::FromString(TargetAbsoluteFilePath));
+			Arguments.Add(TEXT("SourceFileName"), SourceAbsoluteFilePathText);
+			Arguments.Add(TEXT("FinalFileName"), TargetAbsoluteFilePathText);
 			TextBody = FText::Format(LOCTEXT("UnsuccessfulSave_Fallback_Notification", "The requested operation ({WhatIs}) was unsuccessful while copying the layout file from\n{SourceFileName}\ninto\n{FinalFileName}\n\nUsually, this occurs when the introduced file name contains unsupported characters or the total path length exceeds the OS limit."), Arguments);
 			EditorErrors.Warning(TextBody);
 		}
 		EditorErrors.Notify(LOCTEXT("LoadUnsuccessful_Title", "Load Unsuccessful!"));
+
 		// Show reason
 		const FText TextTitle = LOCTEXT("UnsuccessfulCopyHeader", "Unsuccessful copy!");
 		FMessageDialog::Open(EAppMsgType::Ok, TextBody, &TextTitle);
+
 		// Return
 		return false;
 	}
@@ -277,17 +288,21 @@ bool FPrivateLayoutsMenu::TrySaveLayoutOrWarnInternal(const FString& InSourceFil
 			// Update fields
 			FLayoutSaveRestore::SaveSectionToConfig(GEditorLayoutIni, "LayoutName", FText::FromString(""));
 			FLayoutSaveRestore::SaveSectionToConfig(GEditorLayoutIni, "LayoutDescription", FText::FromString(""));
+
 			// Flush file
 			const bool bRead = true;
 			GConfig->Flush(bRead, GEditorLayoutIni);
+
 			// Re-copy file
 			if (TargetAbsoluteFilePath != FPaths::ConvertRelativePathToFull(GEditorLayoutIni))
 			{
-				IFileManager::Get().Copy(*InTargetFilePath, *GEditorLayoutIni, bShouldReplace, bCopyEvenIfReadOnly, bCopyAttributes);
+				const bool bShouldReplace = true;
+				const bool bCopyEvenIfReadOnly = true;
+				const bool bCopyAttributes = false;
+				IFileManager::Get().Copy(*TargetAbsoluteFilePath, *GEditorLayoutIni, bShouldReplace, bCopyEvenIfReadOnly, bCopyAttributes);
 			}
 		}
-		// Unload target file so it can be re-read into cache properly the next time it is used
-		GConfig->UnloadFile(InTargetFilePath); // We must re-read it to avoid the Editor to use a previously cached name and description
+
 		// Display Editor toast to inform the user of the result of the operation
 		if (bShowSaveToast)
 		{
@@ -296,22 +311,23 @@ bool FPrivateLayoutsMenu::TrySaveLayoutOrWarnInternal(const FString& InSourceFil
 			Info.ExpireDuration = 5.0f;
 			Info.bUseSuccessFailIcons = false;
 			Info.bUseLargeFont = false;
+
 			TSharedPtr<SNotificationItem> SaveMessagePtr = FSlateNotificationManager::Get().AddNotification(Info);
 			if (SaveMessagePtr.IsValid())
 			{
-				const FString& HyperLinkString = TargetAbsoluteFilePath;
+				const FString HyperLinkString = TargetAbsoluteFilePath;
 				auto OpenScreenshotFolder = [HyperLinkString]
 				{
-					FPlatformProcess::ExploreFolder(*FPaths::GetPath(HyperLinkString));
+					FPlatformProcess::ExploreFolder(*HyperLinkString);
 				};
 				SaveMessagePtr->SetText(LOCTEXT("SuccessfulSave_Toast", "Editor layout file saved as"));
 				SaveMessagePtr->SetHyperlink(FSimpleDelegate::CreateLambda(OpenScreenshotFolder), FText::FromString(HyperLinkString));
 				SaveMessagePtr->SetCompletionState(SNotificationItem::CS_Success);
 			}
 		}
-		// Return successful copy message
-		return true;
 	}
+
+	return true;
 }
 
 FText FPrivateLayoutsMenu::GetDisplayTextInternal(const FString& InString)
diff --git a/Engine/Source/Runtime/Slate/Private/Framework/Docking/LayoutService.cpp b/Engine/Source/Runtime/Slate/Private/Framework/Docking/LayoutService.cpp
--- a/Engine/Source/Runtime/Slate/Private/Framework/Docking/LayoutService.cpp
+++ b/Engine/Source/Runtime/Slate/Private/Framework/Docking/LayoutService.cpp
@@ -1,6 +1,7 @@
 // Copyright Epic Games, Inc. All Rights Reserved.
 
 #include "Framework/Docking/LayoutService.h"
+#include "HAL/FileManager.h"
 #include "Misc/App.h"
 #include "Misc/ConfigCacheIni.h"
 #include "Misc/FileHelper.h"
@@ -12,6 +13,26 @@ DEFINE_LOG_CATEGORY_STATIC(LogLayoutService, Log, All);
 
 const TCHAR* EditorLayoutsSectionName = TEXT("EditorLayouts");
 
+static FString PrepareLayoutStringForIni(const FString& LayoutString)
+{
+	// Have to store braces as parentheses due to braces causing ini issues
+	return LayoutString
+		.Replace(TEXT("{"), TEXT("("))
+		.Replace(TEXT("}"), TEXT(")"))
+		.Replace(TEXT("\r"), TEXT(""))
+		.Replace(TEXT("\n"), TEXT(""))
+		.Replace(TEXT("\t"), TEXT(""));
+}
+
+static FString GetLayoutStringFromIni(const FString& LayoutString)
+{
+	// Revert parenthesis to braces, from ini readable to Json readable
+	return LayoutString
+		.Replace(TEXT("("), TEXT("{"))
+		.Replace(TEXT(")"), TEXT("}"))
+		.Replace(TEXT("\\") LINE_TERMINATOR, LINE_TERMINATOR);
+}
+
 
 const FString& FLayoutSaveRestore::GetAdditionalLayoutConfigIni()
 {
@@ -19,58 +40,97 @@ const FString& FLayoutSaveRestore::GetAdditionalLayoutConfigIni()
 	return IniSectionAdditionalConfig;
 }
 
-static void SaveLayoutToJson(const FString& InConfigFileName, const TSharedRef<FTabManager::FLayout>& InLayoutToSave)
+static TSharedPtr<FJsonObject> ConvertSectionToJson(const TArray<FString>& SectionStrings)
+{
+	TSharedPtr<FJsonObject> RootObject = MakeShared<FJsonObject>();
+
+	for (const FString& SectionPair : SectionStrings)
+	{
+		FString Key, Value;
+		if (SectionPair.Split(TEXT("="), &Key, &Value))
+		{
+			Value = GetLayoutStringFromIni(Value);
+
+			TSharedPtr<FJsonObject> ChildObject = MakeShared<FJsonObject>();
+			TSharedRef<TJsonReader<>> Reader = TJsonReaderFactory<>::Create(Value);
+			if (FJsonSerializer::Deserialize(Reader, ChildObject))
+			{
+				RootObject->SetObjectField(Key, ChildObject);
+			}
+			else
+			{
+				RootObject->SetStringField(Key, Value);
+			}
+		}
+	}
+
+	return RootObject;
+}
+
+static FString GetLayoutJsonFileName(const FString& InConfigFileName)
 {
 	const FString JsonFileName = FPaths::GetBaseFilename(InConfigFileName) + TEXT(".json");
 	const FString UserSettingsPath = FPaths::Combine(FPlatformProcess::UserSettingsDir(), FApp::GetEpicProductIdentifier(), TEXT("Editor"), JsonFileName);
+	return UserSettingsPath;
+}
 
-	TSharedPtr<FJsonObject> AllLayoutsObject;
+static TSharedPtr<FJsonObject> LoadJsonFile(const FString& InFileName)
+{
+	TSharedPtr<FJsonObject> JsonObject;
 
-	FString ExistingJsonContents;
-	if (FFileHelper::LoadFileToString(ExistingJsonContents, *UserSettingsPath))
+	FString JsonContents;
+	if (FFileHelper::LoadFileToString(JsonContents, *InFileName))
 	{
-		TSharedRef<TJsonReader<>> Reader = TJsonReaderFactory<>::Create(ExistingJsonContents);
-		FJsonSerializer::Deserialize(Reader, AllLayoutsObject);
+		TSharedRef<TJsonReader<>> Reader = TJsonReaderFactory<>::Create(JsonContents);
+		FJsonSerializer::Deserialize(Reader, JsonObject);
 	}
-	
+
+	return JsonObject;
+}
+
+static bool SaveJsonFile(const FString& InFileName, TSharedPtr<FJsonObject> JsonObject)
+{
+	FString NewJsonContents;
+	TSharedRef<TJsonWriter<>> Writer = TJsonWriterFactory<>::Create(&NewJsonContents);
+	if (FJsonSerializer::Serialize(JsonObject.ToSharedRef(), Writer))
+	{
+		return FFileHelper::SaveStringToFile(NewJsonContents, *InFileName);
+	}
+
+	return false;
+}
+
+static void SaveLayoutToJson(const FString& InConfigFileName, const TSharedRef<FTabManager::FLayout>& InLayoutToSave)
+{
+	const FString UserSettingsPath = GetLayoutJsonFileName(InConfigFileName);
+
+	TSharedPtr<FJsonObject> AllLayoutsObject = LoadJsonFile(UserSettingsPath);
 	if (!AllLayoutsObject.IsValid())
 	{
 		// doesn't exist
 		AllLayoutsObject = MakeShared<FJsonObject>();
 	}
 
-	TSharedRef<FJsonObject> LayoutJson = InLayoutToSave->ToJson();
-	AllLayoutsObject->SetObjectField(InLayoutToSave->GetLayoutName().ToString(), LayoutJson);
+	AllLayoutsObject->SetObjectField(InLayoutToSave->GetLayoutName().ToString(), InLayoutToSave->ToJson());
 
-	FString NewJsonContents;
-	NewJsonContents.Reserve(ExistingJsonContents.Len());
-	TSharedRef<TJsonWriter<>> Writer = TJsonWriterFactory<>::Create(&NewJsonContents);
-	if (FJsonSerializer::Serialize(AllLayoutsObject.ToSharedRef(), Writer))
-	{
-		FFileHelper::SaveStringToFile(NewJsonContents, *UserSettingsPath);
-	}
+	SaveJsonFile(UserSettingsPath, AllLayoutsObject);
 }
 
 static bool LoadLayoutFromJson(const FString& InConfigFileName, const FString& InLayoutName, TSharedPtr<FTabManager::FLayout>& OutLayout)
 {
-	const FString JsonFileName = FPaths::GetBaseFilename(InConfigFileName) + TEXT(".json");
-	const FString UserSettingsPath = FPaths::Combine(FPlatformProcess::UserSettingsDir(), FApp::GetEpicProductIdentifier(), TEXT("Editor"), JsonFileName);
+	const FString UserSettingsPath = GetLayoutJsonFileName(InConfigFileName);
 
-	FString ExistingJsonContents;
-	if (FFileHelper::LoadFileToString(ExistingJsonContents, *UserSettingsPath))
+	TSharedPtr<FJsonObject> JsonObject = LoadJsonFile(InConfigFileName);
+	if (!JsonObject.IsValid())
 	{
-		TSharedPtr<FJsonObject> AllLayoutsObject;
+		return false;
+	}
 
-		TSharedRef<TJsonReader<>> Reader = TJsonReaderFactory<>::Create(ExistingJsonContents);
-		if (FJsonSerializer::Deserialize(Reader, AllLayoutsObject))
-		{
-			const TSharedPtr<FJsonObject>* LayoutJson = nullptr;
-			if (AllLayoutsObject->TryGetObjectField(InLayoutName, LayoutJson))
-			{
-				OutLayout = FTabManager::FLayout::NewFromJson(*LayoutJson);
-				return true;
-			}
-		}
+	const TSharedPtr<FJsonObject>* LayoutJson = nullptr;
+	if (JsonObject->TryGetObjectField(InLayoutName, LayoutJson))
+	{
+		OutLayout = FTabManager::FLayout::NewFromJson(*LayoutJson);
+		return true;
 	}
 
 	return false;
@@ -81,7 +141,7 @@ void FLayoutSaveRestore::SaveToConfig( const FString& InConfigFileName, const TS
 	// Only save to config if it's not the FTabManager::FLayout::NullLayout
 	if (InLayoutToSave->GetLayoutName() != FTabManager::FLayout::NullLayout->GetLayoutName())
 	{
-		const FString LayoutAsString = FLayoutSaveRestore::PrepareLayoutStringForIni(InLayoutToSave->ToString());
+		const FString LayoutAsString = PrepareLayoutStringForIni(InLayoutToSave->ToString());
 		GConfig->SetString(EditorLayoutsSectionName, *InLayoutToSave->GetLayoutName().ToString(), *LayoutAsString, InConfigFileName);
 
 		SaveLayoutToJson(InConfigFileName, InLayoutToSave);
@@ -116,7 +176,7 @@ TSharedRef<FTabManager::FLayout> FLayoutSaveRestore::LoadFromConfigPrivate(const
 		FString IniLayoutString;
 		// If the Key (InDefaultLayout->GetLayoutName()) already exists in the section EditorLayoutsSectionName of the file InConfigFileName, try to load the layout from that file
 		GConfig->GetString(EditorLayoutsSectionName, *LayoutNameString, IniLayoutString, InConfigFileName);
-		UserLayout = FTabManager::FLayout::NewFromString( FLayoutSaveRestore::GetLayoutStringFromIni( IniLayoutString ) );
+		UserLayout = FTabManager::FLayout::NewFromString( GetLayoutStringFromIni( IniLayoutString ) );
 	}
 
 	if (UserLayout.IsValid())
@@ -174,19 +234,65 @@ TSharedRef<FTabManager::FLayout> FLayoutSaveRestore::LoadFromConfigPrivate(const
 	return InDefaultLayout;
 }
 
-
 void FLayoutSaveRestore::SaveSectionToConfig(const FString& InConfigFileName, const FString& InSectionName, const FText& InSectionValue)
 {
-	GConfig->SetText(EditorLayoutsSectionName, *InSectionName, InSectionValue, InConfigFileName);
+	FString StrValue;
+	FTextStringHelper::WriteToBuffer(StrValue, InSectionValue);
+
+	GConfig->SetString(EditorLayoutsSectionName, *InSectionName, *StrValue, InConfigFileName);
+
+	const FString JsonFileName = GetLayoutJsonFileName(InConfigFileName);
+	TSharedPtr<FJsonObject> JsonObject = LoadJsonFile(JsonFileName);
+	if (JsonObject.IsValid())
+	{
+		JsonObject->SetStringField(InSectionName, *StrValue);
+		SaveJsonFile(JsonFileName, JsonObject);
+	}
 }
 
 FText FLayoutSaveRestore::LoadSectionFromConfig(const FString& InConfigFileName, const FString& InSectionName)
 {
-	FText LayoutString;
-	GConfig->GetText(EditorLayoutsSectionName, *InSectionName, LayoutString, InConfigFileName);
-	return LayoutString;
+	FString ValueString;
+	const FString JsonFileName = GetLayoutJsonFileName(InConfigFileName);
+	TSharedPtr<FJsonObject> JsonObject = LoadJsonFile(JsonFileName);
+	if (JsonObject.IsValid())
+	{
+		ValueString = JsonObject->GetStringField(InSectionName);
+	}
+
+	if (ValueString.IsEmpty())
+	{
+		GConfig->GetString(EditorLayoutsSectionName, *InSectionName, ValueString, InConfigFileName);
+	}
+
+	FText ValueText;
+	FTextStringHelper::ReadFromBuffer(*ValueString, ValueText, EditorLayoutsSectionName);
+	
+	return ValueText;
 }
 
+bool FLayoutSaveRestore::DuplicateConfig(const FString& SourceConfigFileName, const FString& TargetConfigFileName)
+{
+	const bool bShouldReplace = true;
+	const bool bCopyEvenIfReadOnly = true;
+	const bool bCopyAttributes = false; // If true, we could copy the read-only flag of DefaultLayout.ini and cause save/load to stop working
+
+	if (IFileManager::Get().Copy(*TargetConfigFileName, *SourceConfigFileName, bShouldReplace, bCopyEvenIfReadOnly, bCopyAttributes) == COPY_Fail)
+	{
+		return false;
+	}
+
+	// convert this layout to a JSON file
+	TArray<FString> SectionPairs;
+	GConfig->GetSection(EditorLayoutsSectionName, SectionPairs, TargetConfigFileName);
+
+	TSharedPtr<FJsonObject> RootObject = ConvertSectionToJson(SectionPairs);
+
+	const FString TargetJsonFilename = GetLayoutJsonFileName(TargetConfigFileName);
+	SaveJsonFile(TargetJsonFilename, RootObject);
+
+	return true;
+}
 
 void FLayoutSaveRestore::MigrateConfig( const FString& OldConfigFileName, const FString& NewConfigFileName )
 {
@@ -205,7 +311,7 @@ void FLayoutSaveRestore::MigrateConfig( const FString& OldConfigFileName, const
 	{
 		FString Key, Value;
 
-		for (auto SectionString : OldSectionStrings)
+		for (const FString& SectionString : OldSectionStrings)
 		{
 			if (SectionString.Split(TEXT("="), &Key, &Value))
 			{
@@ -218,24 +324,25 @@ void FLayoutSaveRestore::MigrateConfig( const FString& OldConfigFileName, const
 	GConfig->EmptySection(EditorLayoutsSectionName, OldConfigFileName);
 	GConfig->Flush(false, OldConfigFileName);
 	GConfig->Flush(false, NewConfigFileName);
-}
 
-
-bool FLayoutSaveRestore::IsValidConfig(const FString& InConfigFileName)
-{
-	return GConfig->DoesSectionExist(EditorLayoutsSectionName, *InConfigFileName);
+	// migrate layout to JSON as well
+	const FString NewLayoutJsonFileName = GetLayoutJsonFileName(NewConfigFileName);
+	TSharedPtr<FJsonObject> JsonObject = LoadJsonFile(NewLayoutJsonFileName);
+	if (!JsonObject.IsValid() || JsonObject->Values.Num() == 0)
+	{
+		TSharedPtr<FJsonObject> RootObject = ConvertSectionToJson(OldSectionStrings);
+		SaveJsonFile(NewLayoutJsonFileName, RootObject);
+	}
 }
 
 
-FString FLayoutSaveRestore::PrepareLayoutStringForIni(const FString& LayoutString)
+bool FLayoutSaveRestore::IsValidConfig(const FString& InConfigFileName)
 {
-	// Have to store braces as parentheses due to braces causing ini issues
-	return LayoutString.Replace(TEXT("{"), TEXT("(")).Replace(TEXT("}"), TEXT(")")).Replace(LINE_TERMINATOR, TEXT("\\") LINE_TERMINATOR);
-}
-
+	if (GConfig->DoesSectionExist(EditorLayoutsSectionName, *InConfigFileName))
+	{
+		return true;
+	}
 
-FString FLayoutSaveRestore::GetLayoutStringFromIni(const FString& LayoutString)
-{
-	// Revert parenthesis to braces, from ini readable to Json readable
-	return LayoutString.Replace(TEXT("("), TEXT("{")).Replace(TEXT(")"), TEXT("}")).Replace(TEXT("\\") LINE_TERMINATOR, LINE_TERMINATOR);
+	const FString JsonFileName = GetLayoutJsonFileName(InConfigFileName);
+	return IFileManager::Get().FileExists(*JsonFileName);
 }
diff --git a/Engine/Source/Runtime/Slate/Public/Framework/Docking/LayoutService.h b/Engine/Source/Runtime/Slate/Public/Framework/Docking/LayoutService.h
--- a/Engine/Source/Runtime/Slate/Public/Framework/Docking/LayoutService.h
+++ b/Engine/Source/Runtime/Slate/Public/Framework/Docking/LayoutService.h
@@ -66,7 +66,14 @@ struct SLATE_API FLayoutSaveRestore
 	 * @param OldConfigFileName The name of the old configuration file.
 	 * @param NewConfigFileName The name of the new configuration file.
 	 */
-	static void MigrateConfig( const FString& OldConfigFileName, const FString& NewConfigFileName );
+	static void MigrateConfig(const FString& OldConfigFileName, const FString& NewConfigFileName);
+
+	/** 
+	 * Duplicate the layout config from one file to another.
+	 * @param SourceConfigFileName The name of the source configuration file.
+	 * @param TargetConfigFileName The name of the target configuration file.
+	 */
+	static bool DuplicateConfig(const FString& SourceConfigFileName, const FString& TargetConfigFileName);
 
 	/**
 	 * It checks whether a file is a valid layout config file.
@@ -76,17 +83,6 @@ struct SLATE_API FLayoutSaveRestore
 	static bool IsValidConfig(const FString& InConfigFileName);
 
 private:
-	/**
-	 * Make a Json string friendly for writing out to UE .ini config files.
-	 * The opposite of GetLayoutStringFromIni.
-	 */
-	static FString PrepareLayoutStringForIni(const FString& LayoutString);
-
-	/**
-	 * Convert from UE .ini Json string to a vanilla Json string.
-	 * The opposite of PrepareLayoutStringForIni.
-	 */
-	static FString GetLayoutStringFromIni(const FString& LayoutString);
 
 	/**
 	 * Auxiliary function for both public versions of LoadFromConfig
