Removed temporary logging used to help pinpoint occurrences of "unrecognized tab"s showing up in Blueprint Editor

#rnx

diff --git a/Engine/Source/Editor/UnrealEd/Private/ApplicationMode.cpp b/Engine/Source/Editor/UnrealEd/Private/ApplicationMode.cpp
--- a/Engine/Source/Editor/UnrealEd/Private/ApplicationMode.cpp
+++ b/Engine/Source/Editor/UnrealEd/Private/ApplicationMode.cpp
@@ -28,8 +28,6 @@ FApplicationMode::FApplicationMode(FName InModeName, FText(*GetLocalizedMode)(co
 
 void FApplicationMode::DeactivateMode(TSharedPtr<FTabManager> InTabManager)
 {
-	UE_LOG(LogApplicationMode, Log, TEXT("Deactivating ApplicationMode '%s': 0x%p with TabManager: 0x%p"), *ModeName.ToString(), this, InTabManager.Get());
-
 	// Save the layout to INI
 	check(InTabManager.IsValid());
 	if (FUnrealEdMisc::Get().IsSavingLayoutOnClosedAllowed())
@@ -49,8 +47,6 @@ void FApplicationMode::DeactivateMode(TSharedPtr<FTabManager> InTabManager)
 
 TSharedRef<FTabManager::FLayout> FApplicationMode::ActivateMode(TSharedPtr<FTabManager> InTabManager)
 {
-	UE_LOG(LogApplicationMode, Log, TEXT("Activating ApplicationMode '%s': 0x%p with TabManager: 0x%p"), *ModeName.ToString(), this, InTabManager.Get());
-
 	check(InTabManager.IsValid());
 	RegisterTabFactories(InTabManager);
 
diff --git a/Engine/Source/Runtime/Slate/Private/Framework/Docking/TabManager.cpp b/Engine/Source/Runtime/Slate/Private/Framework/Docking/TabManager.cpp
--- a/Engine/Source/Runtime/Slate/Private/Framework/Docking/TabManager.cpp
+++ b/Engine/Source/Runtime/Slate/Private/Framework/Docking/TabManager.cpp
@@ -876,8 +876,6 @@ void FTabManager::CloseAllAreas()
 
 TSharedRef<FTabManager::FLayout> FTabManager::PersistLayout() const
 {
-	UE_LOG(LogSlate, Log, TEXT("Saving tab layout to %s in TabManager 0x%p"), *ActiveLayoutName.ToString(), this);
-
 	TSharedRef<FLayout> PersistentLayout = FTabManager::NewLayout( this->ActiveLayoutName );
 	
 	// Persist layout for all LiveAreas
@@ -932,8 +930,6 @@ FTabSpawnerEntry& FTabManager::RegisterTabSpawner(const FName TabId, const FOnSp
 	ensure(!TabSpawner.Contains(TabId));
 	ensure(!FGlobalTabmanager::Get()->IsLegacyTabType(TabId));
 
-	UE_LOG(LogSlate, Verbose, TEXT("Registering tab spawner for '%s' in layout '%s' on TabManager 0x%p"), *TabId.ToString(), *ActiveLayoutName.ToString(), this);
-
 	TSharedRef<FTabSpawnerEntry> NewSpawnerEntry = MakeShareable(new FTabSpawnerEntry(TabId, OnSpawnTab, CanSpawnTab));
 	TabSpawner.Add(TabId, NewSpawnerEntry);
 	return NewSpawnerEntry.Get();
@@ -941,19 +937,11 @@ FTabSpawnerEntry& FTabManager::RegisterTabSpawner(const FName TabId, const FOnSp
 
 bool FTabManager::UnregisterTabSpawner( const FName TabId )
 {
-	UE_LOG(LogSlate, Verbose, TEXT("Unregistering tab spawner for '%s' in layout '%s' on TabManager 0x%p"), *TabId.ToString(), *ActiveLayoutName.ToString(), this);
-
 	return TabSpawner.Remove( TabId ) > 0;
 }
 
 void FTabManager::UnregisterAllTabSpawners()
 {
-	if (TabSpawner.Num())
-	{
-		FString AllTabIdsString = FString::JoinBy(TabSpawner, TEXT(", "), [](const TPair<FName, TSharedRef<FTabSpawnerEntry>> Entry) { return Entry.Key.ToString(); });
-		UE_LOG(LogSlate, Verbose, TEXT("Unregistering all tab spawners on TabManager 0x%p for layout '%s': %s"), this, *ActiveLayoutName.ToString(), *AllTabIdsString);
-	}
-
 	TabSpawner.Empty();
 }
 
@@ -962,8 +950,6 @@ TSharedPtr<SWidget> FTabManager::RestoreFrom(const TSharedRef<FLayout>& Layout,
 {
 	ActiveLayoutName = Layout->LayoutName;
 
-	UE_LOG(LogSlate, Log, TEXT("Restoring tab layout: %s in TabManager 0x%p"), *ActiveLayoutName.ToString(), this);
-
 	TSharedPtr<SDockingArea> PrimaryDockArea;
 	for (int32 AreaIndex=0; AreaIndex < Layout->Areas.Num(); ++AreaIndex )
 	{
@@ -2449,12 +2435,9 @@ FTabSpawnerEntry& FGlobalTabmanager::RegisterNomadTabSpawner(const FName TabId,
 	// Remove TabId if it was previously loaded. This allows re-loading the Editor UI layout without restarting the whole Editor (Window->Load Layout)
 	if (NomadTabSpawner->Contains(TabId))
 	{
-		UE_LOG(LogSlate, Verbose, TEXT("Unregistering nomad tab spawner for layout '%s': '%s'"), *ActiveLayoutName.ToString(), *TabId.ToString());
 		UnregisterNomadTabSpawner(TabId);
 	}
 
-	UE_LOG(LogSlate, Verbose, TEXT("Registering nomad tab spawner for layout '%s': '%s'"), *ActiveLayoutName.ToString(), *TabId.ToString());
-
 	// (Re)create and return NewSpawnerEntry
 	TSharedRef<FTabSpawnerEntry> NewSpawnerEntry = MakeShareable(new FTabSpawnerEntry(TabId, OnSpawnTab, CanSpawnTab));
 	NomadTabSpawner->Add(TabId, NewSpawnerEntry);
@@ -2463,7 +2446,6 @@ FTabSpawnerEntry& FGlobalTabmanager::RegisterNomadTabSpawner(const FName TabId,
 
 void FGlobalTabmanager::UnregisterNomadTabSpawner( const FName TabId )
 {
-	UE_LOG(LogSlate, Verbose, TEXT("Unregistering nomad tab spawner for layout '%s': '%s'"), *ActiveLayoutName.ToString(), *TabId.ToString());
 	const int32 NumRemoved = NomadTabSpawner->Remove(TabId);
 }
 
