Slate: Convert SCanvas, SRadialBox, SSplitter, SUniformgridPanel and SWidgetSWitcher FSlot's to the FSlotArgument syntax.

diff --git a/Engine/Plugins/VirtualProduction/LevelSnapshots/Source/LevelSnapshotsEditor/Private/Views/Results/SLevelSnapshotsEditorResults.cpp b/Engine/Plugins/VirtualProduction/LevelSnapshots/Source/LevelSnapshotsEditor/Private/Views/Results/SLevelSnapshotsEditorResults.cpp
--- a/Engine/Plugins/VirtualProduction/LevelSnapshots/Source/LevelSnapshotsEditor/Private/Views/Results/SLevelSnapshotsEditorResults.cpp
+++ b/Engine/Plugins/VirtualProduction/LevelSnapshots/Source/LevelSnapshotsEditor/Private/Views/Results/SLevelSnapshotsEditorResults.cpp
@@ -2343,7 +2343,7 @@ void SLevelSnapshotsEditorResultsRow::Construct(const FArguments& InArgs, const
 				SnapshotChildWidget.ToSharedRef()
 			]
 		]
-	].OnSlotResized_Handler.BindSP(this, &SLevelSnapshotsEditorResultsRow::SetSnapshotColumnSize);
+	].OnSlotResized(this, &SLevelSnapshotsEditorResultsRow::SetSnapshotColumnSize);
 
 	const auto SlotDelegate1 = TAttribute<float>::FGetter::CreateSP(this, &SLevelSnapshotsEditorResultsRow::GetSplitterSlotSize, 1);;
 	SplitterPtr->SlotAt(1).SizeValue.Bind(SlotDelegate1);
@@ -2409,7 +2409,7 @@ void SLevelSnapshotsEditorResultsRow::Construct(const FArguments& InArgs, const
 				WorldChildWidget.ToSharedRef()
 			]
 		]
-	].OnSlotResized_Handler.BindSP(this, &SLevelSnapshotsEditorResultsRow::SetWorldObjectColumnSize);
+	].OnSlotResized(this, &SLevelSnapshotsEditorResultsRow::SetWorldObjectColumnSize);
 
 	const auto SlotDelegate2 = TAttribute<float>::FGetter::CreateSP(this, &SLevelSnapshotsEditorResultsRow::GetSplitterSlotSize, 2);;
 	SplitterPtr->SlotAt(2).SizeValue.Bind(SlotDelegate2);
diff --git a/Engine/Source/Editor/Persona/Private/AnimTimeline/SAnimTimelineSplitterOverlay.cpp b/Engine/Source/Editor/Persona/Private/AnimTimeline/SAnimTimelineSplitterOverlay.cpp
--- a/Engine/Source/Editor/Persona/Private/AnimTimeline/SAnimTimelineSplitterOverlay.cpp
+++ b/Engine/Source/Editor/Persona/Private/AnimTimeline/SAnimTimelineSplitterOverlay.cpp
@@ -7,7 +7,7 @@ void SAnimTimelineSplitterOverlay::Construct( const FArguments& InArgs )
 {
 	SetVisibility(EVisibility::SelfHitTestInvisible);
 
-	Splitter = SNew(SSplitter) = InArgs;
+	Splitter = SArgumentNew(InArgs, SSplitter);
 	Splitter->SetVisibility(EVisibility::HitTestInvisible);
 	AddSlot()
 	[
diff --git a/Engine/Source/Editor/Sequencer/Private/SSequencerSplitterOverlay.h b/Engine/Source/Editor/Sequencer/Private/SSequencerSplitterOverlay.h
--- a/Engine/Source/Editor/Sequencer/Private/SSequencerSplitterOverlay.h
+++ b/Engine/Source/Editor/Sequencer/Private/SSequencerSplitterOverlay.h
@@ -24,7 +24,7 @@ public:
 	{
 		SetVisibility(EVisibility::SelfHitTestInvisible);
 
-		Splitter = SNew(SSplitter) = InArgs;
+		Splitter = SArgumentNew(InArgs, SSplitter);
 		Splitter->SetVisibility(EVisibility::HitTestInvisible);
 		AddSlot()
 		[
diff --git a/Engine/Source/Runtime/Engine/Private/Slate/SGameLayerManager.cpp b/Engine/Source/Runtime/Engine/Private/Slate/SGameLayerManager.cpp
--- a/Engine/Source/Runtime/Engine/Private/Slate/SGameLayerManager.cpp
+++ b/Engine/Source/Runtime/Engine/Private/Slate/SGameLayerManager.cpp
@@ -523,8 +523,8 @@ void SGameLayerManager::AddOrUpdatePlayerLayers(const FGeometry& AllottedGeometr
 				Size.Y -= WindowTitleBarVerticalBox->GetDesiredSize().Y;
 			}
 
-			PlayerLayer->Slot->Size(Size);
-			PlayerLayer->Slot->Position(Position);
+			PlayerLayer->Slot->SetSize(Size);
+			PlayerLayer->Slot->SetPosition(Position);
 		}
 	}
 }
diff --git a/Engine/Source/Runtime/Slate/Private/Framework/Docking/SDockingSplitter.cpp b/Engine/Source/Runtime/Slate/Private/Framework/Docking/SDockingSplitter.cpp
--- a/Engine/Source/Runtime/Slate/Private/Framework/Docking/SDockingSplitter.cpp
+++ b/Engine/Source/Runtime/Slate/Private/Framework/Docking/SDockingSplitter.cpp
@@ -50,13 +50,11 @@ void SDockingSplitter::ReplaceChild( const TSharedRef<SDockingNode>& InChildToRe
 
 	Replacement->SetSizeCoefficient(InChildToReplace->GetSizeCoefficient());
 
-	Splitter->SlotAt( IndexInParentSplitter )
-	.Value( TAttribute<float>(Replacement, &SDockingNode::GetSizeCoefficient) )
-	.OnSlotResized( SSplitter::FOnSlotResized::CreateSP( Replacement, &SDockingNode::SetSizeCoefficient ) )
-	.SizeRule( TAttribute<SSplitter::ESizeRule>(Replacement, &SDockingNode::GetSizeRule) )
-	[
-		Replacement
-	];
+	SSplitter::FSlot& Slot = Splitter->SlotAt(IndexInParentSplitter);
+	Slot.SetSizeValue(TAttribute<float>(Replacement, &SDockingNode::GetSizeCoefficient));
+	Slot.OnSlotResized().BindSP(Replacement, &SDockingNode::SetSizeCoefficient);
+	Slot.SetSizingRule(TAttribute<SSplitter::ESizeRule>(Replacement, &SDockingNode::GetSizeRule));
+	Slot[Replacement];
 
 	Replacement->SetParentNode( SharedThis(this) );
 }
diff --git a/Engine/Source/Runtime/Slate/Private/Framework/MultiBox/SToolBarButtonBlock.cpp b/Engine/Source/Runtime/Slate/Private/Framework/MultiBox/SToolBarButtonBlock.cpp
--- a/Engine/Source/Runtime/Slate/Private/Framework/MultiBox/SToolBarButtonBlock.cpp
+++ b/Engine/Source/Runtime/Slate/Private/Framework/MultiBox/SToolBarButtonBlock.cpp
@@ -319,6 +319,7 @@ void SToolBarButtonBlock::BuildMultiBlockWidget(const ISlateStyle* StyleSet, con
 	if (OptionsBlockWidget.IsValid())
 	{
 		ChildSlot
+		.Padding(ToolBarStyle.ComboButtonPadding.Left, 0.0f, ToolBarStyle.ComboButtonPadding.Right, 0.0f)
 		[
 			SNew(SHorizontalBox)
 			+ SHorizontalBox::Slot()
@@ -345,7 +346,6 @@ void SToolBarButtonBlock::BuildMultiBlockWidget(const ISlateStyle* StyleSet, con
 				]
 			]
 		];
-		ChildSlot.Padding(ToolBarStyle.ComboButtonPadding.Left, 0.0f, ToolBarStyle.ComboButtonPadding.Right, 0.0f);
 	}
 	else
 	{
diff --git a/Engine/Source/Runtime/Slate/Private/Framework/MultiBox/SUniformToolbarPanel.cpp b/Engine/Source/Runtime/Slate/Private/Framework/MultiBox/SUniformToolbarPanel.cpp
--- a/Engine/Source/Runtime/Slate/Private/Framework/MultiBox/SUniformToolbarPanel.cpp
+++ b/Engine/Source/Runtime/Slate/Private/Framework/MultiBox/SUniformToolbarPanel.cpp
@@ -26,19 +26,12 @@ void SUniformToolbarPanel::Construct( const FArguments& InArgs )
 
 	ClippedIndex = INDEX_NONE;
 
-	Children.Reserve( InArgs.Slots.Num() );
-	for (FSlot* ChildSlot : InArgs.Slots)
-	{
-		Children.Add( ChildSlot );
-	}
+	Children.AddSlots(MoveTemp(const_cast<TArray<FSlot::FSlotArguments>&>(InArgs._Slots)));
 
 	const FToolBarStyle& ToolBarStyle = StyleSet->GetWidgetStyle<FToolBarStyle>(StyleName);
 	
-	// Add the optional dropdown arrow as a child.  
-	FSlot& NewSlot = *(new FSlot());
-
-	NewSlot
-	[
+	// Add the optional dropdown arrow as a child.
+	Children.AddSlot(FSlot::FSlotArguments(MakeUnique<FSlot>(
 		SAssignNew(Dropdown, SComboButton)
 		.HasDownArrow(false)
 		.ButtonStyle(&ToolBarStyle.ButtonStyle)
@@ -51,9 +44,7 @@ void SUniformToolbarPanel::Construct( const FArguments& InArgs )
 			SNew(SImage)
 			.Image(&ToolBarStyle.ExpandBrush)
 		]
-	];
-
-	Children.Add(&NewSlot);
+		)));
 }
 
 void SUniformToolbarPanel::OnArrangeChildren( const FGeometry& AllottedGeometry, FArrangedChildren& ArrangedChildren ) const
@@ -239,14 +230,14 @@ void SUniformToolbarPanel::SetSlotPadding(TAttribute<FMargin> InSlotPadding)
 	Invalidate(EInvalidateWidgetReason::Layout);
 }
 
-
-SUniformToolbarPanel::FSlot& SUniformToolbarPanel::AddSlot()
+SUniformToolbarPanel::FSlot::FSlotArguments SUniformToolbarPanel::Slot()
 {
-	FSlot& NewSlot = *(new FSlot());
-
-	Children.Insert( &NewSlot, Children.Num()-1);
+	return FSlot::FSlotArguments(MakeUnique<FSlot>());
+}
 
-	return NewSlot;
+SUniformToolbarPanel::FScopedWidgetSlotArguments SUniformToolbarPanel::AddSlot()
+{
+	return FScopedWidgetSlotArguments{ MakeUnique<FSlot>(), Children, INDEX_NONE };
 }
 
 bool SUniformToolbarPanel::RemoveSlot( const TSharedRef<SWidget>& SlotWidget )
diff --git a/Engine/Source/Runtime/Slate/Private/Framework/MultiBox/SUniformToolbarPanel.h b/Engine/Source/Runtime/Slate/Private/Framework/MultiBox/SUniformToolbarPanel.h
--- a/Engine/Source/Runtime/Slate/Private/Framework/MultiBox/SUniformToolbarPanel.h
+++ b/Engine/Source/Runtime/Slate/Private/Framework/MultiBox/SUniformToolbarPanel.h
@@ -28,23 +28,28 @@ public:
 	struct FSlot : public TSlotBase<FSlot>, public TAlignmentWidgetSlotMixin<FSlot>
 	{
 		FSlot()
-		: TSlotBase<FSlot>()
-		, TAlignmentWidgetSlotMixin<FSlot>(HAlign_Fill, VAlign_Fill)
+			: TSlotBase<FSlot>()
+			, TAlignmentWidgetSlotMixin<FSlot>(HAlign_Fill, VAlign_Fill)
+		{
+		}
+		FSlot(const TSharedRef<SWidget>& InWidget)
+			: TSlotBase<FSlot>(InWidget)
+			, TAlignmentWidgetSlotMixin<FSlot>(HAlign_Fill, VAlign_Fill)
 		{
 		}
 
+		SLATE_SLOT_BEGIN_ARGS_OneMixin(FSlot, TSlotBase<FSlot>, TAlignmentWidgetSlotMixin<FSlot>)
+		SLATE_SLOT_END_ARGS()
+
+		void Construct(const FChildren& SlotOwner, FSlotArguments&& InArgs)
+		{
+			TAlignmentWidgetSlotMixin<FSlot>::ConstructMixin(SlotOwner, MoveTemp(InArgs));
+			TSlotBase<FSlot>::Construct(SlotOwner, MoveTemp(InArgs));
+		}
 	};
 
 	SUniformToolbarPanel();
 
-	/**
-	 * Used by declarative syntax to create a Slot in the specified Column, Row.
-	 */
-	static FSlot& Slot()
-	{
-		return *(new FSlot());
-	}
-
 	SLATE_BEGIN_ARGS(SUniformToolbarPanel)
 		: _Orientation(EOrientation::Orient_Horizontal)
 		, _SlotPadding(FMargin(0.0f))
@@ -56,7 +61,7 @@ public:
 		}
 
 		/** Slot type supported by this panel */
-		SLATE_SUPPORTS_SLOT(FSlot)
+		SLATE_SLOT_ARGUMENT(FSlot, Slots)
 
 		SLATE_ARGUMENT(const ISlateStyle*, StyleSet)
 		SLATE_ARGUMENT(FName, StyleName)
@@ -99,12 +104,18 @@ public:
 	/** See SlotPadding attribute */
 	void SetSlotPadding(TAttribute<FMargin> InSlotPadding);
 
+	/**
+	 * Used by declarative syntax to create a Slot in the specified Column, Row.
+	 */
+	static FSlot::FSlotArguments Slot();
+
+	using FScopedWidgetSlotArguments = TPanelChildren<FSlot>::FScopedWidgetSlotArguments;
 	/**
 	 * Dynamically add a new slot to the Panel
 	 *
 	 * @return A reference to the newly-added slot
 	 */
-	FSlot& AddSlot();
+	FScopedWidgetSlotArguments AddSlot();
 	
 	/**
 	 * Removes a slot from this panel which contains the specified SWidget
diff --git a/Engine/Source/Runtime/Slate/Private/Widgets/Layout/SRadialBox.cpp b/Engine/Source/Runtime/Slate/Private/Widgets/Layout/SRadialBox.cpp
--- a/Engine/Source/Runtime/Slate/Private/Widgets/Layout/SRadialBox.cpp
+++ b/Engine/Source/Runtime/Slate/Private/Widgets/Layout/SRadialBox.cpp
@@ -8,16 +8,14 @@ SRadialBox::SRadialBox()
 {
 }
 
-SRadialBox::FSlot& SRadialBox::Slot()
+SRadialBox::FSlot::FSlotArguments SRadialBox::Slot()
 {
-	return *( new SRadialBox::FSlot() );
+	return FSlot::FSlotArguments(MakeUnique<FSlot>());
 }
 
-SRadialBox::FSlot& SRadialBox::AddSlot()
+SRadialBox::FScopedWidgetSlotArguments SRadialBox::AddSlot()
 {
-	SRadialBox::FSlot* NewSlot = new SRadialBox::FSlot();
-	Slots.Add(NewSlot);
-	return *NewSlot;
+	return FScopedWidgetSlotArguments{ MakeUnique<FSlot>(), Slots, INDEX_NONE };
 }
 
 int32 SRadialBox::RemoveSlot( const TSharedRef<SWidget>& SlotWidget )
@@ -34,10 +32,7 @@ void SRadialBox::Construct( const FArguments& InArgs )
 	SectorCentralAngle = InternalNormalizeAngle(InArgs._SectorCentralAngle);
 	AngleBetweenItems = InArgs._AngleBetweenItems;
 
-	for ( int32 ChildIndex=0; ChildIndex < InArgs.Slots.Num(); ++ChildIndex )
-	{
-		Slots.Add( InArgs.Slots[ChildIndex] );
-	}
+	Slots.AddSlots(MoveTemp(const_cast<TArray<FSlot::FSlotArguments>&>(InArgs._Slots)));
 }
 
 void SRadialBox::Tick( const FGeometry& AllottedGeometry, const double InCurrentTime, const float InDeltaTime )
diff --git a/Engine/Source/Runtime/Slate/Private/Widgets/Layout/SSplitter.cpp b/Engine/Source/Runtime/Slate/Private/Widgets/Layout/SSplitter.cpp
--- a/Engine/Source/Runtime/Slate/Private/Widgets/Layout/SSplitter.cpp
+++ b/Engine/Source/Runtime/Slate/Private/Widgets/Layout/SSplitter.cpp
@@ -4,6 +4,31 @@
 #include "Layout/ArrangedChildren.h"
 #include "Rendering/DrawElements.h"
 
+void SSplitter::FSlot::Construct(const FChildren& SlotOwner, FSlotArguments&& InArgs)
+{
+	TSlotBase<FSlot>::Construct(SlotOwner, MoveTemp(InArgs));
+	if (InArgs._Value.IsSet())
+	{
+		SizeValue = MoveTemp(InArgs._Value);
+	}
+	if (InArgs._Resizable.IsSet())
+	{
+		bIsResizable = InArgs._Resizable.GetValue();
+	}
+	if (InArgs._MinSize.IsSet())
+	{
+		MinSizeValue = InArgs._MinSize.GetValue();
+	}
+	if (InArgs._OnSlotResized.IsBound())
+	{
+		OnSlotResized_Handler = MoveTemp(InArgs._OnSlotResized);
+	}
+	if (InArgs._SizeRule.IsSet())
+	{
+		SizingRule = MoveTemp(InArgs._SizeRule);
+	}
+}
+
 bool SSplitter::FSlot::CanBeResized() const
 {
 	const ESizeRule CurrentSizingRule = SizingRule.Get();
@@ -17,26 +42,14 @@ bool SSplitter::FSlot::CanBeResized() const
 	}
 }
 
-SSplitter::FSlot& SSplitter::Slot()
+SSplitter::FSlot::FSlotArguments SSplitter::Slot()
 {
-	return (*new FSlot());
+	return FSlot::FSlotArguments(MakeUnique<FSlot>());
 }
 
-SSplitter::FSlot& SSplitter::AddSlot( int32 AtIndex )
+SSplitter::FScopedWidgetSlotArguments SSplitter::AddSlot( int32 AtIndex )
 {
-	FSlot& NewSlot = *new FSlot();
-	if ( AtIndex == INDEX_NONE )
-	{
-		// No index was specified; just add to the end of the list.
-		this->Children.Add( &NewSlot );
-	}
-	else
-	{
-		// Add a slot at the desired location
-		this->Children.Insert( &NewSlot, AtIndex );
-	}
-
-	return NewSlot;
+	return FScopedWidgetSlotArguments{ MakeUnique<FSlot>(), Children, AtIndex };
 }
 
 SSplitter::FSlot& SSplitter::SlotAt( int32 SlotIndex )
@@ -70,10 +83,7 @@ void SSplitter::Construct( const SSplitter::FArguments& InArgs )
 	Style = InArgs._Style;
 	OnGetMaxSlotSize = InArgs._OnGetMaxSlotSize;
 	MinSplitterChildLength = InArgs._MinimumSlotHeight;
-	for (int32 SlotIndex = 0; SlotIndex < InArgs.Slots.Num(); ++SlotIndex )
-	{
-		Children.Add( InArgs.Slots[SlotIndex] );
-	}
+	Children.AddSlots(MoveTemp(const_cast<TArray<FSlot::FSlotArguments>&>(InArgs._Slots)));
 }
 
 TArray<FLayoutGeometry> SSplitter::ArrangeChildrenForLayout(const FGeometry& AllottedGeometry) const
@@ -831,6 +841,22 @@ int32 SSplitter::GetHandleBeingResizedFromMousePosition( float InPhysicalSplitte
 * by dragging the center of the splitter.							*
 ********************************************************************/
 
+SSplitter2x2::FSlot::FSlot(const TSharedRef<SWidget>& InWidget)
+	: TSlotBase<FSlot>(InWidget)
+	, PercentageAttribute(FVector2D(0.5, 0.5))
+{
+
+}
+
+void SSplitter2x2::FSlot::Construct(const FChildren& SlotOwner, FSlotArguments&& InArg)
+{
+	TSlotBase<FSlot>::Construct(SlotOwner, MoveTemp(InArg));
+	if (InArg._Percentage.IsSet())
+	{
+		PercentageAttribute = MoveTemp(InArg._Percentage);
+	}
+}
+
 SSplitter2x2::SSplitter2x2()
 : Children(this)
 {
@@ -838,10 +864,10 @@ SSplitter2x2::SSplitter2x2()
 
 void SSplitter2x2::Construct( const FArguments& InArgs )
 {
-	Children.Add( new FSlot(InArgs._TopLeft.Widget) );
-	Children.Add( new FSlot(InArgs._BottomLeft.Widget) );
-	Children.Add( new FSlot(InArgs._TopRight.Widget) );
-	Children.Add( new FSlot(InArgs._BottomRight.Widget) );
+	Children.AddSlot(FSlot::FSlotArguments(MakeUnique<FSlot>(InArgs._TopLeft.Widget)));
+	Children.AddSlot(FSlot::FSlotArguments(MakeUnique<FSlot>(InArgs._BottomLeft.Widget)));
+	Children.AddSlot(FSlot::FSlotArguments(MakeUnique<FSlot>(InArgs._TopRight.Widget)));
+	Children.AddSlot(FSlot::FSlotArguments(MakeUnique<FSlot>(InArgs._BottomRight.Widget)));
 
 	SplitterHandleSize = 5.0f;
 	MinSplitterChildLength = 20.0f;
@@ -871,7 +897,7 @@ TArray<FLayoutGeometry> SSplitter2x2::ArrangeChildrenForLayout( const FGeometry&
 
 		// Calculate the amount of space that this child should take up.  
 		// It is based on the current percentage of space it should take up which is defined by a user moving the splitters
-		const FVector2D ChildSpace = SpaceAllottedForChildren * CurSlot.PercentageAttribute.Get();
+		const FVector2D ChildSpace = SpaceAllottedForChildren * CurSlot.GetPercentage();
 
 		// put them in their spot
 		Result.Emplace(FSlateLayoutTransform(Offset), ChildSpace);
@@ -1166,16 +1192,17 @@ void SSplitter2x2::SetBottomRightContent( TSharedRef< SWidget > BottomRightConte
 
 void SSplitter2x2::GetSplitterPercentages( TArray< FVector2D >& OutPercentages ) const
 {
-	OutPercentages.Empty();
+	OutPercentages.Reset(4);
 	for (int32 i = 0; i < 4; ++i)
 	{
-		OutPercentages.Add(Children[i].PercentageAttribute.Get());
+		OutPercentages.Add(Children[i].GetPercentage());
 	}
 }
 
-void SSplitter2x2::SetSplitterPercentages( const TArray< FVector2D >& InPercentages )
+void SSplitter2x2::SetSplitterPercentages( TArrayView< FVector2D > InPercentages )
 {
-	for (int32 i = 0; i < 4; ++i)
+	int32 MaxSize = FMath::Min(InPercentages.Num(), 4);
+	for (int32 i = 0; i < MaxSize; ++i)
 	{
 		Children[i].SetPercentage(InPercentages[i]);
 	}
diff --git a/Engine/Source/Runtime/Slate/Private/Widgets/Layout/SUniformGridPanel.cpp b/Engine/Source/Runtime/Slate/Private/Widgets/Layout/SUniformGridPanel.cpp
--- a/Engine/Source/Runtime/Slate/Private/Widgets/Layout/SUniformGridPanel.cpp
+++ b/Engine/Source/Runtime/Slate/Private/Widgets/Layout/SUniformGridPanel.cpp
@@ -8,6 +8,11 @@ SUniformGridPanel::SUniformGridPanel()
 {
 }
 
+SUniformGridPanel::FSlot::FSlotArguments SUniformGridPanel::Slot(int32 Column, int32 Row)
+{
+	return FSlot::FSlotArguments(MakeUnique<FSlot>(Column, Row));
+}
+
 void SUniformGridPanel::Construct( const FArguments& InArgs )
 {
 	SlotPadding = InArgs._SlotPadding;
@@ -16,12 +21,7 @@ void SUniformGridPanel::Construct( const FArguments& InArgs )
 	MinDesiredSlotWidth = InArgs._MinDesiredSlotWidth.Get();
 	MinDesiredSlotHeight = InArgs._MinDesiredSlotHeight.Get();
 
-	Children.Reserve( InArgs.Slots.Num() );
-	for (int32 ChildIndex=0; ChildIndex < InArgs.Slots.Num(); ChildIndex++)
-	{
-		FSlot* ChildSlot = InArgs.Slots[ChildIndex];
-		Children.Add( ChildSlot );
-	}
+	Children.AddSlots(MoveTemp(const_cast<TArray<FSlot::FSlotArguments>&>(InArgs._Slots)));
 }
 
 void SUniformGridPanel::OnArrangeChildren( const FGeometry& AllottedGeometry, FArrangedChildren& ArrangedChildren ) const
@@ -43,7 +43,7 @@ void SUniformGridPanel::OnArrangeChildren( const FGeometry& AllottedGeometry, FA
 
 				ArrangedChildren.AddWidget(ChildVisibility,
 					AllottedGeometry.MakeChild(Child.GetWidget(),
-					FVector2D(CellSize.X*Child.Column + XAxisResult.Offset, CellSize.Y*Child.Row + YAxisResult.Offset),
+					FVector2D(CellSize.X*Child.GetColumn() + XAxisResult.Offset, CellSize.Y*Child.GetRow() + YAxisResult.Offset),
 					FVector2D(XAxisResult.Size, YAxisResult.Size)
 					));
 			}
@@ -70,8 +70,8 @@ FVector2D SUniformGridPanel::ComputeDesiredSize( float ) const
 		if (Child.GetWidget()->GetVisibility() != EVisibility::Collapsed)
 		{
 			// A single cell at (N,M) means our grid size is (N+1, M+1)
-			NumColumns = FMath::Max(Child.Column + 1, NumColumns);
-			NumRows = FMath::Max(Child.Row + 1, NumRows);
+			NumColumns = FMath::Max(Child.GetColumn() + 1, NumColumns);
+			NumRows = FMath::Max(Child.GetRow() + 1, NumRows);
 
 			FVector2D ChildDesiredSize = Child.GetWidget()->GetDesiredSize() + SlotPaddingDesiredSize;
 
@@ -106,13 +106,9 @@ void SUniformGridPanel::SetMinDesiredSlotHeight(TAttribute<float> InMinDesiredSl
 	MinDesiredSlotHeight = InMinDesiredSlotHeight;
 }
 
-SUniformGridPanel::FSlot& SUniformGridPanel::AddSlot( int32 Column, int32 Row )
+SUniformGridPanel::FScopedWidgetSlotArguments SUniformGridPanel::AddSlot( int32 Column, int32 Row )
 {
-	FSlot& NewSlot = *(new FSlot( Column, Row ));
-
-	Children.Add( &NewSlot );
-
-	return NewSlot;
+	return FScopedWidgetSlotArguments{ MakeUnique<FSlot>(Column, Row), Children, INDEX_NONE };
 }
 
 bool SUniformGridPanel::RemoveSlot( const TSharedRef<SWidget>& SlotWidget )
diff --git a/Engine/Source/Runtime/Slate/Private/Widgets/Layout/SWidgetSwitcher.cpp b/Engine/Source/Runtime/Slate/Private/Widgets/Layout/SWidgetSwitcher.cpp
--- a/Engine/Source/Runtime/Slate/Private/Widgets/Layout/SWidgetSwitcher.cpp
+++ b/Engine/Source/Runtime/Slate/Private/Widgets/Layout/SWidgetSwitcher.cpp
@@ -17,34 +17,37 @@ SWidgetSwitcher::SWidgetSwitcher()
 	SetCanTick(false);
 }
 
-SWidgetSwitcher::FSlot& SWidgetSwitcher::AddSlot( int32 SlotIndex )
+SWidgetSwitcher::FSlot::FSlotArguments SWidgetSwitcher::Slot()
 {
-	FSlot* NewSlot = new FSlot();
+	return FSlot::FSlotArguments(MakeUnique<FSlot>());
+}
 
+SWidgetSwitcher::FScopedWidgetSlotArguments SWidgetSwitcher::AddSlot(int32 SlotIndex)
+{
 	if (!AllChildren.IsValidIndex(SlotIndex))
 	{
-		AllChildren.Add(NewSlot);
+		return FScopedWidgetSlotArguments{ MakeUnique<FSlot>(), AllChildren, INDEX_NONE };
 	}
 	else
 	{
-		AllChildren.Insert(NewSlot, SlotIndex);
-
-		// adjust the active WidgetIndex based on this slot change
-		if (!WidgetIndex.IsBound() && WidgetIndex.Get() >= SlotIndex)
-		{
-			WidgetIndex = WidgetIndex.Get() + 1;
-		}
+		TWeakPtr<SWidgetSwitcher> WeakSwitcher = SharedThis(this);
+		return FScopedWidgetSlotArguments{ MakeUnique<FSlot>(), AllChildren, INDEX_NONE, [WeakSwitcher](const FSlot*, int32 SlotIndex)
+			{
+				if (TSharedPtr<SWidgetSwitcher> Switcher = WeakSwitcher.Pin())
+				{
+					// adjust the active WidgetIndex based on this slot change
+					if (!Switcher->WidgetIndex.IsBound() && Switcher->WidgetIndex.Get() >= SlotIndex)
+					{
+						Switcher->WidgetIndex = Switcher->WidgetIndex.Get() + 1;
+					}
+				}
+			}};
 	}
-
-	return *NewSlot;
 }
 
 void SWidgetSwitcher::Construct( const FArguments& InArgs )
 {
-	for (int32 Index = 0; Index < InArgs.Slots.Num(); ++Index)
-	{
-		AllChildren.Add(InArgs.Slots[Index]);
-	}
+	AllChildren.AddSlots(MoveTemp(const_cast<TArray<FSlot::FSlotArguments>&>(InArgs._Slots)));
 
 	WidgetIndex = InArgs._WidgetIndex;
 }
diff --git a/Engine/Source/Runtime/Slate/Private/Widgets/SCanvas.cpp b/Engine/Source/Runtime/Slate/Private/Widgets/SCanvas.cpp
--- a/Engine/Source/Runtime/Slate/Private/Widgets/SCanvas.cpp
+++ b/Engine/Source/Runtime/Slate/Private/Widgets/SCanvas.cpp
@@ -5,8 +5,25 @@
 #include "Layout/ArrangedChildren.h"
 
 
+PRAGMA_DISABLE_DEPRECATION_WARNINGS
+void SCanvas::FSlot::Construct(const FChildren& SlotOwner, FSlotArguments&& InArg)
+{
+	TSlotBase<FSlot>::Construct(SlotOwner, MoveTemp(InArg));
+	if (InArg._Position.IsSet())
+	{
+		PositionAttr = MoveTemp(InArg._Position);
+	}
+	if (InArg._Size.IsSet())
+	{
+		SizeAttr = MoveTemp(InArg._Size);
+	}
+	TAlignmentWidgetSlotMixin<FSlot>::ConstructMixin(SlotOwner, MoveTemp(InArg));
+}
+PRAGMA_ENABLE_DEPRECATION_WARNINGS
+
+
 SCanvas::SCanvas()
-: Children(this)
+	: Children(this)
 {
 	SetCanTick(false);
 	bCanSupportFocus = false;
@@ -14,11 +31,17 @@ SCanvas::SCanvas()
 
 void SCanvas::Construct( const SCanvas::FArguments& InArgs )
 {
-	const int32 NumSlots = InArgs.Slots.Num();
-	for ( int32 SlotIndex = 0; SlotIndex < NumSlots; ++SlotIndex )
-	{
-		Children.Add( InArgs.Slots[SlotIndex] );
-	}
+	Children.AddSlots(MoveTemp(const_cast<TArray<FSlot::FSlotArguments>&>(InArgs._Slots)));
+}
+
+SCanvas::FSlot::FSlotArguments SCanvas::Slot()
+{
+	return FSlot::FSlotArguments(MakeUnique<FSlot>());
+}
+
+SCanvas::FScopedWidgetSlotArguments SCanvas::AddSlot()
+{
+	return FScopedWidgetSlotArguments{ MakeUnique<FSlot>(), Children, INDEX_NONE };
 }
 
 void SCanvas::ClearChildren( )
@@ -38,7 +61,7 @@ void SCanvas::OnArrangeChildren( const FGeometry& AllottedGeometry, FArrangedChi
 		for (int32 ChildIndex = 0; ChildIndex < Children.Num(); ++ChildIndex)
 		{
 			const SCanvas::FSlot& CurChild = Children[ChildIndex];
-			const FVector2D Size = CurChild.SizeAttr.Get();
+			const FVector2D Size = CurChild.GetSize();
 
 			//Handle HAlignment
 			FVector2D Offset(0.0f, 0.0f);
@@ -75,7 +98,7 @@ void SCanvas::OnArrangeChildren( const FGeometry& AllottedGeometry, FArrangedChi
 				// The child widget being arranged
 				CurChild.GetWidget(),
 				// Child's local position (i.e. position within parent)
-				CurChild.PositionAttr.Get() + Offset,
+				CurChild.GetPosition() + Offset,
 				// Child's size
 				Size
 			));
diff --git a/Engine/Source/Runtime/Slate/Public/Widgets/Layout/SRadialBox.h b/Engine/Source/Runtime/Slate/Public/Widgets/Layout/SRadialBox.h
--- a/Engine/Source/Runtime/Slate/Public/Widgets/Layout/SRadialBox.h
+++ b/Engine/Source/Runtime/Slate/Public/Widgets/Layout/SRadialBox.h
@@ -25,10 +25,11 @@ public:
 	class FSlot : public TSlotBase<FSlot>
 	{
 	public:
-		FSlot()
-			: TSlotBase<FSlot>()
-		{
-		}
+		using TSlotBase<FSlot>::TSlotBase;
+
+		SLATE_SLOT_BEGIN_ARGS(FSlot, TSlotBase<FSlot>)
+		SLATE_SLOT_END_ARGS()
+		using TSlotBase<FSlot>::Construct;
 	};
 
 
@@ -39,12 +40,12 @@ public:
 		, _bDistributeItemsEvenly(true)
 		, _AngleBetweenItems(0.f)
 		, _SectorCentralAngle(360.f)
-	{
+		{
 			_Visibility = EVisibility::SelfHitTestInvisible;
 		}
 
 		/** The slot supported by this panel */
-		SLATE_SUPPORTS_SLOT( FSlot )
+		SLATE_SLOT_ARGUMENT( FSlot, Slots )
 
 		/** The preferred width, if not set will fill the space */
 		SLATE_ATTRIBUTE( float, PreferredWidth )
@@ -68,9 +69,10 @@ public:
 
 	SRadialBox();
 
-	static FSlot& Slot();
+	static FSlot::FSlotArguments Slot();
 
-	FSlot& AddSlot();
+	using FScopedWidgetSlotArguments = TPanelChildren<FSlot>::FScopedWidgetSlotArguments;
+	FScopedWidgetSlotArguments AddSlot();
 
 	/** Removes a slot from this radial box which contains the specified SWidget
 	 *
diff --git a/Engine/Source/Runtime/Slate/Public/Widgets/Layout/SSplitter.h b/Engine/Source/Runtime/Slate/Public/Widgets/Layout/SSplitter.h
--- a/Engine/Source/Runtime/Slate/Public/Widgets/Layout/SSplitter.h
+++ b/Engine/Source/Runtime/Slate/Public/Widgets/Layout/SSplitter.h
@@ -60,7 +60,7 @@ public:
 	DECLARE_DELEGATE_RetVal_OneParam(FVector2D, FOnGetMaxSlotSize, int32);
 
 public:
-	class FSlot : public TSlotBase<FSlot>
+	class SLATE_API FSlot : public TSlotBase<FSlot>
 	{
 	public:		
 		FSlot()
@@ -70,45 +70,75 @@ public:
 		{
 		}
 
+		SLATE_SLOT_BEGIN_ARGS(FSlot, TSlotBase<FSlot>)
+			/** The size rule used by the slot. */
+			SLATE_ATTRIBUTE(ESizeRule, SizeRule)
+			/** When the RuleSize is set to FractionOfParent, the size of the slot is the Value percentage of its parent size. */
+			SLATE_ATTRIBUTE(float, Value)
+			/** Minimum slot size when resizing. */
+			SLATE_ARGUMENT(TOptional<float>, MinSize)
+			/** Can the slot be resize by the user. */
+			SLATE_ARGUMENT(TOptional<bool>, Resizable)
+			/** Callback when the slot is resized. */
+			SLATE_EVENT(FOnSlotResized, OnSlotResized)
+		SLATE_SLOT_END_ARGS()
+
+		void Construct(const FChildren& SlotOwner, FSlotArguments&& InArgs);
+
 		/** When the RuleSize is set to FractionOfParent, the size of the slot is the Value percentage of its parent size. */
-		FSlot& Value( const TAttribute<float>& InValue )
+		void SetSizeValue( TAttribute<float> InValue )
+		{
+			SizeValue = MoveTemp(InValue);
+		}
+		float GetSizeValue() const
 		{
-			SizeValue = InValue;
-			return *this;
+			return SizeValue.Get();
 		}
 
 		/**
 		 * Can the slot be resize by the user.
 		 * @see CanBeResized()
 		 */
-		FSlot& Resizable(bool bInIsResizable)
+		void SetResizable(bool bInIsResizable)
 		{
 			bIsResizable = bInIsResizable;
-			return *this;
+		}
+		bool IsResizable() const
+		{
+			return bIsResizable.Get(false);
 		}
 
-		/** Minimun slot size when resizing. */
-		FSlot& MinSize(float InMinSize)
+		/** Minimum slot size when resizing. */
+		void SetMinSize(float InMinSize)
 		{
 			MinSizeValue = InMinSize;
-			return *this;
+		}
+		float GetMinSize() const
+		{
+			return MinSizeValue.Get(0.f);
 		}
 		
 		/**
 		 * Callback when the slot is resized.
 		 * @see CanBeResized()
 		 */
-		FSlot& OnSlotResized( const FOnSlotResized& InHandler )
+		FOnSlotResized& OnSlotResized()
 		{
-			OnSlotResized_Handler = InHandler;
-			return *this;
+			return OnSlotResized_Handler;
+		}
+		const FOnSlotResized& OnSlotResized() const
+		{
+			return OnSlotResized_Handler;
 		}
 
 		/** The size rule used by the slot. */
-		FSlot& SizeRule( const TAttribute<ESizeRule>& InSizeRule ) 
+		void SetSizingRule( TAttribute<ESizeRule> InSizeRule ) 
+		{
+			SizingRule = MoveTemp(InSizeRule);
+		}
+		ESizeRule GetSizingRule() const
 		{
-			SizingRule = InSizeRule;
-			return *this;
+			return SizingRule.Get();
 		}
 
 	public:
@@ -123,9 +153,10 @@ public:
 		TOptional<bool> bIsResizable;
 	};
 
-	/** @return a new SSplitter::FSlot() */
-	static FSlot& Slot();
+	/** @return Add a new FSlot() */
+	static FSlot::FSlotArguments Slot();
 	
+	using FScopedWidgetSlotArguments = TPanelChildren<FSlot>::FScopedWidgetSlotArguments;
 	/**
 	 * Add a slot to the splitter at the specified index
 	 * Sample usage:
@@ -136,7 +167,7 @@ public:
 	 *
 	 * @return the new slot.
 	 */
-	FSlot& AddSlot( int32 AtIndex = INDEX_NONE );
+	FScopedWidgetSlotArguments AddSlot( int32 AtIndex = INDEX_NONE );
 
 	DECLARE_DELEGATE_OneParam(FOnHandleHovered, int32);
 
@@ -151,7 +182,7 @@ public:
 		{
 		}
 
-		SLATE_SUPPORTS_SLOT(FSlot)
+		SLATE_SLOT_ARGUMENT(FSlot, Slots)
 
 		/** Style used to draw this splitter */
 		SLATE_STYLE_ARGUMENT( FSplitterStyle, Style )
@@ -354,35 +385,34 @@ protected:
  */
 class SLATE_API SSplitter2x2 : public SPanel
 {
-public:
-	class FSlot : public TSlotBase<FSlot>
+private:
+	class SLATE_API FSlot : public TSlotBase<FSlot>
 	{
 	public:	
-		/** Default Constructor.  Initially each slot takes up a quarter of the entire space */
-		FSlot()
-			: TSlotBase<FSlot>( SNullWidget::NullWidget )
-			, PercentageAttribute( FVector2D(0.5, 0.5) )
-		{
-		}
-
-		/** Copy Constructor */
-		FSlot( const TSharedRef<SWidget>& InWidget )
-			: TSlotBase<FSlot>( InWidget )
-			, PercentageAttribute( FVector2D(0.5, 0.5) )
-		{
-		}
+		SLATE_SLOT_BEGIN_ARGS(FSlot, TSlotBase<FSlot>)
+			SLATE_ATTRIBUTE(FVector2D, Percentage)
+		SLATE_SLOT_END_ARGS()
 
 		/**
 		 * Sets the percentage attribute
 		 *
 		 * @param Value The new percentage value
 		 */
-		FSlot& SetPercentage( const FVector2D& Value )
+		void SetPercentage( const FVector2D& Value )
 		{
 			PercentageAttribute.Set( Value );
-			return *this;
 		}
-	public:
+
+		FVector2D GetPercentage() const
+		{
+			return PercentageAttribute.Get();
+		}
+
+		FSlot(const TSharedRef<SWidget>& InWidget);
+
+		void Construct(const FChildren& SlotOwner, FSlotArguments&& InArg);
+
+	private:
 		/** The percentage of the alloted space of the splitter that this slot requires */
 		TAttribute<FVector2D> PercentageAttribute;
 	};
@@ -458,7 +488,7 @@ public:
 	void GetSplitterPercentages( TArray< FVector2D >& OutPercentages ) const;
 	
 	/** Sets the size percentages for the children in this order: TopLeft, BottomLeft, TopRight, BottomRight */
-	void SetSplitterPercentages( const TArray< FVector2D >& InPercentages );
+	void SetSplitterPercentages( TArrayView< FVector2D > InPercentages );
 
 
 private:
diff --git a/Engine/Source/Runtime/Slate/Public/Widgets/Layout/SUniformGridPanel.h b/Engine/Source/Runtime/Slate/Public/Widgets/Layout/SUniformGridPanel.h
--- a/Engine/Source/Runtime/Slate/Public/Widgets/Layout/SUniformGridPanel.h
+++ b/Engine/Source/Runtime/Slate/Public/Widgets/Layout/SUniformGridPanel.h
@@ -19,29 +19,68 @@ class SLATE_API SUniformGridPanel : public SPanel
 {
 public:
 	/** Stores the per-child info for this panel type */
+	PRAGMA_DISABLE_DEPRECATION_WARNINGS
 	struct FSlot : public TSlotBase<FSlot>, public TAlignmentWidgetSlotMixin<FSlot>
 	{
+		SLATE_SLOT_BEGIN_ARGS_OneMixin(FSlot, TSlotBase<FSlot>, TAlignmentWidgetSlotMixin<FSlot>)
+		SLATE_SLOT_END_ARGS()
+
 		FSlot( int32 InColumn, int32 InRow )
-		: TSlotBase<FSlot>()
-		, TAlignmentWidgetSlotMixin<FSlot>(HAlign_Fill, VAlign_Fill)
-		, Column( InColumn )
-		, Row( InRow )
+			: TSlotBase<FSlot>()
+			, TAlignmentWidgetSlotMixin<FSlot>(HAlign_Fill, VAlign_Fill)
+			, Column( InColumn )
+			, Row( InRow )
+			{
+			}
+
+		void Construct(const FChildren& SlotOwner, FSlotArguments&& InArgs)
+		{
+			TSlotBase<FSlot>::Construct(SlotOwner, MoveTemp(InArgs));
+			TAlignmentWidgetSlotMixin<FSlot>::ConstructMixin(SlotOwner, MoveTemp(InArgs));
+		}
+
+		void SetColumn(int32 InColumn)
+		{
+			if (InColumn != Column)
+			{
+				Column = InColumn;
+				Invalidate(EInvalidateWidgetReason::Layout);
+			}
+		}
+
+		int32 GetColumn() const
 		{
+			return Column;
 		}
 
+		int32 GetRow() const
+		{
+			return Row;
+		}
+
+		void SetRow(int32 InRow)
+		{
+			if (InRow != Row)
+			{
+				Row = InRow;
+				Invalidate(EInvalidateWidgetReason::Layout);
+			}
+		}
+
+	public:
+		UE_DEPRECATED(5.0, "Direct access to Column is now deprecated. Use the getter or setter.")
 		int32 Column;
+		UE_DEPRECATED(5.0, "Direct access to Row is now deprecated. Use the getter or setter.")
 		int32 Row;
 	};
+	PRAGMA_ENABLE_DEPRECATION_WARNINGS
 
 	SUniformGridPanel();
 
 	/**
 	 * Used by declarative syntax to create a Slot in the specified Column, Row.
 	 */
-	static FSlot& Slot( int32 Column, int32 Row )
-	{
-		return *(new FSlot( Column, Row ));
-	}
+	static FSlot::FSlotArguments Slot( int32 Column, int32 Row );
 
 	SLATE_BEGIN_ARGS( SUniformGridPanel )
 		: _SlotPadding( FMargin(0.0f) )
@@ -52,7 +91,7 @@ public:
 		}
 
 		/** Slot type supported by this panel */
-		SLATE_SUPPORTS_SLOT(FSlot)
+		SLATE_SLOT_ARGUMENT(FSlot, Slots)
 		
 		/** Padding given to each slot */
 		SLATE_ATTRIBUTE(FMargin, SlotPadding)
@@ -81,12 +120,13 @@ public:
 	/** See MinDesiredSlotHeight attribute */
 	void SetMinDesiredSlotHeight(TAttribute<float> InMinDesiredSlotHeight);
 
+	using FScopedWidgetSlotArguments = TPanelChildren<FSlot>::FScopedWidgetSlotArguments;
 	/**
 	 * Dynamically add a new slot to the UI at specified Column and Row.
 	 *
 	 * @return A reference to the newly-added slot
 	 */
-	FSlot& AddSlot( int32 Column, int32 Row );
+	FScopedWidgetSlotArguments AddSlot( int32 Column, int32 Row );
 	
 	/**
 	 * Removes a slot from this panel which contains the specified SWidget
diff --git a/Engine/Source/Runtime/Slate/Public/Widgets/Layout/SWidgetSwitcher.h b/Engine/Source/Runtime/Slate/Public/Widgets/Layout/SWidgetSwitcher.h
--- a/Engine/Source/Runtime/Slate/Public/Widgets/Layout/SWidgetSwitcher.h
+++ b/Engine/Source/Runtime/Slate/Public/Widgets/Layout/SWidgetSwitcher.h
@@ -31,7 +31,7 @@ public:
 			_Visibility = EVisibility::SelfHitTestInvisible;
 		}
 
-		SLATE_SUPPORTS_SLOT(FSlot)
+		SLATE_SLOT_ARGUMENT(FSlot, Slots)
 
 		/** Holds the index of the initial widget to be displayed (INDEX_NONE = default). */
 		SLATE_ATTRIBUTE(int32, WidgetIndex)
@@ -42,12 +42,13 @@ public:
 
 public:
 
+	using FScopedWidgetSlotArguments = TPanelChildren<FSlot>::FScopedWidgetSlotArguments;
 	/**
 	 * Adds a slot to the widget switcher at the specified location.
 	 *
 	 * @param SlotIndex The index at which to insert the slot, or INDEX_NONE to append.
 	 */
-	FSlot& AddSlot( int32 SlotIndex = INDEX_NONE );
+	FScopedWidgetSlotArguments AddSlot( int32 SlotIndex = INDEX_NONE );
 
 	/**
 	 * Constructs the widget.
@@ -130,10 +131,7 @@ public:
 	 *
 	 * @return A new slot.
 	 */
-	static FSlot& Slot( )
-	{
-		return *(new FSlot());
-	}
+	static FSlot::FSlotArguments Slot();
 
 protected:
 
diff --git a/Engine/Source/Runtime/Slate/Public/Widgets/SCanvas.h b/Engine/Source/Runtime/Slate/Public/Widgets/SCanvas.h
--- a/Engine/Source/Runtime/Slate/Public/Widgets/SCanvas.h
+++ b/Engine/Source/Runtime/Slate/Public/Widgets/SCanvas.h
@@ -49,27 +49,43 @@ public:
 	 * 
 	 *  Note: FILL is NOT supported.
 	 */
-	class FSlot : public TSlotBase<FSlot>, public TAlignmentWidgetSlotMixin<FSlot>
+	PRAGMA_DISABLE_DEPRECATION_WARNINGS
+	class SLATE_API FSlot : public TSlotBase<FSlot>, public TAlignmentWidgetSlotMixin<FSlot>
 	{
-	public:		
-		FSlot& Position( const TAttribute<FVector2D>& InPosition )
+	public:
+		SLATE_SLOT_BEGIN_ARGS_OneMixin(FSlot, TSlotBase<FSlot>, TAlignmentWidgetSlotMixin<FSlot>)
+			SLATE_ATTRIBUTE(FVector2D, Position)
+			SLATE_ATTRIBUTE(FVector2D, Size)
+		SLATE_SLOT_END_ARGS()
+
+		void SetPosition( TAttribute<FVector2D> InPosition )
+		{
+			PositionAttr = MoveTemp(InPosition);
+		}
+		FVector2D GetPosition() const
 		{
-			PositionAttr = InPosition;
-			return *this;
+			return PositionAttr.Get();
 		}
 
-		FSlot& Size( const TAttribute<FVector2D>& InSize )
+		void SetSize( TAttribute<FVector2D> InSize )
+		{
+			SizeAttr = MoveTemp(InSize);
+		}
+		FVector2D GetSize() const
 		{
-			SizeAttr = InSize;
-			return *this;
+			return SizeAttr.Get();
 		}
 
+	public:
 		/** Position */
+		UE_DEPRECATED(5.0, "Direct access to PositionAttr is now deprecated. Use the getter or setter.")
 		TAttribute<FVector2D> PositionAttr;
 
 		/** Size */
+		UE_DEPRECATED(5.0, "Direct access to SizeAttr is now deprecated. Use the getter or setter.")
 		TAttribute<FVector2D> SizeAttr;
 
+	public:
 		/** Default values for a slot. */
 		FSlot()
 			: TSlotBase<FSlot>()
@@ -77,14 +93,17 @@ public:
 			, PositionAttr( FVector2D::ZeroVector )
 			, SizeAttr( FVector2D( 1.0f, 1.0f ) )
 		{ }
+
+		void Construct(const FChildren& SlotOwner, FSlotArguments&& InArg);
 	};
+	PRAGMA_ENABLE_DEPRECATION_WARNINGS
 
 	SLATE_BEGIN_ARGS( SCanvas )
 		{
 			_Visibility = EVisibility::SelfHitTestInvisible;
 		}
 
-		SLATE_SUPPORTS_SLOT( SCanvas::FSlot )
+		SLATE_SLOT_ARGUMENT( FSlot, Slots )
 
 	SLATE_END_ARGS()
 
@@ -97,22 +116,15 @@ public:
 	 */
 	void Construct( const FArguments& InArgs );
 
-	static FSlot& Slot()
-	{
-		return *(new FSlot());
-	}
+	static FSlot::FSlotArguments Slot();
 
+	using FScopedWidgetSlotArguments = TPanelChildren<FSlot>::FScopedWidgetSlotArguments;
 	/**
 	 * Adds a content slot.
 	 *
 	 * @return The added slot.
 	 */
-	FSlot& AddSlot()
-	{
-		SCanvas::FSlot& NewSlot = *new FSlot();
-		this->Children.Add( &NewSlot );
-		return NewSlot;
-	}
+	FScopedWidgetSlotArguments AddSlot();
 
 	/**
 	 * Removes a particular content slot.
diff --git a/Engine/Source/Runtime/SlateCore/Public/Widgets/DeclarativeSyntaxSupport.h b/Engine/Source/Runtime/SlateCore/Public/Widgets/DeclarativeSyntaxSupport.h
--- a/Engine/Source/Runtime/SlateCore/Public/Widgets/DeclarativeSyntaxSupport.h
+++ b/Engine/Source/Runtime/SlateCore/Public/Widgets/DeclarativeSyntaxSupport.h
@@ -38,7 +38,7 @@ template<typename WidgetType> struct TSlateBaseNamedArgs;
 	MakeTDecl<WidgetType>( #WidgetType, __FILE__, __LINE__, RequiredArgs::MakeRequiredArgs(__VA_ARGS__) ) . Expose( ExposeAs ) <<= TYPENAME_OUTSIDE_TEMPLATE WidgetType::FArguments()
 
 #define SArgumentNew( InArgs, WidgetType, ... ) \
-	MakeTDecl<WidgetType>( #WidgetType, __FILE__, __LINE__, RequiredArgs::MakeRequiredArgs(__VA_ARGS__) ) <<= TYPENAME_OUTSIDE_TEMPLATE InArgs
+	MakeTDecl<WidgetType>( #WidgetType, __FILE__, __LINE__, RequiredArgs::MakeRequiredArgs(__VA_ARGS__) ) <<= InArgs
 
 
 /**
diff --git a/Engine/Source/Runtime/UMG/Private/Components/UniformGridSlot.cpp b/Engine/Source/Runtime/UMG/Private/Components/UniformGridSlot.cpp
--- a/Engine/Source/Runtime/UMG/Private/Components/UniformGridSlot.cpp
+++ b/Engine/Source/Runtime/UMG/Private/Components/UniformGridSlot.cpp
@@ -37,7 +37,7 @@ void UUniformGridSlot::SetRow(int32 InRow)
 	Row = InRow;
 	if ( Slot )
 	{
-		Slot->Row = InRow;
+		Slot->SetRow(InRow);
 	}
 }
 
@@ -46,7 +46,7 @@ void UUniformGridSlot::SetColumn(int32 InColumn)
 	Column = InColumn;
 	if ( Slot )
 	{
-		Slot->Column = InColumn;
+		Slot->SetColumn(InColumn);
 	}
 }
 
