Slate: Convert SGridPanel FSlot's to the FSlotArgument syntax.

diff --git a/Engine/Source/Developer/TraceInsights/Private/Insights/Widgets/SStartPageWindow.cpp b/Engine/Source/Developer/TraceInsights/Private/Insights/Widgets/SStartPageWindow.cpp
--- a/Engine/Source/Developer/TraceInsights/Private/Insights/Widgets/SStartPageWindow.cpp
+++ b/Engine/Source/Developer/TraceInsights/Private/Insights/Widgets/SStartPageWindow.cpp
@@ -555,7 +555,9 @@ private:
 		typename TAttribute<FText>::FGetter::template TConstMethodPtr<STraceListRow> InValueTextFn,
 		typename TAttribute<EVisibility>::FGetter::template TConstMethodPtr<STraceListRow> InVisibilityFn = nullptr) const
 	{
-		SGridPanel::FSlot& Slot0 = Grid->AddSlot(0, Row)
+		SGridPanel::FSlot* Slot0 = nullptr;
+		Grid->AddSlot(0, Row)
+			.Expose(Slot0)
 			.Padding(2.0f)
 			.HAlign(HAlign_Right)
 			[
@@ -565,7 +567,9 @@ private:
 				.ColorAndOpacity(FLinearColor::Gray)
 			];
 
-		SGridPanel::FSlot& Slot1 = Grid->AddSlot(1, Row)
+		SGridPanel::FSlot* Slot1 = nullptr;
+		Grid->AddSlot(1, Row)
+			.Expose(Slot1)
 			.Padding(2.0f)
 			.HAlign(HAlign_Left)
 			[
@@ -578,14 +582,14 @@ private:
 
 		if (InVisibilityFn)
 		{
-			Slot0.GetWidget()->SetVisibility(MakeAttributeSP(this, InVisibilityFn));
-			Slot1.GetWidget()->SetVisibility(MakeAttributeSP(this, InVisibilityFn));
+			Slot0->GetWidget()->SetVisibility(MakeAttributeSP(this, InVisibilityFn));
+			Slot1->GetWidget()->SetVisibility(MakeAttributeSP(this, InVisibilityFn));
 		}
 		else
 		{
 			auto Fn = MakeAttributeSP(this, InValueTextFn);
-			Slot0.GetWidget()->SetVisibility(MakeAttributeLambda([this, Fn]() { return Fn.Get().IsEmpty() ? EVisibility::Collapsed : EVisibility::Visible; }));
-			Slot1.GetWidget()->SetVisibility(MakeAttributeLambda([this, Fn]() { return Fn.Get().IsEmpty() ? EVisibility::Collapsed : EVisibility::Visible; }));
+			Slot0->GetWidget()->SetVisibility(MakeAttributeLambda([this, Fn]() { return Fn.Get().IsEmpty() ? EVisibility::Collapsed : EVisibility::Visible; }));
+			Slot1->GetWidget()->SetVisibility(MakeAttributeLambda([this, Fn]() { return Fn.Get().IsEmpty() ? EVisibility::Collapsed : EVisibility::Visible; }));
 		}
 	}
 
diff --git a/Engine/Source/Editor/UnrealEd/Private/SScalabilitySettings.cpp b/Engine/Source/Editor/UnrealEd/Private/SScalabilitySettings.cpp
--- a/Engine/Source/Editor/UnrealEd/Private/SScalabilitySettings.cpp
+++ b/Engine/Source/Editor/UnrealEd/Private/SScalabilitySettings.cpp
@@ -137,14 +137,14 @@ FReply SScalabilitySettings::OnAutoClicked()
 	return FReply::Handled();
 }
 
-SGridPanel::FSlot& SScalabilitySettings::MakeGridSlot(int32 InCol, int32 InRow, int32 InColSpan /*= 1*/, int32 InRowSpan /*= 1*/)
+SGridPanel::FSlot::FSlotArguments SScalabilitySettings::MakeGridSlot(int32 InCol, int32 InRow, int32 InColSpan /*= 1*/, int32 InRowSpan /*= 1*/)
 {
 	float PaddingH = 2.0f;
 	float PaddingV = InRow == 0 ? 8.0f : 2.0f;
-	return SGridPanel::Slot(InCol, InRow)
+	return MoveTemp(SGridPanel::Slot(InCol, InRow)
 		.Padding(PaddingH, PaddingV)
 		.RowSpan(InRowSpan)
-		.ColumnSpan(InColSpan);
+		.ColumnSpan(InColSpan));
 }
 
 ECheckBoxState SScalabilitySettings::IsMonitoringPerformance() const
diff --git a/Engine/Source/Editor/UnrealEd/Public/SScalabilitySettings.h b/Engine/Source/Editor/UnrealEd/Public/SScalabilitySettings.h
--- a/Engine/Source/Editor/UnrealEd/Public/SScalabilitySettings.h
+++ b/Engine/Source/Editor/UnrealEd/Public/SScalabilitySettings.h
@@ -61,7 +61,7 @@ private:
 	FReply OnAutoClicked();
 
 	// Create a gridslot for the group quality level with all the required formatting
-	SGridPanel::FSlot& MakeGridSlot(int32 InCol, int32 InRow, int32 InColSpan = 1, int32 InRowSpan = 1);
+	SGridPanel::FSlot::FSlotArguments MakeGridSlot(int32 InCol, int32 InRow, int32 InColSpan = 1, int32 InRowSpan = 1);
 
 	/** Called to get the "Show notification" check box state */
 	ECheckBoxState IsMonitoringPerformance() const;
diff --git a/Engine/Source/Runtime/Slate/Private/Widgets/Layout/SGridPanel.cpp b/Engine/Source/Runtime/Slate/Private/Widgets/Layout/SGridPanel.cpp
--- a/Engine/Source/Runtime/Slate/Private/Widgets/Layout/SGridPanel.cpp
+++ b/Engine/Source/Runtime/Slate/Private/Widgets/Layout/SGridPanel.cpp
@@ -6,15 +6,63 @@
 #include "Rendering/DrawElements.h"
 #include "Layout/LayoutUtils.h"
 
+
+void SGridPanel::FSlot::Construct(const FChildren& SlotOwner, FSlotArguments&& InArgs)
+{
+	TBasicLayoutWidgetSlot<FSlot>::Construct(SlotOwner, MoveTemp(InArgs));
+	if (InArgs._Column.IsSet())
+	{
+		ColumnParam = FMath::Max(0, InArgs._Column.GetValue());
+	}
+	if (InArgs._ColumnSpan.IsSet())
+	{
+		ColumnSpanParam = FMath::Max(1, InArgs._ColumnSpan.GetValue());
+	}
+	if (InArgs._Row.IsSet())
+	{
+		RowParam = FMath::Max(0, InArgs._Row.GetValue());
+	}
+	if (InArgs._RowSpan.IsSet())
+	{
+		RowSpanParam = FMath::Max(1, InArgs._RowSpan.GetValue());
+	}
+	if (InArgs._Layer.IsSet())
+	{
+		LayerParam = InArgs._Layer.GetValue();
+	}
+	if (InArgs._Nudge.IsSet())
+	{
+		NudgeParam = InArgs._Nudge.GetValue();
+	}
+}
+
 SGridPanel::SGridPanel()
-: Slots(this)
+	: Slots(this)
 {
 	SetCanTick(false);
 }
 
-SGridPanel::FSlot& SGridPanel::AddSlot( int32 Column, int32 Row, SGridPanel::Layer InLayer )
+SGridPanel::FSlot::FSlotArguments SGridPanel::Slot(int32 Column, int32 Row, SGridPanel::Layer InLayer)
+{
+	return FSlot::FSlotArguments(MakeUnique<FSlot>(Column, Row, InLayer.TheLayer));
+}
+
+SGridPanel::FScopedWidgetSlotArguments SGridPanel::AddSlot( int32 Column, int32 Row, SGridPanel::Layer InLayer )
 {
-	return InsertSlot( new FSlot( Column, Row, InLayer.TheLayer ) );
+	TWeakPtr<SGridPanel> WeakPanel = SharedThis(this);
+
+	TUniquePtr<FSlot> NewSlot = MakeUnique<FSlot>(Column, Row, InLayer.TheLayer);
+	NewSlot->Panel = WeakPanel;
+
+	int32 InsertLocation = FindInsertSlotLocation(NewSlot.Get());
+	return FScopedWidgetSlotArguments{ MoveTemp(NewSlot), this->Slots, InsertLocation,
+		[WeakPanel](const FSlot* InNewSlot, int32 InsertedLocation)
+		{
+			if (TSharedPtr<SGridPanel> Panel = WeakPanel.Pin())
+			{
+				Panel->NotifySlotChanged(InNewSlot);
+			}
+		}};
 }
 
 bool SGridPanel::RemoveSlot(const TSharedRef<SWidget>& SlotWidget)
@@ -35,9 +83,15 @@ void SGridPanel::Construct( const FArguments& InArgs )
 
 	// Populate the slots such that they are sorted by Layer (order preserved within layers)
 	// Also determine the grid size
-	for (int32 SlotIndex = 0; SlotIndex < InArgs.Slots.Num(); ++SlotIndex )
+	Slots.Reserve(InArgs._Slots.Num());
+	for (int32 SlotIndex = 0; SlotIndex < InArgs._Slots.Num(); ++SlotIndex )
 	{
-		InsertSlot( InArgs.Slots[SlotIndex] );
+		int32 SlotLocation = FindInsertSlotLocation(InArgs._Slots[SlotIndex].GetSlot());
+		Slots.InsertSlot(MoveTemp(const_cast<FSlot::FSlotArguments&>(InArgs._Slots[SlotIndex])), SlotLocation);
+
+		FSlot& NewSlot = Slots[SlotLocation];
+		NewSlot.Panel = SharedThis(this);
+		NotifySlotChanged(&NewSlot);
 	}
 
 	ColFillCoefficients = InArgs.ColFillCoefficients;
@@ -74,10 +128,10 @@ int32 SGridPanel::OnPaint( const FPaintArgs& Args, const FGeometry& AllottedGeom
 
 				if (!IsChildWidgetCulled(MyCullingRect, CurWidget))
 				{
-					if (LastGridLayer != CurSlot.LayerParam)
+					if (LastGridLayer != CurSlot.GetLayer())
 					{
 						// We starting a new grid layer group?
-						LastGridLayer = CurSlot.LayerParam;
+						LastGridLayer = CurSlot.GetLayer();
 						// Ensure that everything here is drawn on top of 
 						// previously drawn grid content.
 						LayerId = MaxLayerId + 1;
@@ -183,12 +237,12 @@ void SGridPanel::OnArrangeChildren( const FGeometry& AllottedGeometry, FArranged
 		if ( ArrangedChildren.Accepts(ChildVisibility) )
 		{
 			// Figure out the position of this cell.
-			const FVector2D ThisCellOffset( FinalColumns[CurSlot.ColumnParam], FinalRows[CurSlot.RowParam] );
+			const FVector2D ThisCellOffset( FinalColumns[CurSlot.GetColumn()], FinalRows[CurSlot.GetRow()] );
 			// Figure out the size of this slot; takes row span into account.
 			// We use the properties of partial sums arrays to achieve this.
 			const FVector2D CellSize(
-				FinalColumns[CurSlot.ColumnParam+CurSlot.ColumnSpanParam] - ThisCellOffset.X ,
-				FinalRows[CurSlot.RowParam+CurSlot.RowSpanParam] - ThisCellOffset.Y );
+				FinalColumns[CurSlot.GetColumn()+CurSlot.GetColumnSpan()] - ThisCellOffset.X ,
+				FinalRows[CurSlot.GetRow()+CurSlot.GetRowSpan()] - ThisCellOffset.Y );
 
 			// Do the standard arrangement of elements within a slot
 			// Takes care of alignment and padding.
@@ -199,7 +253,7 @@ void SGridPanel::OnArrangeChildren( const FGeometry& AllottedGeometry, FArranged
 			// Output the result
 			ArrangedChildren.AddWidget( ChildVisibility, AllottedGeometry.MakeChild( 
 				CurSlot.GetWidget(),
-				ThisCellOffset + FVector2D( XAxisResult.Offset, YAxisResult.Offset ) + CurSlot.NudgeParam,
+				ThisCellOffset + FVector2D( XAxisResult.Offset, YAxisResult.Offset ) + CurSlot.GetNudge(),
 				FVector2D(XAxisResult.Size, YAxisResult.Size)
 			));
 		}
@@ -315,46 +369,33 @@ void SGridPanel::DistributeSizeContributions( float SizeContribution, TArray<flo
 }
 
 
-SGridPanel::FSlot& SGridPanel::InsertSlot( SGridPanel::FSlot* InSlot )
+int32 SGridPanel::FindInsertSlotLocation( const FSlot* InSlot )
 {
-	bool bInserted = false;
-
-	InSlot->Panel = SharedThis(this);
-
 	// Insert the slot in the list such that slots are sorted by LayerOffset.
-	for( int32 SlotIndex=0; !bInserted && SlotIndex < Slots.Num(); ++SlotIndex )
+	for( int32 SlotIndex=0; SlotIndex < Slots.Num(); ++SlotIndex )
 	{
-		if ( InSlot->LayerParam < this->Slots[SlotIndex].LayerParam )
+		if ( InSlot->GetLayer() < this->Slots[SlotIndex].GetLayer() )
 		{
-			Slots.Insert( InSlot, SlotIndex );
-			bInserted = true;
+			return SlotIndex;
 		}
 	}
-
-	// We haven't inserted yet, so add to the end of the list.
-	if ( !bInserted )
-	{
-		Slots.Add( InSlot );
-	}
-
-	NotifySlotChanged(InSlot);
-
-	return *InSlot;
+	return INDEX_NONE;
 }
 
-void SGridPanel::NotifySlotChanged(SGridPanel::FSlot* InSlot, bool bSlotLayerChanged /*= false*/)
+
+void SGridPanel::NotifySlotChanged(const FSlot* InSlot, bool bSlotLayerChanged /*= false*/)
 {
 	// Keep the size of the grid up to date.
 	// We need an extra cell at the end for easily figuring out the size across any number of cells
 	// by doing Columns[End] - Columns[Start] or Rows[End] - Rows[Start].
 	// The first Columns[]/Rows[] entry will be 0.
-	const int32 NumColumnsRequiredForThisSlot = InSlot->ColumnParam + InSlot->ColumnSpanParam + 1;
+	const int32 NumColumnsRequiredForThisSlot = InSlot->GetColumn() + InSlot->GetColumnSpan() + 1;
 	if ( NumColumnsRequiredForThisSlot > Columns.Num() )
 	{
 		Columns.AddZeroed( NumColumnsRequiredForThisSlot - Columns.Num() );
 	}
 
-	const int32 NumRowsRequiredForThisSlot = InSlot->RowParam + InSlot->RowSpanParam + 1;
+	const int32 NumRowsRequiredForThisSlot = InSlot->GetRow() + InSlot->GetRowSpan() + 1;
 	if ( NumRowsRequiredForThisSlot > Rows.Num() )
 	{
 		Rows.AddZeroed( NumRowsRequiredForThisSlot - Rows.Num() );
@@ -364,7 +405,7 @@ void SGridPanel::NotifySlotChanged(SGridPanel::FSlot* InSlot, bool bSlotLayerCha
 	{
 		Slots.Sort([](const FSlot& LHS, const FSlot& RHS)
 		{
-			return LHS.LayerParam < RHS.LayerParam;
+			return LHS.GetLayer() < RHS.GetLayer();
 		});
 	}
 
@@ -386,11 +427,11 @@ void SGridPanel::ComputeDesiredCellSizes( TArray<float>& OutColumns, TArray<floa
 
 			// If the slot has a (colspan,rowspan) of (1,1) it will only affect that slot.
 			// For larger spans, the slots size will be evenly distributed across all the affected slots.
-			const FVector2D SizeContribution( SlotDesiredSize.X / CurSlot.ColumnSpanParam, SlotDesiredSize.Y / CurSlot.RowSpanParam );
+			const FVector2D SizeContribution( SlotDesiredSize.X / CurSlot.GetColumnSpan(), SlotDesiredSize.Y / CurSlot.GetRowSpan() );
 
 			// Distribute the size contributions over all the columns and rows that this slot spans
-			DistributeSizeContributions( SizeContribution.X, OutColumns, CurSlot.ColumnParam, CurSlot.ColumnParam + CurSlot.ColumnSpanParam );
-			DistributeSizeContributions( SizeContribution.Y, OutRows, CurSlot.RowParam, CurSlot.RowParam + CurSlot.RowSpanParam );
+			DistributeSizeContributions( SizeContribution.X, OutColumns, CurSlot.GetColumn(), CurSlot.GetColumn() + CurSlot.GetColumnSpan());
+			DistributeSizeContributions( SizeContribution.Y, OutRows, CurSlot.GetRow(), CurSlot.GetRow() + CurSlot.GetRowSpan() );
 		}
 	}
 }
diff --git a/Engine/Source/Runtime/Slate/Public/Widgets/Layout/SGridPanel.h b/Engine/Source/Runtime/Slate/Public/Widgets/Layout/SGridPanel.h
--- a/Engine/Source/Runtime/Slate/Public/Widgets/Layout/SGridPanel.h
+++ b/Engine/Source/Runtime/Slate/Public/Widgets/Layout/SGridPanel.h
@@ -29,11 +29,13 @@ public:
 
 		}
 
-		int TheLayer;
+		int32 TheLayer;
 	};
 
-	class FSlot : public TBasicLayoutWidgetSlot<FSlot>
+	class SLATE_API FSlot : public TBasicLayoutWidgetSlot<FSlot>
 	{
+		friend SGridPanel;
+
 		public:
 			/** Default values for a slot. */
 			FSlot( int32 Column, int32 Row, int32 InLayer )
@@ -47,29 +49,31 @@ public:
 			{
 			}
 
-			/** The panel that contains this slot */
-			TWeakPtr<SGridPanel> Panel;
+			SLATE_SLOT_BEGIN_ARGS(FSlot, TBasicLayoutWidgetSlot<FSlot>)
+				/** Which column in the grid this cell belongs to */
+				SLATE_ARGUMENT(TOptional<int32>, Column)
+				/** How many columns this slot spans over */
+				SLATE_ARGUMENT(TOptional<int32>, ColumnSpan)
+				/** Which row in the grid this cell belongs to */
+				SLATE_ARGUMENT(TOptional<int32>, Row)
+				/** How many rows this this slot spans over */
+				SLATE_ARGUMENT(TOptional<int32>, RowSpan)
+				/** Positive values offset this cell to be hit-tested and drawn on top of others. Default is 0; i.e. no offset. */
+				SLATE_ARGUMENT(TOptional<int32>, Layer)
+				/** Offset this slot's content by some amount; positive values offset to lower right*/
+				SLATE_ARGUMENT(TOptional<FVector2D>, Nudge)
+			SLATE_SLOT_END_ARGS()
+
+			void Construct(const FChildren& SlotOwner, FSlotArguments&& InArgs);
 
+		public:
 			/** Which column in the grid this cell belongs to */
-			int32 ColumnParam;
-
-			/** How many columns this slot spans over */
-			int32 ColumnSpanParam;
-
-			/** Which row in the grid this cell belongs to */
-			int32 RowParam;
-
-			/** How many rows this this slot spans over */
-			int32 RowSpanParam;
-
-			/** Positive values offset this cell to be hit-tested and drawn on top of others. Default is 0; i.e. no offset. */
-			int32 LayerParam;
-
-			/** Offset this slot's content by some amount; positive values offset to lower right*/
-			FVector2D NudgeParam;
+			int32 GetColumn() const
+			{
+				return ColumnParam;
+			}
 
-			/**  */
-			FSlot& Column(int32 Column)
+			void SetColumn(int32 Column)
 			{
 				Column = FMath::Max(0, Column);
 				if (Column != ColumnParam)
@@ -77,23 +81,31 @@ public:
 					ColumnParam = Column;
 					NotifySlotChanged();
 				}
-				return *this;
 			}
 
 			/** How many columns this slot spans over */
-			FSlot& ColumnSpan( int32 ColumnSpan )
+			int32 GetColumnSpan() const
+			{
+				return ColumnSpanParam;
+			}
+
+			void SetColumnSpan(int32 ColumnSpan)
 			{
-				ColumnSpan = FMath::Max(1,ColumnSpan);
+				ColumnSpan = FMath::Max(1, ColumnSpan);
 				if (ColumnSpan != ColumnSpanParam)
 				{
 					ColumnSpanParam = ColumnSpan;
 					NotifySlotChanged();
 				}
-				return *this;
 			}
 
-			/**  */
-			FSlot& Row(int32 Row)
+			/** Which row in the grid this cell belongs to */
+			int32 GetRow() const
+			{
+				return RowParam;
+			}
+
+			void SetRow(int32 Row)
 			{
 				Row = FMath::Max(0, Row);
 				if (Row != RowParam)
@@ -101,11 +113,15 @@ public:
 					RowParam = Row;
 					NotifySlotChanged();
 				}
-				return *this;
 			}
 
 			/** How many rows this this slot spans over */
-			FSlot& RowSpan( int32 RowSpan )
+			int32 GetRowSpan() const
+			{
+				return RowSpanParam;
+			}
+
+			void SetRowSpan(int32 RowSpan)
 			{
 				RowSpan = FMath::Max(1, RowSpan);
 				if (RowSpan != RowSpanParam)
@@ -113,11 +129,15 @@ public:
 					RowSpanParam = RowSpan;
 					NotifySlotChanged();
 				}
-				return *this;
 			}
 
 			/** Positive values offset this cell to be hit-tested and drawn on top of others. Default is 0; i.e. no offset. */
-			FSlot& Layer(int32 Layer)
+			int32 GetLayer() const
+			{
+				return LayerParam;
+			}
+
+			void SetLayer(int32 Layer)
 			{
 				if (Layer != LayerParam)
 				{
@@ -125,16 +145,31 @@ public:
 					const bool bSlotLayerChanged = true;
 					NotifySlotChanged(bSlotLayerChanged);
 				}
-				return *this;
 			}
 
 			/** Offset this slot's content by some amount; positive values offset to lower right */
-			FSlot& Nudge( const FVector2D& Nudge )
+			FVector2D GetNudge() const
+			{
+				return NudgeParam;
+			}
+
+			void SetNudge(const FVector2D& Nudge)
 			{
 				NudgeParam = Nudge;
-				return *this;
+				Invalidate(EInvalidateWidgetReason::Paint);
 			}
 
+		private:
+			/** The panel that contains this slot */
+			TWeakPtr<SGridPanel> Panel;
+
+			int32 ColumnParam;
+			int32 ColumnSpanParam;
+			int32 RowParam;
+			int32 RowSpanParam;
+			int32 LayerParam;
+			FVector2D NudgeParam;
+
 			/** Notify that the slot was changed */
 			FORCEINLINE void NotifySlotChanged(bool bSlotLayerChanged = false)
 			{
@@ -148,17 +183,15 @@ public:
 	/**
 	 * Used by declarative syntax to create a Slot in the specified Column, Row and Layer.
 	 */
-	static FSlot& Slot( int32 Column, int32 Row, Layer InLayer = Layer(0) )
-	{
-		return *(new FSlot( Column, Row, InLayer.TheLayer ));
-	}
+	static FSlot::FSlotArguments Slot( int32 Column, int32 Row, Layer InLayer = Layer(0) );
 
+	using FScopedWidgetSlotArguments = TPanelChildren<FSlot>::FScopedWidgetSlotArguments;
 	/**
 	 * Dynamically add a new slot to the UI at specified Column and Row. Optionally, specify a layer at which this slot should be added.
 	 *
 	 * @return A reference to the newly-added slot
 	 */
-	FSlot& AddSlot( int32 Column, int32 Row, Layer InLayer = Layer(0) );
+	FScopedWidgetSlotArguments AddSlot( int32 Column, int32 Row, Layer InLayer = Layer(0) );
 
 	/**
 	* Removes a slot from this panel which contains the specified SWidget
@@ -172,7 +205,8 @@ public:
 		{
 			_Visibility = EVisibility::SelfHitTestInvisible;
 		}
-		SLATE_SUPPORTS_SLOT( FSlot )
+
+		SLATE_SLOT_ARGUMENT( FSlot, Slots )
 
 		/** Specify a column to stretch instead of sizing to content. */
 		FArguments& FillColumn( int32 ColumnId, const TAttribute<float>& Coefficient )
@@ -254,12 +288,12 @@ private:
 	static void DistributeSizeContributions( float SizeContribution, TArray<float>& DistributeOverMe, int32 StartIndex, int32 UpperBound );
 
 	/**
-	 * Inserts the given slot into the list of Slots based on its LayerParam, such that Slots are sorter by layer.
+	 * Find the index where the given slot should be inserted into the list of Slots based on its LayerParam, such that Slots are sorter by layer.
 	 *
 	 * @param The newly-allocated slot to insert.
-	 * @return A reference to the added slot
+	 * @return The index where the slot should be inserted.
 	 */
-	FSlot& InsertSlot( FSlot* InSlot );	
+	int32 FindInsertSlotLocation( const FSlot* InSlot );
 
 	/** Compute the sizes of columns and rows needed to fit all the slots in this grid. */
 	void ComputeDesiredCellSizes( TArray<float>& OutColumns, TArray<float>& OutRows ) const;
@@ -273,7 +307,7 @@ private:
 	 * @param InSlot The slot that has just changed.
 	 * @param bSlotLayerChanged Whether the slot layer changed.
 	 */
-	void NotifySlotChanged(FSlot* InSlot, bool bSlotLayerChanged = false);
+	void NotifySlotChanged(const FSlot* InSlot, bool bSlotLayerChanged = false);
 
 private:
 
diff --git a/Engine/Source/Runtime/UMG/Private/Components/GridSlot.cpp b/Engine/Source/Runtime/UMG/Private/Components/GridSlot.cpp
--- a/Engine/Source/Runtime/UMG/Private/Components/GridSlot.cpp
+++ b/Engine/Source/Runtime/UMG/Private/Components/GridSlot.cpp
@@ -26,7 +26,8 @@ void UGridSlot::ReleaseSlateResources(bool bReleaseChildren)
 
 void UGridSlot::BuildSlot(TSharedRef<SGridPanel> GridPanel)
 {
-	Slot = &GridPanel->AddSlot(Column, Row, SGridPanel::Layer(Layer))
+	GridPanel->AddSlot(Column, Row, SGridPanel::Layer(Layer))
+		.Expose(Slot)
 		.Padding(Padding)
 		.HAlign(HorizontalAlignment)
 		.VAlign(VerticalAlignment)
@@ -43,7 +44,7 @@ void UGridSlot::SetPadding(FMargin InPadding)
 	Padding = InPadding;
 	if ( Slot )
 	{
-		Slot->Padding(InPadding);
+		Slot->SetPadding(InPadding);
 	}
 }
 
@@ -52,7 +53,7 @@ void UGridSlot::SetRow(int32 InRow)
 	Row = InRow;
 	if ( Slot )
 	{
-		Slot->Row(InRow);
+		Slot->SetRow(InRow);
 	}
 }
 
@@ -61,7 +62,7 @@ void UGridSlot::SetRowSpan(int32 InRowSpan)
 	RowSpan = InRowSpan;
 	if ( Slot )
 	{
-		Slot->RowSpan(InRowSpan);
+		Slot->SetRowSpan(InRowSpan);
 	}
 }
 
@@ -70,7 +71,7 @@ void UGridSlot::SetColumn(int32 InColumn)
 	Column = InColumn;
 	if ( Slot )
 	{
-		Slot->Column(InColumn);
+		Slot->SetColumn(InColumn);
 	}
 }
 
@@ -79,7 +80,7 @@ void UGridSlot::SetColumnSpan(int32 InColumnSpan)
 	ColumnSpan = InColumnSpan;
 	if ( Slot )
 	{
-		Slot->ColumnSpan(InColumnSpan);
+		Slot->SetColumnSpan(InColumnSpan);
 	}
 }
 
@@ -88,7 +89,7 @@ void UGridSlot::SetLayer(int32 InLayer)
 	Layer = InLayer;
 	if (Slot)
 	{
-		Slot->Layer(InLayer);
+		Slot->SetLayer(InLayer);
 	}
 }
 
@@ -97,7 +98,7 @@ void UGridSlot::SetNudge(FVector2D InNudge)
 	Nudge = InNudge;
 	if ( Slot )
 	{
-		Slot->Nudge(InNudge);
+		Slot->SetNudge(InNudge);
 	}
 }
 
