Fix crash when
opening the Sequencer UI. Introduced with CL 16509379
#rnx

diff --git a/Engine/Source/Runtime/Slate/Private/Widgets/Layout/SGridPanel.cpp b/Engine/Source/Runtime/Slate/Private/Widgets/Layout/SGridPanel.cpp
--- a/Engine/Source/Runtime/Slate/Private/Widgets/Layout/SGridPanel.cpp
+++ b/Engine/Source/Runtime/Slate/Private/Widgets/Layout/SGridPanel.cpp
@@ -87,7 +87,14 @@ void SGridPanel::Construct( const FArguments& InArgs )
 	for (int32 SlotIndex = 0; SlotIndex < InArgs._Slots.Num(); ++SlotIndex )
 	{
 		int32 SlotLocation = FindInsertSlotLocation(InArgs._Slots[SlotIndex].GetSlot());
-		Slots.InsertSlot(MoveTemp(const_cast<FSlot::FSlotArguments&>(InArgs._Slots[SlotIndex])), SlotLocation);
+		if (SlotLocation == INDEX_NONE)
+		{
+			SlotLocation = Slots.AddSlot(MoveTemp(const_cast<FSlot::FSlotArguments&>(InArgs._Slots[SlotIndex])));
+		}
+		else
+		{
+			Slots.InsertSlot(MoveTemp(const_cast<FSlot::FSlotArguments&>(InArgs._Slots[SlotIndex])), SlotLocation);
+		}
 
 		FSlot& NewSlot = Slots[SlotLocation];
 		NewSlot.Panel = SharedThis(this);
