Slate: Convert SSegmentedControl FSlot's to the FSlotArgument syntax.

diff --git a/Engine/Source/Runtime/Slate/Public/Widgets/Input/SSegmentedControl.h b/Engine/Source/Runtime/Slate/Public/Widgets/Input/SSegmentedControl.h
--- a/Engine/Source/Runtime/Slate/Public/Widgets/Input/SSegmentedControl.h
+++ b/Engine/Source/Runtime/Slate/Public/Widgets/Input/SSegmentedControl.h
@@ -28,50 +28,89 @@ class SSegmentedControl : public SCompoundWidget
 public:
 
 	/** Stores the per-child info for this panel type */
-	template< typename SlotOptionType >
-	struct FSlot : public TSlotBase<FSlot<SlotOptionType>>, TAlignmentWidgetSlotMixin<FSlot<SlotOptionType>>
+	struct FSlot : public TSlotBase<FSlot>, public TAlignmentWidgetSlotMixin<FSlot>
 	{
-		FSlot(const SlotOptionType& InValue) 
-		: TSlotBase<FSlot<SlotOptionType>>()
-		, TAlignmentWidgetSlotMixin<FSlot<SlotOptionType>>(HAlign_Center, VAlign_Fill)
-		, _Text()
-		, _Tooltip()
-		, _Icon(nullptr)
-		, _Value(InValue)
-		{ }
-
-		FSlot<SlotOptionType>& Text(const TAttribute<FText>& InText)
+		FSlot(const OptionType& InValue)
+			: TSlotBase<FSlot>()
+			, TAlignmentWidgetSlotMixin<FSlot>(HAlign_Center, VAlign_Fill)
+			, _Text()
+			, _Tooltip()
+			, _Icon(nullptr)
+			, _Value(InValue)
+			{ }
+
+		SLATE_SLOT_BEGIN_ARGS_OneMixin(FSlot, TSlotBase<FSlot>, TAlignmentWidgetSlotMixin<FSlot>)
+			SLATE_ATTRIBUTE(FText, Text)
+			SLATE_ATTRIBUTE(FText, ToolTip)
+			SLATE_ATTRIBUTE(const FSlateBrush*, Icon)
+			SLATE_ARGUMENT(TOptional<OptionType>, Value)
+		SLATE_SLOT_END_ARGS()
+
+		void Construct(const FChildren& SlotOwner, FSlotArguments&& InArgs)
 		{
-			_Text = InText;
-			return *this;
+			TSlotBase<FSlot>::Construct(SlotOwner, MoveTemp(InArgs));
+			TAlignmentWidgetSlotMixin<FSlot>::ConstructMixin(SlotOwner, MoveTemp(InArgs));
+			if (InArgs._Text.IsSet())
+			{
+				_Text = MoveTemp(InArgs._Text);
+			}
+			if (InArgs._ToolTip.IsSet())
+			{
+				_Tooltip = MoveTemp(InArgs._ToolTip);
+			}
+			if (InArgs._Icon.IsSet())
+			{
+				_Icon = MoveTemp(InArgs._Icon);
+			}
+			if (InArgs._Value.IsSet())
+			{
+				_Value = MoveTemp(InArgs._Value.GetValue());
+			}
+		}
+
+		void SetText(TAttribute<FText> InText)
+		{
+			_Text = MoveTemp(InText);
+		}
+
+		FText GetText() const
+		{
+			return _Text.Get();
 		}
 
-		FSlot<SlotOptionType>& Icon(const TAttribute<const FSlateBrush*>& InBrush)
+		void SetIcon(TAttribute<const FSlateBrush*> InBrush)
 		{
-			_Icon= InBrush;
-			return *this;
+			_Icon = MoveTemp(InBrush);
 		}
 
-		FSlot<SlotOptionType>& ToolTip(const TAttribute<FText>& InTooltip)
+		const FSlateBrush* GetIcon() const
 		{
-			_Tooltip = InTooltip;
-			return *this;
+			return _Icon.Get();
 		}
 
-	friend class SSegmentedControl<SlotOptionType>;
+		void SetToolTip(TAttribute<FText> InTooltip)
+		{
+			_Tooltip = MoveTemp(InTooltip);
+		}
 
-	protected:
+		FText GetToolTip() const
+		{
+			return _Tooltip.Get();
+		}
 
+	friend SSegmentedControl<OptionType>;
+
+	private:
 		TAttribute<FText> _Text;
 		TAttribute<FText> _Tooltip;
 		TAttribute<const FSlateBrush*> _Icon;
 
-		SlotOptionType _Value;
+		OptionType _Value;
 	};
 
-	static FSlot<OptionType>& Slot(const OptionType& InValue)
+	static typename FSlot::FSlotArguments Slot(const OptionType& InValue)
 	{
-		return *(new FSlot<OptionType>(InValue));
+		return typename FSlot::FSlotArguments(MakeUnique<FSlot>(InValue));
 	}
 
 	DECLARE_DELEGATE_OneParam( FOnValueChanged, OptionType );
@@ -82,7 +121,7 @@ public:
 		, _MaxSegmentsPerLine(0)
 	{}
 		/** Slot type supported by this panel */
-		SLATE_SUPPORTS_SLOT(FSlot<OptionType>)
+		SLATE_SLOT_ARGUMENT(FSlot, Slots)
 
 		/** Styling for this control */
 		SLATE_STYLE_ARGUMENT(FSegmentedControlStyle, Style)
@@ -103,6 +142,10 @@ public:
 		SLATE_ARGUMENT(int32, MaxSegmentsPerLine)
 	SLATE_END_ARGS()
 
+	SSegmentedControl()
+		: Children(this)
+	{}
+
 	void Construct( const FArguments& InArgs )
 	{
 		check(InArgs._Style);
@@ -116,13 +159,7 @@ public:
 		UniformPadding = InArgs._UniformPadding;
 
 		MaxSegmentsPerLine = InArgs._MaxSegmentsPerLine;
-
-		const int32 NumSlots = InArgs.Slots.Num();
-		Children.Reserve(NumSlots);
-		for (int32 SlotIndex = 0; SlotIndex < NumSlots; ++SlotIndex)
-		{
-			Children.Add(InArgs.Slots[SlotIndex]);
-		}
+		Children.AddSlots(MoveTemp(const_cast<TArray<typename FSlot::FSlotArguments>&>(InArgs._Slots)));
 		RebuildChildren();
 	}
 
@@ -135,7 +172,7 @@ public:
 		for ( int32 SlotIndex = 0; SlotIndex < NumSlots; ++SlotIndex )
 		{
 			TSharedRef<SWidget> Child = Children[SlotIndex].GetWidget();
-			FSlot<OptionType>* ChildSlotPtr = &Children[SlotIndex];
+			FSlot* ChildSlotPtr = &Children[SlotIndex];
 			OptionType& ChildValue = ChildSlotPtr->_Value;
 
 			if (Child == SNullWidget::NullWidget)
@@ -148,7 +185,7 @@ public:
 				[
 					SNew(SImage)
 					.ColorAndOpacity(FSlateColor::UseForeground())
-					.Image_Lambda([ChildSlotPtr] () { return ChildSlotPtr->_Icon.Get(); })
+					.Image(ChildSlotPtr->_Icon)
 				]	
 
 				+SHorizontalBox::Slot()
@@ -166,7 +203,7 @@ public:
 					SNew(STextBlock)
 					.TextStyle(TextStyle)
 					.Justification(ETextJustify::Center)
-					.Text_Lambda([ChildSlotPtr] () { return ChildSlotPtr->_Text.Get(); } ) 
+					.Text(ChildSlotPtr->_Text) 
 				];
 			}
 
@@ -197,17 +234,24 @@ public:
 	}
 
 	// Slot Management
-	FSlot<OptionType>& AddSlot(const OptionType& InValue, bool bRebuildChildren = true)
+	using FScopedWidgetSlotArguments = typename TPanelChildren<FSlot>::FScopedWidgetSlotArguments;
+	FScopedWidgetSlotArguments AddSlot(const OptionType& InValue, bool bRebuildChildren = true)
 	{
-		FSlot<OptionType>& NewSlot = *(new FSlot<OptionType>(InValue));
-
-		Children.Add( &NewSlot );
-
 		if (bRebuildChildren)
 		{
-			RebuildChildren();
+			TWeakPtr<SSegmentedControl> AsWeak = SharedThis(this);
+			return FScopedWidgetSlotArguments { MakeUnique<FSlot>(InValue), this->Children, INDEX_NONE, [AsWeak](const FSlot*, int32)
+				{
+					if (TSharedPtr<SSegmentedControl> SharedThis = AsWeak.Pin())
+					{
+						SharedThis->RebuildChildren();
+					}
+				}};
+		}
+		else
+		{
+			return FScopedWidgetSlotArguments{ MakeUnique<FSlot>(InValue), this->Children, INDEX_NONE };
 		}
-		return NewSlot;
 	}
 
 	int32 NumSlots() const
@@ -243,7 +287,7 @@ public:
 
 
 protected:
-	TIndirectArray<FSlot<OptionType>> Children;
+	TPanelChildren<FSlot> Children;
 
 	FOnValueChanged OnValueChanged;
 
diff --git a/Engine/Source/Runtime/SlateCore/Public/Widgets/DeclarativeSyntaxSupport.h b/Engine/Source/Runtime/SlateCore/Public/Widgets/DeclarativeSyntaxSupport.h
--- a/Engine/Source/Runtime/SlateCore/Public/Widgets/DeclarativeSyntaxSupport.h
+++ b/Engine/Source/Runtime/SlateCore/Public/Widgets/DeclarativeSyntaxSupport.h
@@ -294,7 +294,7 @@ template<typename WidgetType> struct TSlateBaseNamedArgs;
  * in order to add support for slots with the construct pattern.
  */
 #define SLATE_SLOT_ARGUMENT( SlotType, SlotName ) \
-		TArray<SlotType::FSlotArguments> _##SlotName; \
+		TArray<typename SlotType::FSlotArguments> _##SlotName; \
 		WidgetArgsType& operator + (typename SlotType::FSlotArguments& SlotToAdd) \
 		{ \
 			_##SlotName.Add( MoveTemp(SlotToAdd) ); \
