Do not use a mem mark to globally catch invalid memory use.  These are now handled by a check on allocation.  This was causing -onethread to fail due to the render thread requiring an empty thread stack.  Long term, it might be better to have a render thread speciallization for the mem stack.

diff --git a/Engine/Source/Runtime/Launch/Private/LaunchEngineLoop.cpp b/Engine/Source/Runtime/Launch/Private/LaunchEngineLoop.cpp
--- a/Engine/Source/Runtime/Launch/Private/LaunchEngineLoop.cpp
+++ b/Engine/Source/Runtime/Launch/Private/LaunchEngineLoop.cpp
@@ -4853,9 +4853,6 @@ void FEngineLoop::Tick()
 {
 	SCOPE_STALL_COUNTER(FEngineLoop::Tick, 2.0);
 
-	// make sure to catch any FMemStack uses outside of UWorld::Tick
-	FMemMark MemStackMark(FMemStack::Get());
-
 #if !UE_BUILD_SHIPPING && !UE_BUILD_TEST && MALLOC_GT_HOOKS
 	FScopedSampleMallocChurn ChurnTracker;
 #endif
diff --git a/Engine/Source/Runtime/Slate/Private/Framework/Application/SlateApplication.cpp b/Engine/Source/Runtime/Slate/Private/Framework/Application/SlateApplication.cpp
--- a/Engine/Source/Runtime/Slate/Private/Framework/Application/SlateApplication.cpp
+++ b/Engine/Source/Runtime/Slate/Private/Framework/Application/SlateApplication.cpp
@@ -1236,8 +1236,6 @@ void FSlateApplication::PrivateDrawWindows( TSharedPtr<SWindow> DrawOnlyThisWind
 	// with the loading thread.
 	FScopeLock ScopeLock(Renderer->GetResourceCriticalSection());
 
-	FMemMark Mark(FMemStack::Get());
-
 	FWidgetPath WidgetsToVisualizeUnderCursor;
 	
 #if SLATE_HAS_WIDGET_REFLECTOR
