Slate: Convert SSynthTooltip, SUniformWrapPanel, and other panels FSlot to the FSlotArgument syntax.

diff --git a/Engine/Plugins/Runtime/Synthesis/Source/Synthesis/Private/UI/SSynthTooltip.cpp b/Engine/Plugins/Runtime/Synthesis/Source/Synthesis/Private/UI/SSynthTooltip.cpp
--- a/Engine/Plugins/Runtime/Synthesis/Source/Synthesis/Private/UI/SSynthTooltip.cpp
+++ b/Engine/Plugins/Runtime/Synthesis/Source/Synthesis/Private/UI/SSynthTooltip.cpp
@@ -30,11 +30,8 @@ void SSynthTooltip::Construct(const FArguments& InArgs)
 			TooltipText.ToSharedRef()
 		]);
 
-	const int32 NumSlots = InArgs.Slots.Num();
-	for (int32 SlotIndex = 0; SlotIndex < NumSlots; ++SlotIndex)
-	{
-		Children.Add(InArgs.Slots[SlotIndex]);
-	}
+
+	Children.AddSlots(MoveTemp(const_cast<TArray<SOverlay::FOverlaySlot::FSlotArguments>&>(InArgs._Slots)));
 
 	bIsVisible = false;
 }
diff --git a/Engine/Plugins/Runtime/Synthesis/Source/Synthesis/Public/UI/SSynthTooltip.h b/Engine/Plugins/Runtime/Synthesis/Source/Synthesis/Public/UI/SSynthTooltip.h
--- a/Engine/Plugins/Runtime/Synthesis/Source/Synthesis/Public/UI/SSynthTooltip.h
+++ b/Engine/Plugins/Runtime/Synthesis/Source/Synthesis/Public/UI/SSynthTooltip.h
@@ -20,7 +20,7 @@ class SSynthTooltip : public SOverlay
 public:
 	SLATE_BEGIN_ARGS(SSynthTooltip)
 	{}
-	SLATE_SUPPORTS_SLOT(SOverlay::FOverlaySlot)
+		SLATE_SLOT_ARGUMENT(SOverlay::FOverlaySlot, Slots)
 	SLATE_END_ARGS()
 
 	~SSynthTooltip();
diff --git a/Engine/Source/Runtime/Slate/Private/Widgets/Layout/SUniformWrapPanel.cpp b/Engine/Source/Runtime/Slate/Private/Widgets/Layout/SUniformWrapPanel.cpp
--- a/Engine/Source/Runtime/Slate/Private/Widgets/Layout/SUniformWrapPanel.cpp
+++ b/Engine/Source/Runtime/Slate/Private/Widgets/Layout/SUniformWrapPanel.cpp
@@ -22,12 +22,7 @@ void SUniformWrapPanel::Construct( const FArguments& InArgs )
 	EvenRowDistribution =  InArgs._EvenRowDistribution.Get();
 	HAlign = InArgs._HAlign.Get();
 
-	Children.Reserve( InArgs.Slots.Num() );
-	for (int32 ChildIndex=0; ChildIndex < InArgs.Slots.Num(); ChildIndex++)
-	{
-		FSlot* ChildSlot = InArgs.Slots[ChildIndex];
-		Children.Add( ChildSlot );
-	}
+	Children.AddSlots(MoveTemp(const_cast<TArray<FSlot::FSlotArguments>&>(InArgs._Slots)));
 }
 
 void SUniformWrapPanel::OnArrangeChildren( const FGeometry& AllottedGeometry, FArrangedChildren& ArrangedChildren ) const
@@ -214,13 +209,9 @@ void SUniformWrapPanel::SetEvenRowDistribution(TAttribute<bool> InEvenRowDistrib
 	EvenRowDistribution = InEvenRowDistribution;
 }
 
-SUniformWrapPanel::FSlot& SUniformWrapPanel::AddSlot()
+SUniformWrapPanel::FScopedWidgetSlotArguments SUniformWrapPanel::AddSlot()
 {
-	FSlot& NewSlot = *(new FSlot());
-
-	Children.Add( &NewSlot );
-
-	return NewSlot;
+	return FScopedWidgetSlotArguments{ MakeUnique<FSlot>(), this->Children, INDEX_NONE };
 }
 
 bool SUniformWrapPanel::RemoveSlot( const TSharedRef<SWidget>& SlotWidget )
diff --git a/Engine/Source/Runtime/Slate/Public/Widgets/Layout/SUniformWrapPanel.h b/Engine/Source/Runtime/Slate/Public/Widgets/Layout/SUniformWrapPanel.h
--- a/Engine/Source/Runtime/Slate/Public/Widgets/Layout/SUniformWrapPanel.h
+++ b/Engine/Source/Runtime/Slate/Public/Widgets/Layout/SUniformWrapPanel.h
@@ -22,10 +22,17 @@ public:
 	struct FSlot : public TSlotBase<FSlot>, public TAlignmentWidgetSlotMixin<FSlot>
 	{
 		FSlot()
-		: TSlotBase<FSlot>()
-		, TAlignmentWidgetSlotMixin<FSlot>(HAlign_Fill, VAlign_Fill)
-		{
+			: TSlotBase<FSlot>()
+			, TAlignmentWidgetSlotMixin<FSlot>(HAlign_Fill, VAlign_Fill)
+		{ }
+
+		SLATE_SLOT_BEGIN_ARGS_OneMixin(FSlot, TSlotBase<FSlot>, TAlignmentWidgetSlotMixin<FSlot>)
+		SLATE_SLOT_END_ARGS()
 
+		void Construct(const FChildren& SlotOwner, FSlotArguments&& InArgs)
+		{
+			TSlotBase<FSlot>::Construct(SlotOwner, MoveTemp(InArgs));
+			TAlignmentWidgetSlotMixin<FSlot>::ConstructMixin(SlotOwner, MoveTemp(InArgs));
 		}
 	};
 
@@ -34,9 +41,9 @@ public:
 	/**
 	 * Used by declarative syntax to create a Slot.
 	 */
-	static FSlot& Slot()
+	static FSlot::FSlotArguments Slot()
 	{
-		return *(new FSlot());
+		return FSlot::FSlotArguments(MakeUnique<FSlot>());
 	}
 
 	SLATE_BEGIN_ARGS( SUniformWrapPanel )
@@ -52,7 +59,7 @@ public:
 		}
 
 		/** Slot type supported by this panel */
-		SLATE_SUPPORTS_SLOT(FSlot)
+		SLATE_SLOT_ARGUMENT(FSlot, Slots)
 		
 		/** Padding given to each slot */
 		SLATE_ATTRIBUTE(FMargin, SlotPadding)
@@ -107,7 +114,8 @@ public:
 	void SetEvenRowDistribution(TAttribute<bool> InEvenRowDistribution);
 	bool GetEvenRowDistribution() { return EvenRowDistribution.Get(); }
 
-	FSlot& AddSlot();
+	using FScopedWidgetSlotArguments = TPanelChildren<FSlot>::FScopedWidgetSlotArguments;
+	FScopedWidgetSlotArguments AddSlot();
 	
 	/**
 	 * Removes a slot from this panel which contains the specified SWidget
diff --git a/Engine/Source/Runtime/SlateCore/Public/Layout/Children.h b/Engine/Source/Runtime/SlateCore/Public/Layout/Children.h
--- a/Engine/Source/Runtime/SlateCore/Public/Layout/Children.h
+++ b/Engine/Source/Runtime/SlateCore/Public/Layout/Children.h
@@ -699,7 +699,6 @@ public:
 		{
 			AddSlot(MoveTemp(Arg));
 		}
-		SlotArguments.Reset();
 	}
 
 	void RemoveAt( int32 Index )
