Slate: Update Hittestgrid and some FWidgetPath's function to takes a SWidget pointer instead of a SWidget shared ptr to reduce the number shared ptr copy. FWidgetAndPointer uses a TOptional instead of a TSharedPtr.

diff --git a/Engine/Plugins/Experimental/CommonUI/Source/CommonUI/Private/Input/CommonAnalogCursor.cpp b/Engine/Plugins/Experimental/CommonUI/Source/CommonUI/Private/Input/CommonAnalogCursor.cpp
--- a/Engine/Plugins/Experimental/CommonUI/Source/CommonUI/Private/Input/CommonAnalogCursor.cpp
+++ b/Engine/Plugins/Experimental/CommonUI/Source/CommonUI/Private/Input/CommonAnalogCursor.cpp
@@ -346,7 +346,7 @@ bool FCommonAnalogCursor::HandleMouseButtonDownEvent(FSlateApplication& SlateApp
 				//		Mousing over to a web browser and clicking the back button, for example, will shift OS focus to the browser
 				//			Within the Slate app, clicking back will transfer focus just like LMB would
 				const FWidgetPath WidgetsUnderCursor = SlateApp.LocateWindowUnderMouse(PointerEvent.GetScreenSpacePosition(), SlateApp.GetInteractiveTopLevelWindows());
-				if (WidgetsUnderCursor.ContainsWidget(ViewportWidget.ToSharedRef()))
+				if (WidgetsUnderCursor.ContainsWidget(ViewportWidget.Get()))
 				{
 					//@todo DanH: Do we need to go through the whole process here of generating a false key down event?
 					//		All we really want to do here is allow UI input bindings to mouse buttons right?
diff --git a/Engine/Plugins/Experimental/DigitalMultiplex/DMXPixelMapping/Source/DMXPixelMappingEditor/Private/Views/SDMXPixelMappingDesignerView.cpp b/Engine/Plugins/Experimental/DigitalMultiplex/DMXPixelMapping/Source/DMXPixelMappingEditor/Private/Views/SDMXPixelMappingDesignerView.cpp
--- a/Engine/Plugins/Experimental/DigitalMultiplex/DMXPixelMapping/Source/DMXPixelMappingEditor/Private/Views/SDMXPixelMappingDesignerView.cpp
+++ b/Engine/Plugins/Experimental/DigitalMultiplex/DMXPixelMapping/Source/DMXPixelMappingEditor/Private/Views/SDMXPixelMappingDesignerView.cpp
@@ -682,7 +682,7 @@ void SDMXPixelMappingDesignerView::PopulateWidgetGeometryCache_Loop(FArrangedWid
 
 	if (bIncludeInHitTestGrid)
 	{
-		HittestGrid->AddWidget(CurrentWidget.Widget, 0, 0, FSlateInvalidationWidgetSortOrder());
+		HittestGrid->AddWidget(&(CurrentWidget.Widget.Get()), 0, 0, FSlateInvalidationWidgetSortOrder());
 	}
 
 	FArrangedChildren ArrangedChildren(EVisibility::All);
diff --git a/Engine/Plugins/Experimental/VirtualProductionUtilities/Source/VPUtilities/Private/VPFullScreenUserWidget.cpp b/Engine/Plugins/Experimental/VirtualProductionUtilities/Source/VPUtilities/Private/VPFullScreenUserWidget.cpp
--- a/Engine/Plugins/Experimental/VirtualProductionUtilities/Source/VPUtilities/Private/VPFullScreenUserWidget.cpp
+++ b/Engine/Plugins/Experimental/VirtualProductionUtilities/Source/VPUtilities/Private/VPFullScreenUserWidget.cpp
@@ -78,15 +78,13 @@ public:
 			float CursorRadius = 0.f;
 			ArrangedWidgets = SlateWindowPin->GetHittestGrid().GetBubblePath(LocalMouseCoordinate, CursorRadius, bIgnoreEnabledStatus);
 
-			TSharedRef<FVirtualPointerPosition> VirtualMouseCoordinate = MakeShared<FVirtualPointerPosition>();
-			VirtualMouseCoordinate->CurrentCursorPosition = LocalMouseCoordinate;
-			VirtualMouseCoordinate->LastCursorPosition = LastLocalHitLocation;
+			FVirtualPointerPosition VirtualMouseCoordinate(LocalMouseCoordinate, LastLocalHitLocation);
 
 			LastLocalHitLocation = LocalMouseCoordinate;
 
 			for (FWidgetAndPointer& ArrangedWidget : ArrangedWidgets)
 			{
-				ArrangedWidget.PointerPosition = VirtualMouseCoordinate;
+				ArrangedWidget.SetPointerPosition(VirtualMouseCoordinate);
 			}
 		}
 
@@ -103,9 +101,9 @@ public:
 		}
 	}
 
-	virtual TSharedPtr<struct FVirtualPointerPosition> TranslateMouseCoordinateForCustomHitTestChild(const TSharedRef<SWidget>& ChildWidget, const FGeometry& ViewportGeometry, const FVector2D& ScreenSpaceMouseCoordinate, const FVector2D& LastScreenSpaceMouseCoordinate) const override
+	virtual TOptional<FVirtualPointerPosition> TranslateMouseCoordinateForCustomHitTestChild(const SWidget& ChildWidget, const FGeometry& MyGeometry, const FVector2D ScreenSpaceMouseCoordinate, const FVector2D LastScreenSpaceMouseCoordinate) const override
 	{
-		return nullptr;
+		return TOptional<FVirtualPointerPosition>();
 	}
 
 	void SetWidgetDrawSize(FIntPoint NewWidgetDrawSize)
diff --git a/Engine/Plugins/Runtime/nDisplay/Source/DisplayClusterConfigurator/Private/DisplayClusterConfiguratorBlueprintEditor.cpp b/Engine/Plugins/Runtime/nDisplay/Source/DisplayClusterConfigurator/Private/DisplayClusterConfiguratorBlueprintEditor.cpp
--- a/Engine/Plugins/Runtime/nDisplay/Source/DisplayClusterConfigurator/Private/DisplayClusterConfiguratorBlueprintEditor.cpp
+++ b/Engine/Plugins/Runtime/nDisplay/Source/DisplayClusterConfigurator/Private/DisplayClusterConfiguratorBlueprintEditor.cpp
@@ -730,7 +730,7 @@ void FDisplayClusterConfiguratorBlueprintEditor::OnFocusChanged(const FFocusEven
 		// If the SCSEditor or the Cluster View are being focused, update the property inspector to show the items currently selected in the respective tree view.
 		// This ensures that if any item as selected as part of an ancillary selection, the user can still "select" it and have its properties show up in the
 		// inspector, since already selected items don't fire a OnSelectionChanged event where we would normally update the property inspector
-		if (NewFocusedWidgetPath.ContainsWidget(SubobjectEditor.ToSharedRef()))
+		if (NewFocusedWidgetPath.ContainsWidget(SubobjectEditor.Get()))
 		{
 			TArray<FSubobjectEditorTreeNodePtrType> SelectedNodes = SubobjectEditor->GetSelectedNodes();
 			if (Inspector.IsValid())
@@ -777,7 +777,7 @@ void FDisplayClusterConfiguratorBlueprintEditor::OnFocusChanged(const FFocusEven
 				Inspector->ShowDetailsForObjects(InspectorObjects, Options);
 			}
 		}
-		else if (NewFocusedWidgetPath.ContainsWidget(ViewCluster->GetWidget()))
+		else if (NewFocusedWidgetPath.ContainsWidget(&ViewCluster->GetWidget().Get()))
 		{
 			TArray<UObject*> SelectedClusterObjects;
 			ViewCluster->GetSelectedObjects(SelectedClusterObjects);
diff --git a/Engine/Source/Developer/SlateReflector/Private/SlateNavigationEventSimulator.cpp b/Engine/Source/Developer/SlateReflector/Private/SlateNavigationEventSimulator.cpp
--- a/Engine/Source/Developer/SlateReflector/Private/SlateNavigationEventSimulator.cpp
+++ b/Engine/Source/Developer/SlateReflector/Private/SlateNavigationEventSimulator.cpp
@@ -349,7 +349,7 @@ FSlateNavigationEventSimulator::FSimulationResult FSlateNavigationEventSimulator
 			{
 				if (TSharedPtr<SWidget> ViewportWidget = Viewport->GetWidget().Pin())
 				{
-					if (NavigationSource.ContainsWidget(ViewportWidget.ToSharedRef()))
+					if (NavigationSource.ContainsWidget(ViewportWidget.Get()))
 					{
 						TOptional<FWidgetPath> ViewportResult = OnViewportHandleNavigation.Execute(NavigationEvent.GetUserIndex(), Result.WidgetThatShouldReceivedFocus);
 						Result.bHandledByViewport = ViewportResult.IsSet();
diff --git a/Engine/Source/Developer/SlateReflector/Private/Widgets/SWidgetEventLog.cpp b/Engine/Source/Developer/SlateReflector/Private/Widgets/SWidgetEventLog.cpp
--- a/Engine/Source/Developer/SlateReflector/Private/Widgets/SWidgetEventLog.cpp
+++ b/Engine/Source/Developer/SlateReflector/Private/Widgets/SWidgetEventLog.cpp
@@ -520,7 +520,7 @@ bool SWidgetEventLog::IsInsideWidgetReflector(const FWidgetPath& WidgetPath) con
 {
 	if (TSharedPtr<const SWidget> PinnedReflectorWidget = ReflectorWidget.Pin())
 	{
-		return WidgetPath.ContainsWidget(PinnedReflectorWidget.ToSharedRef());
+		return WidgetPath.ContainsWidget(PinnedReflectorWidget.Get());
 	}
 	return false;
 }
@@ -529,7 +529,7 @@ bool SWidgetEventLog::IsInsideWidgetReflector(const FWeakWidgetPath& WidgetPath)
 {
 	if (TSharedPtr<const SWidget> PinnedReflectorWidget = ReflectorWidget.Pin())
 	{
-		return WidgetPath.ContainsWidget(PinnedReflectorWidget.ToSharedRef());
+		return WidgetPath.ContainsWidget(PinnedReflectorWidget.Get());
 	}
 	return false;
 }
diff --git a/Engine/Source/Developer/SlateReflector/Private/Widgets/SWidgetReflector.cpp b/Engine/Source/Developer/SlateReflector/Private/Widgets/SWidgetReflector.cpp
--- a/Engine/Source/Developer/SlateReflector/Private/Widgets/SWidgetReflector.cpp
+++ b/Engine/Source/Developer/SlateReflector/Private/Widgets/SWidgetReflector.cpp
@@ -1127,7 +1127,7 @@ int32 SWidgetReflector::Visualize( const FWidgetPath& InWidgetsToVisualize, FSla
 		}
 	}
 
-	const bool bAttemptingToVisualizeReflector = InWidgetsToVisualize.ContainsWidget(ReflectorTree.ToSharedRef());
+	const bool bAttemptingToVisualizeReflector = InWidgetsToVisualize.ContainsWidget(ReflectorTree.Get());
 
 	if (PickingMode == EWidgetPickingMode::Drawable)
 	{
diff --git a/Engine/Source/Editor/Sequencer/Private/SSequencer.cpp b/Engine/Source/Editor/Sequencer/Private/SSequencer.cpp
--- a/Engine/Source/Editor/Sequencer/Private/SSequencer.cpp
+++ b/Engine/Source/Editor/Sequencer/Private/SSequencer.cpp
@@ -3211,7 +3211,7 @@ FReply SSequencer::OnKeyDown( const FGeometry& MyGeometry, const FKeyEvent& InKe
 
 void SSequencer::OnFocusChanging( const FWeakWidgetPath& PreviousFocusPath, const FWidgetPath& NewWidgetPath, const FFocusEvent& InFocusEvent )
 {
-	if (NewWidgetPath.ContainsWidget(AsShared()))
+	if (NewWidgetPath.ContainsWidget(this))
 	{
 		OnReceivedFocus.ExecuteIfBound();
 	}
diff --git a/Engine/Source/Editor/StatusBar/Private/SStatusBar.cpp b/Engine/Source/Editor/StatusBar/Private/SStatusBar.cpp
--- a/Engine/Source/Editor/StatusBar/Private/SStatusBar.cpp
+++ b/Engine/Source/Editor/StatusBar/Private/SStatusBar.cpp
@@ -771,7 +771,7 @@ void SStatusBar::OnGlobalFocusChanging(const FFocusEvent& FocusEvent, const FWea
 
 		bool bShouldDismiss = false;
 		// Do not close due to slow tasks as those opening send window activation events
-		if (!GIsSlowTask && !FSlateApplication::Get().GetActiveModalWindow().IsValid() && ActiveDrawerOverlayContent.IsValid() && (!NewFocusedWidgetPath.ContainsWidget(ActiveDrawerOverlayContent.ToSharedRef()) && !NewFocusedWidgetPath.ContainsWidget(ThisWidget)))
+		if (!GIsSlowTask && !FSlateApplication::Get().GetActiveModalWindow().IsValid() && ActiveDrawerOverlayContent.IsValid() && (!NewFocusedWidgetPath.ContainsWidget(ActiveDrawerOverlayContent.Get()) && !NewFocusedWidgetPath.ContainsWidget(this)))
 		{
 			if (TSharedPtr<SWidget> MenuHost = FSlateApplication::Get().GetMenuHostWidget())
 			{
@@ -779,7 +779,7 @@ void SStatusBar::OnGlobalFocusChanging(const FFocusEvent& FocusEvent, const FWea
 
 				// See if the menu being opened is part of the content browser path and if so the menu should not be dismissed
 				FSlateApplication::Get().GeneratePathToWidgetUnchecked(MenuHost.ToSharedRef(), MenuHostPath, EVisibility::All);
-				if (!MenuHostPath.ContainsWidget(ActiveDrawerOverlayContent.ToSharedRef()))
+				if (!MenuHostPath.ContainsWidget(ActiveDrawerOverlayContent.Get()))
 				{
 					bShouldDismiss = true;
 				}
diff --git a/Engine/Source/Editor/UMGEditor/Private/Designer/SDesignerView.cpp b/Engine/Source/Editor/UMGEditor/Private/Designer/SDesignerView.cpp
--- a/Engine/Source/Editor/UMGEditor/Private/Designer/SDesignerView.cpp
+++ b/Engine/Source/Editor/UMGEditor/Private/Designer/SDesignerView.cpp
@@ -2151,7 +2151,7 @@ void SDesignerView::PopulateWidgetGeometryCache_Loop(FArrangedWidget& CurrentWid
 
 	if (bIncludeInHitTestGrid)
 	{
-		DesignerHittestGrid->AddWidget(CurrentWidget.Widget, 0, 0, FSlateInvalidationWidgetSortOrder());
+		DesignerHittestGrid->AddWidget(&(CurrentWidget.Widget.Get()), 0, 0, FSlateInvalidationWidgetSortOrder());
 	}
 
 	FArrangedChildren ArrangedChildren(EVisibility::All);
diff --git a/Engine/Source/Runtime/Slate/Private/Framework/Application/MenuStack.cpp b/Engine/Source/Runtime/Slate/Private/Framework/Application/MenuStack.cpp
--- a/Engine/Source/Runtime/Slate/Private/Framework/Application/MenuStack.cpp
+++ b/Engine/Source/Runtime/Slate/Private/Framework/Application/MenuStack.cpp
@@ -134,7 +134,7 @@ namespace MenuStackInternal
 		virtual void OnFocusChanging(const FWeakWidgetPath& PreviousFocusPath, const FWidgetPath& NewWidgetPath, const FFocusEvent& InFocusEvent) override
 		{
 			// if focus changed and this menu content had focus (or one of its children did) then inform the stack via the OnMenuLostFocus event
-			if (OnMenuLostFocus.IsBound() && PreviousFocusPath.ContainsWidget(AsShared()))
+			if (OnMenuLostFocus.IsBound() && PreviousFocusPath.ContainsWidget(this))
 			{
 				return OnMenuLostFocus.Execute(NewWidgetPath);
 			}
@@ -633,7 +633,7 @@ void FMenuStack::SetHostPath(const FWidgetPath& InOwnerPath)
 
 	if ( HostPopupLayer.IsValid() )
 	{
-		if ( !InOwnerPath.ContainsWidget(HostPopupLayer->GetHost()) )
+		if ( !InOwnerPath.ContainsWidget(&HostPopupLayer->GetHost().Get()) )
 		{
 			HostPopupLayer->Remove();
 			HostPopupLayer.Reset();
diff --git a/Engine/Source/Runtime/Slate/Private/Framework/Application/SlateApplication.cpp b/Engine/Source/Runtime/Slate/Private/Framework/Application/SlateApplication.cpp
--- a/Engine/Source/Runtime/Slate/Private/Framework/Application/SlateApplication.cpp
+++ b/Engine/Source/Runtime/Slate/Private/Framework/Application/SlateApplication.cpp
@@ -241,7 +241,7 @@ public:
 		FWidgetAndPointer GetWidget() const
 		{
 			const int32 WidgetIndex = RoutingPath.Widgets.Num()-1;
-			return FWidgetAndPointer(RoutingPath.Widgets[WidgetIndex], RoutingPath.VirtualPointerPositions[WidgetIndex]);
+			return FWidgetAndPointer(RoutingPath.Widgets[WidgetIndex], RoutingPath.GetVirtualPointerPosition(WidgetIndex));
 		}
 
 		const FWidgetPath& GetRoutingPath() const
@@ -282,7 +282,7 @@ public:
 
 		FWidgetAndPointer GetWidget() const
 		{
-			return FWidgetAndPointer(RoutingPath.Widgets[WidgetIndex], RoutingPath.VirtualPointerPositions[WidgetIndex]);
+			return FWidgetAndPointer(RoutingPath.Widgets[WidgetIndex], RoutingPath.GetVirtualPointerPosition(WidgetIndex));
 		}
 		
 		const FWidgetPath& GetRoutingPath() const
@@ -323,7 +323,7 @@ public:
 
 		FWidgetAndPointer GetWidget() const
 		{
-			return FWidgetAndPointer(RoutingPath.Widgets[WidgetIndex], RoutingPath.VirtualPointerPositions[WidgetIndex]);
+			return FWidgetAndPointer(RoutingPath.Widgets[WidgetIndex], RoutingPath.GetVirtualPointerPosition(WidgetIndex));
 		}
 
 		const FWidgetPath& GetRoutingPath() const
@@ -377,7 +377,7 @@ public:
 #if PLATFORM_COMPILER_HAS_IF_CONSTEXPR
 			if constexpr (Translate<EventType>::TranslationNeeded())
 			{
-				const EventType TranslatedEvent = Translate<EventType>::PointerEvent(ArrangedWidget.PointerPosition, EventCopy);
+				const EventType TranslatedEvent = Translate<EventType>::PointerEvent(ArrangedWidget, EventCopy);
 				Reply = Lambda(ArrangedWidget, TranslatedEvent).SetHandler(ArrangedWidget.Widget);
 				ProcessReply(ThisApplication, RoutingPath, Reply, WidgetsUnderCursor, &TranslatedEvent);
 			}
@@ -387,7 +387,7 @@ public:
 				ProcessReply(ThisApplication, RoutingPath, Reply, WidgetsUnderCursor, &EventCopy);
 			}
 #else
-			const EventType TranslatedEvent = Translate<EventType>::PointerEvent(ArrangedWidget.PointerPosition, EventCopy);
+			const EventType TranslatedEvent = Translate<EventType>::PointerEvent(ArrangedWidget, EventCopy);
 			Reply = Lambda(ArrangedWidget, TranslatedEvent).SetHandler(ArrangedWidget.Widget);
 			ProcessReply(ThisApplication, RoutingPath, Reply, WidgetsUnderCursor, &TranslatedEvent);
 #endif
@@ -418,7 +418,7 @@ public:
 	struct Translate
 	{
 		static constexpr bool TranslationNeeded() { return false; }
-		static EventType PointerEvent( const TSharedPtr<FVirtualPointerPosition>& InPosition, const EventType& InEvent )
+		static EventType PointerEvent( const FWidgetAndPointer& InPosition, const EventType& InEvent )
 		{
 			// Most events do not do any coordinate translation.
 			return InEvent;
@@ -436,16 +436,16 @@ template<>
 struct FEventRouter::Translate<FPointerEvent>
 {
 	static constexpr bool TranslationNeeded() { return true; }
-	static  FPointerEvent PointerEvent( const TSharedPtr<FVirtualPointerPosition>& InPosition, const FPointerEvent& InEvent )
+	static  FPointerEvent PointerEvent( const FWidgetAndPointer& InPosition, const FPointerEvent& InEvent )
 	{
 		// Pointer events are translated into the virtual window space. For 3D Widget Components this means
-		if ( !InPosition.IsValid() )
+		if ( !InPosition.GetPointerPosition().IsSet() )
 		{
 			return InEvent;
 		}
 		else
 		{
-			return FPointerEvent::MakeTranslatedEvent<FPointerEvent>( InEvent, *InPosition );
+			return FPointerEvent::MakeTranslatedEvent<FPointerEvent>( InEvent, InPosition.GetPointerPosition().GetValue() );
 		}
 	}
 };
@@ -2841,10 +2841,10 @@ void FSlateApplication::SetAllUserFocus(const FWidgetPath& InFocusPath, const EF
 
 void FSlateApplication::SetAllUserFocusAllowingDescendantFocus(const FWidgetPath& InFocusPath, const EFocusCause InCause)
 {
-	TSharedRef<SWidget> FocusWidget = InFocusPath.Widgets.Last().Widget;
+	const TSharedRef<SWidget>& FocusWidget = InFocusPath.Widgets.Last().Widget;
 
 	ForEachUser([&] (FSlateUser& User) {
-		if (!User.GetWeakFocusPath().ContainsWidget(FocusWidget))
+		if (!User.GetWeakFocusPath().ContainsWidget(&FocusWidget.Get()))
 		{
 			SetUserFocus(User, InFocusPath, InCause);
 		}
@@ -5045,7 +5045,7 @@ bool FSlateApplication::RoutePointerMoveEvent(const FWidgetPath& WidgetsUnderPoi
 						else
 						{
 							// Only fire mouse leave events for widgets inside the captor path, or whoever if there is no captor path.
-							if (MouseCaptorPath.IsValid() == false || MouseCaptorPath.ContainsWidget(SomeWidgetPreviouslyUnderCursor.ToSharedRef()))
+							if (MouseCaptorPath.IsValid() == false || MouseCaptorPath.ContainsWidget(SomeWidgetPreviouslyUnderCursor.Get()))
 							{
 								// Note that the event's pointer position is not translated.
 								SomeWidgetPreviouslyUnderCursor->OnMouseLeave(PointerEvent);
@@ -5067,9 +5067,9 @@ bool FSlateApplication::RoutePointerMoveEvent(const FWidgetPath& WidgetsUnderPoi
 
 		FEventRouter::Route<FNoReply>(this, FEventRouter::FBubblePolicy(WidgetsUnderPointer), PointerEvent, [&MouseCaptorPath, &LastWidgetsUnderPointer](const FArrangedWidget& WidgetUnderCursor, const FPointerEvent& Event)
 			{
-				if (!LastWidgetsUnderPointer.ContainsWidget(WidgetUnderCursor.Widget))
+				if (!LastWidgetsUnderPointer.ContainsWidget(WidgetUnderCursor.GetWidgetPtr()))
 				{
-					if (MouseCaptorPath.ContainsWidget(WidgetUnderCursor.Widget))
+					if (MouseCaptorPath.ContainsWidget(WidgetUnderCursor.GetWidgetPtr()))
 					{
 						WidgetUnderCursor.Widget->OnMouseEnter(WidgetUnderCursor.Geometry, Event);
 #if WITH_SLATE_DEBUGGING
@@ -5139,7 +5139,7 @@ bool FSlateApplication::RoutePointerMoveEvent(const FWidgetPath& WidgetsUnderPoi
 			FDragDropEvent DragDropEvent(PointerEvent, DragDropContent);
 			FEventRouter::Route<FNoReply>(this, FEventRouter::FBubblePolicy(WidgetsUnderPointer), DragDropEvent, [&LastWidgetsUnderPointer](const FArrangedWidget& WidgetUnderCursor, const FDragDropEvent& InDragDropEvent)
 			{
-				if (!LastWidgetsUnderPointer.ContainsWidget(WidgetUnderCursor.Widget))
+				if (!LastWidgetsUnderPointer.ContainsWidget(WidgetUnderCursor.GetWidgetPtr()))
 				{
 					WidgetUnderCursor.Widget->OnDragEnter(WidgetUnderCursor.Geometry, InDragDropEvent);
 #if WITH_SLATE_DEBUGGING
@@ -5153,7 +5153,7 @@ bool FSlateApplication::RoutePointerMoveEvent(const FWidgetPath& WidgetsUnderPoi
 		{
 			FEventRouter::Route<FNoReply>(this, FEventRouter::FBubblePolicy(WidgetsUnderPointer), PointerEvent, [&LastWidgetsUnderPointer](const FArrangedWidget& WidgetUnderCursor, const FPointerEvent& Event)
 			{
-				if (!LastWidgetsUnderPointer.ContainsWidget(WidgetUnderCursor.Widget))
+				if (!LastWidgetsUnderPointer.ContainsWidget(WidgetUnderCursor.GetWidgetPtr()))
 				{
 					WidgetUnderCursor.Widget->OnMouseEnter(WidgetUnderCursor.Geometry, Event);
 #if WITH_SLATE_DEBUGGING
@@ -5749,7 +5749,7 @@ bool FSlateApplication::ExecuteNavigation(const FWidgetPath& NavigationSource, T
 		TSharedPtr<SWidget> ViewportWidget = Viewport->GetWidget().Pin();
 		if (ViewportWidget.IsValid())
 		{
-			if (NavigationSource.ContainsWidget(ViewportWidget.ToSharedRef()))
+			if (NavigationSource.ContainsWidget(ViewportWidget.Get()))
 			{
 				bHandled = Viewport->HandleNavigation(UserIndex, DestinationWidget);
 			}
diff --git a/Engine/Source/Runtime/Slate/Private/Framework/Application/SlateUser.cpp b/Engine/Source/Runtime/Slate/Private/Framework/Application/SlateUser.cpp
--- a/Engine/Source/Runtime/Slate/Private/Framework/Application/SlateUser.cpp
+++ b/Engine/Source/Runtime/Slate/Private/Framework/Application/SlateUser.cpp
@@ -129,12 +129,12 @@ bool FSlateUser::ShouldShowFocus(TSharedPtr<const SWidget> Widget) const
 
 bool FSlateUser::HasFocusedDescendants(TSharedRef<const SWidget> Widget) const
 {
-	return WeakFocusPath.IsValid() && WeakFocusPath.GetLastWidget().Pin() != Widget && WeakFocusPath.ContainsWidget(Widget);
+	return WeakFocusPath.IsValid() && WeakFocusPath.GetLastWidget().Pin() != Widget && WeakFocusPath.ContainsWidget(&Widget.Get());
 }
 
 bool FSlateUser::IsWidgetInFocusPath(TSharedPtr<const SWidget> Widget) const
 {
-	return Widget && WeakFocusPath.IsValid() && WeakFocusPath.ContainsWidget(Widget.ToSharedRef());
+	return Widget && WeakFocusPath.IsValid() && WeakFocusPath.ContainsWidget(Widget.Get());
 }
 
 bool FSlateUser::SetFocus(const TSharedRef<SWidget>& WidgetToFocus, EFocusCause ReasonFocusIsChanging)
@@ -406,7 +406,7 @@ bool FSlateUser::IsWidgetUnderCursor(TSharedPtr<const SWidget> Widget) const
 bool FSlateUser::IsWidgetUnderPointer(TSharedPtr<const SWidget> Widget, uint32 PointerIndex) const
 {
 	const FWeakWidgetPath* WidgetsUnderPointer = WidgetsUnderPointerLastEventByIndex.Find(PointerIndex);
-	return Widget && WidgetsUnderPointer && WidgetsUnderPointer->ContainsWidget(Widget.ToSharedRef());
+	return Widget && WidgetsUnderPointer && WidgetsUnderPointer->ContainsWidget(Widget.Get());
 }
 
 bool FSlateUser::IsWidgetUnderAnyPointer(TSharedPtr<const SWidget> Widget) const
@@ -415,7 +415,7 @@ bool FSlateUser::IsWidgetUnderAnyPointer(TSharedPtr<const SWidget> Widget) const
 	{
 		for (const auto& IndexPathPair : WidgetsUnderPointerLastEventByIndex)
 		{
-			if (IndexPathPair.Value.ContainsWidget(Widget.ToSharedRef()))
+			if (IndexPathPair.Value.ContainsWidget(Widget.Get()))
 			{
 				return true;
 			}
diff --git a/Engine/Source/Runtime/Slate/Private/Framework/Docking/SDockingTabStack.cpp b/Engine/Source/Runtime/Slate/Private/Framework/Docking/SDockingTabStack.cpp
--- a/Engine/Source/Runtime/Slate/Private/Framework/Docking/SDockingTabStack.cpp
+++ b/Engine/Source/Runtime/Slate/Private/Framework/Docking/SDockingTabStack.cpp
@@ -461,7 +461,7 @@ void SDockingTabStack::OnFocusChanging( const FWeakWidgetPath& PreviousFocusPath
 	const TSharedPtr<SDockTab> ForegroundTab = TabWell->GetForegroundTab();
 	if ( ForegroundTab.IsValid() )
 	{
-		const bool bIsForegroundTabActive = NewWidgetPath.ContainsWidget( SharedThis(this) );
+		const bool bIsForegroundTabActive = NewWidgetPath.ContainsWidget( this );
 	
 		if (bIsForegroundTabActive)
 		{
diff --git a/Engine/Source/Runtime/Slate/Private/Framework/Docking/STabDrawer.cpp b/Engine/Source/Runtime/Slate/Private/Framework/Docking/STabDrawer.cpp
--- a/Engine/Source/Runtime/Slate/Private/Framework/Docking/STabDrawer.cpp
+++ b/Engine/Source/Runtime/Slate/Private/Framework/Docking/STabDrawer.cpp
@@ -358,7 +358,7 @@ static bool IsLegalWidgetFocused(const FWidgetPath& FocusPath, const TArrayView<
 {
 	for (const TSharedRef<SWidget>& Widget : LegalFocusWidgets)
 	{
-		if (FocusPath.ContainsWidget(Widget))
+		if (FocusPath.ContainsWidget(&Widget.Get()))
 		{
 			return true;
 		}
@@ -398,7 +398,7 @@ void STabDrawer::OnGlobalFocusChanging(const FFocusEvent& FocusEvent, const FWea
 
 						// See if the menu being opened is owned by the drawer contents and if so the menu should not be dismissed
 						FSlateApplication::Get().GeneratePathToWidgetUnchecked(MenuHost.ToSharedRef(), MenuHostPath);
-						if (!MenuHostPath.ContainsWidget(ChildSlot.GetWidget()))
+						if (!MenuHostPath.ContainsWidget(&ChildSlot.GetWidget().Get()))
 						{
 							bShouldLoseFocus = true;
 						}
diff --git a/Engine/Source/Runtime/Slate/Private/Widgets/Layout/SScrollBox.cpp b/Engine/Source/Runtime/Slate/Private/Widgets/Layout/SScrollBox.cpp
--- a/Engine/Source/Runtime/Slate/Private/Widgets/Layout/SScrollBox.cpp
+++ b/Engine/Source/Runtime/Slate/Private/Widgets/Layout/SScrollBox.cpp
@@ -1021,7 +1021,7 @@ void SScrollBox::OnFocusChanging(const FWeakWidgetPath& PreviousFocusPath, const
 {
 	if (ScrollWhenFocusChanges != EScrollWhenFocusChanges::NoScroll)
 	{
-		if (NewWidgetPath.IsValid() && NewWidgetPath.ContainsWidget(SharedThis(this)))
+		if (NewWidgetPath.IsValid() && NewWidgetPath.ContainsWidget(this))
 		{
 			ScrollDescendantIntoView(NewWidgetPath.GetLastWidget(), ScrollWhenFocusChanges == EScrollWhenFocusChanges::AnimatedScroll ? true : false, NavigationDestination, NavigationScrollPadding);
 		}
diff --git a/Engine/Source/Runtime/Slate/Private/Widgets/SViewport.cpp b/Engine/Source/Runtime/Slate/Private/Widgets/SViewport.cpp
--- a/Engine/Source/Runtime/Slate/Private/Widgets/SViewport.cpp
+++ b/Engine/Source/Runtime/Slate/Private/Widgets/SViewport.cpp
@@ -428,14 +428,14 @@ void SViewport::OnArrangeChildren( const FGeometry& AllottedGeometry, FArrangedC
 	}
 }
 
-TSharedPtr<FVirtualPointerPosition> SViewport::TranslateMouseCoordinateForCustomHitTestChild(const TSharedRef<SWidget>& ChildWidget, const FGeometry& MyGeometry, const FVector2D& ScreenSpaceMouseCoordinate, const FVector2D& LastScreenSpaceMouseCoordinate) const
+TOptional<FVirtualPointerPosition> SViewport::TranslateMouseCoordinateForCustomHitTestChild(const SWidget& ChildWidget, const FGeometry& MyGeometry, const FVector2D ScreenSpaceMouseCoordinate, const FVector2D LastScreenSpaceMouseCoordinate) const
 {
 	if( CustomHitTestPath.IsValid() )
 	{
 		return CustomHitTestPath->TranslateMouseCoordinateForCustomHitTestChild( ChildWidget, MyGeometry, ScreenSpaceMouseCoordinate, LastScreenSpaceMouseCoordinate );
 	}
 
-	return nullptr;
+	return TOptional<FVirtualPointerPosition>();
 }
 
 FNavigationReply SViewport::OnNavigation(const FGeometry& MyGeometry, const FNavigationEvent& InNavigationEvent)
diff --git a/Engine/Source/Runtime/Slate/Public/Widgets/SViewport.h b/Engine/Source/Runtime/Slate/Public/Widgets/SViewport.h
--- a/Engine/Source/Runtime/Slate/Public/Widgets/SViewport.h
+++ b/Engine/Source/Runtime/Slate/Public/Widgets/SViewport.h
@@ -233,7 +233,7 @@ public:
 	virtual FPopupMethodReply OnQueryPopupMethod() const override;
 	virtual void OnFinishedPointerInput() override;
 	virtual void OnArrangeChildren( const FGeometry& AllottedGeometry, FArrangedChildren& ArrangedChildren ) const override;
-	virtual TSharedPtr<struct FVirtualPointerPosition> TranslateMouseCoordinateForCustomHitTestChild(const TSharedRef<SWidget>& ChildWidget, const FGeometry& MyGeometry, const FVector2D& ScreenSpaceMouseCoordinate, const FVector2D& LastScreenSpaceMouseCoordinate) const override;
+	virtual TOptional<FVirtualPointerPosition> TranslateMouseCoordinateForCustomHitTestChild(const SWidget& ChildWidget, const FGeometry& MyGeometry, const FVector2D ScreenSpaceMouseCoordinate, const FVector2D LastScreenSpaceMouseCoordinate) const override;
 	virtual FNavigationReply OnNavigation( const FGeometry& MyGeometry, const FNavigationEvent& InNavigationEvent ) override;
 
 private:
diff --git a/Engine/Source/Runtime/SlateCore/Private/Input/HittestGrid.cpp b/Engine/Source/Runtime/SlateCore/Private/Input/HittestGrid.cpp
--- a/Engine/Source/Runtime/SlateCore/Private/Input/HittestGrid.cpp
+++ b/Engine/Source/Runtime/SlateCore/Private/Input/HittestGrid.cpp
@@ -194,7 +194,7 @@ TArray<FWidgetAndPointer> FHittestGrid::GetBubblePath(FVector2D DesktopSpaceCoor
 					FGeometry DesktopSpaceGeometry = CurWidget->GetPaintSpaceGeometry();
 					DesktopSpaceGeometry.AppendTransform(FSlateLayoutTransform(GridOrigin - GridWindowOrigin));
 
-					Path.Emplace(FArrangedWidget(CurWidget.ToSharedRef(), DesktopSpaceGeometry), TSharedPtr<FVirtualPointerPosition>());
+					Path.Emplace(FArrangedWidget(CurWidget.ToSharedRef(), DesktopSpaceGeometry));
 					CurWidget = CurWidget->Advanced_GetPaintParentWidget();
 				}
 
@@ -277,17 +277,15 @@ void FHittestGrid::ClearInternal(int32 TotalCells)
 	AppendedGridArray.Reset();
 }
 
-bool FHittestGrid::IsDescendantOf(const TSharedRef<SWidget> Parent, const FWidgetData& ChildData) const
+bool FHittestGrid::IsDescendantOf(const SWidget* ParentWidget, const FWidgetData& ChildData) const
 {
 	const TSharedPtr<SWidget> ChildWidgetPtr = ChildData.GetWidget();
-	if (ChildWidgetPtr == Parent)
+	const SWidget* CurWidget = ChildWidgetPtr.Get();
+	if (CurWidget == ParentWidget)
 	{
 		return false;
 	}
 
-	const SWidget* ParentWidget = &Parent.Get();
-	const SWidget* CurWidget = ChildWidgetPtr.Get();
-
 	while (CurWidget)
 	{
 		if (ParentWidget == CurWidget)
@@ -397,7 +395,7 @@ TSharedPtr<SWidget> FHittestGrid::FindFocusableWidget(FSlateRect WidgetRect, con
 				// If we have a non escape boundary condition and this widget isn't a descendant of our boundary condition widget then it's invalid so we keep looking.
 				if (NavigationReply.GetBoundaryRule() != EUINavigationRule::Escape
 					&& NavigationReply.GetHandler().IsValid()
-					&& !IsDescendantOf(NavigationReply.GetHandler().ToSharedRef(), TestCandidate))
+					&& !IsDescendantOf(NavigationReply.GetHandler().Get(), TestCandidate))
 				{
 					AddToNextFocusableWidgetCondidateDebugResults(TestWidget, HittestGridDebuggingText::NotADescendant);
 					continue;
@@ -676,11 +674,17 @@ bool FHittestGrid::SameSize(const FHittestGrid* OtherGrid) const
 
 void FHittestGrid::AddWidget(const TSharedRef<SWidget>& InWidget, int32 InBatchPriorityGroup, int32 InLayerId, int32 InSecondarySort)
 {
-	AddWidget(InWidget, InBatchPriorityGroup, InLayerId, FSlateInvalidationWidgetSortOrder());
+	AddWidget(&(InWidget.Get()), InBatchPriorityGroup, InLayerId, FSlateInvalidationWidgetSortOrder());
 }
 
 void FHittestGrid::AddWidget(const TSharedRef<SWidget>& InWidget, int32 InBatchPriorityGroup, int32 InLayerId, FSlateInvalidationWidgetSortOrder InSecondarySort)
 {
+	AddWidget(&(InWidget.Get()), InBatchPriorityGroup, InLayerId, InSecondarySort);
+}
+
+void FHittestGrid::AddWidget(const SWidget* InWidget, int32 InBatchPriorityGroup, int32 InLayerId, FSlateInvalidationWidgetSortOrder InSecondarySort)
+{
+	check(InWidget);
 	if (!InWidget->GetVisibility().IsHitTestVisible())
 	{
 		return;
@@ -704,7 +708,7 @@ void FHittestGrid::AddWidget(const TSharedRef<SWidget>& InWidget, int32 InBatchP
 	const int64 PrimarySort = (((int64)InBatchPriorityGroup << 32) | InLayerId);
 
 	bool bAddWidget = true;
-	if (int32* FoundIndex = WidgetMap.Find(&*InWidget))
+	if (int32* FoundIndex = WidgetMap.Find(InWidget))
 	{
 		FWidgetData& WidgetData = WidgetArray[*FoundIndex];
 		if (WidgetData.UpperLeftCell != UpperLeftCell || WidgetData.LowerRightCell != LowerRightCell)
@@ -724,8 +728,8 @@ void FHittestGrid::AddWidget(const TSharedRef<SWidget>& InWidget, int32 InBatchP
 
 	if (bAddWidget)
 	{
-		int32& WidgetIndex = WidgetMap.Add(&*InWidget);
-		WidgetIndex = WidgetArray.Emplace(InWidget, UpperLeftCell, LowerRightCell, PrimarySort, InSecondarySort, CurrentUserIndex);
+		int32& WidgetIndex = WidgetMap.Add(InWidget);
+		WidgetIndex = WidgetArray.Emplace(const_cast<SWidget*>(InWidget)->AsShared(), UpperLeftCell, LowerRightCell, PrimarySort, InSecondarySort, CurrentUserIndex);
 		for (int32 XIndex = UpperLeftCell.X; XIndex <= LowerRightCell.X; ++XIndex)
 		{
 			for (int32 YIndex = UpperLeftCell.Y; YIndex <= LowerRightCell.Y; ++YIndex)
@@ -776,7 +780,13 @@ void FHittestGrid::RemoveWidget(const SWidget* InWidget)
 
 void FHittestGrid::UpdateWidget(const TSharedRef<SWidget>& InWidget, FSlateInvalidationWidgetSortOrder InSecondarySort)
 {
-	if (int32* FoundWidgetIndex = WidgetMap.Find(&*InWidget))
+	UpdateWidget(&(InWidget.Get()), InSecondarySort);
+}
+
+void FHittestGrid::UpdateWidget(const SWidget* InWidget, FSlateInvalidationWidgetSortOrder InSecondarySort)
+{
+	check(InWidget);
+	if (int32* FoundWidgetIndex = WidgetMap.Find(InWidget))
 	{
 		WidgetArray[*FoundWidgetIndex].SecondarySort = InSecondarySort;
 	}
@@ -784,7 +794,12 @@ void FHittestGrid::UpdateWidget(const TSharedRef<SWidget>& InWidget, FSlateInval
 
 void FHittestGrid::InsertCustomHitTestPath(const TSharedRef<SWidget> InWidget, TSharedRef<ICustomHitTestPath> CustomHitTestPath)
 {
-	int32 WidgetIndex = WidgetMap.FindChecked(&*InWidget);
+	InsertCustomHitTestPath(&InWidget.Get(), CustomHitTestPath);
+}
+
+void FHittestGrid::InsertCustomHitTestPath(const SWidget* InWidget, const TSharedRef<ICustomHitTestPath>& CustomHitTestPath)
+{
+	int32 WidgetIndex = WidgetMap.FindChecked(InWidget);
 	FWidgetData& WidgetData = WidgetArray[WidgetIndex];
 	WidgetData.CustomPath = CustomHitTestPath;
 }
diff --git a/Engine/Source/Runtime/SlateCore/Private/Layout/ArrangedWidget.cpp b/Engine/Source/Runtime/SlateCore/Private/Layout/ArrangedWidget.cpp
--- a/Engine/Source/Runtime/SlateCore/Private/Layout/ArrangedWidget.cpp
+++ b/Engine/Source/Runtime/SlateCore/Private/Layout/ArrangedWidget.cpp
@@ -19,11 +19,24 @@ FString FArrangedWidget::ToString( ) const
 
 
 FWidgetAndPointer::FWidgetAndPointer()
-: FArrangedWidget(FArrangedWidget::GetNullWidget())
-, PointerPosition(TSharedPtr<FVirtualPointerPosition>())
+	: FArrangedWidget(FArrangedWidget::GetNullWidget())
 {}
 
-FWidgetAndPointer::FWidgetAndPointer( const FArrangedWidget& InWidget, const TSharedPtr<FVirtualPointerPosition>& InPosition )
-: FArrangedWidget(InWidget)
-, PointerPosition(InPosition)
+FWidgetAndPointer::FWidgetAndPointer( const FArrangedWidget& InWidget, const TSharedPtr<const FVirtualPointerPosition>& InPosition )
+	: FArrangedWidget(InWidget)
+{
+	if (InPosition)
+	{
+		OptionalPointerPosition = *InPosition.Get();
+	}
+}
+
+FWidgetAndPointer::FWidgetAndPointer(const FArrangedWidget& InWidget)
+	: FArrangedWidget(InWidget)
+	, OptionalPointerPosition()
+{}
+
+FWidgetAndPointer::FWidgetAndPointer(const FArrangedWidget& InWidget, TOptional<FVirtualPointerPosition> InPosition)
+	: FArrangedWidget(InWidget)
+	, OptionalPointerPosition(InPosition)
 {}
diff --git a/Engine/Source/Runtime/SlateCore/Private/Layout/WidgetPath.cpp b/Engine/Source/Runtime/SlateCore/Private/Layout/WidgetPath.cpp
--- a/Engine/Source/Runtime/SlateCore/Private/Layout/WidgetPath.cpp
+++ b/Engine/Source/Runtime/SlateCore/Private/Layout/WidgetPath.cpp
@@ -7,7 +7,7 @@
 DECLARE_CYCLE_STAT(TEXT("Weak-To-Strong WidgetPath"), STAT_WeakToStrong_WidgetPath, STATGROUP_Slate);
 
 
-
+PRAGMA_DISABLE_DEPRECATION_WARNINGS
 FWidgetPath::FWidgetPath()
 : Widgets( EVisibility::Visible )
 , TopLevelWindow()
@@ -22,21 +22,18 @@ FWidgetPath::FWidgetPath( TSharedPtr<SWindow> InTopLevelWindow, const FArrangedC
 {
 }
 
-FWidgetPath::FWidgetPath( TArray<FWidgetAndPointer> InWidgetsAndPointers )
+FWidgetPath::FWidgetPath( TArrayView<FWidgetAndPointer> InWidgetsAndPointers )
 : Widgets( FArrangedChildren::Hittest2_FromArray(InWidgetsAndPointers) )
 , TopLevelWindow( InWidgetsAndPointers.Num() > 0 ? StaticCastSharedRef<SWindow>(InWidgetsAndPointers[0].Widget) : TSharedPtr<SWindow>(nullptr) )
-, VirtualPointerPositions( [&InWidgetsAndPointers]()
-	{ 
-		TArray< TSharedPtr<FVirtualPointerPosition> > Pointers;
-		Pointers.Reserve(InWidgetsAndPointers.Num());
-		for ( const FWidgetAndPointer& WidgetAndPointer : InWidgetsAndPointers )
-		{
-			Pointers.Add( WidgetAndPointer.PointerPosition );
-		};
-		return Pointers;
-	}())
 {
+	check(InWidgetsAndPointers.Num() == 0 || InWidgetsAndPointers[0].Widget->Advanced_IsWindow());
+	VirtualPointerPositions.Reserve(InWidgetsAndPointers.Num());
+	for (const FWidgetAndPointer& WidgetAndPointer : InWidgetsAndPointers)
+	{
+		VirtualPointerPositions.Add(WidgetAndPointer.GetPointerPosition());
+	}
 }
+PRAGMA_ENABLE_DEPRECATION_WARNINGS
 
 FWidgetPath FWidgetPath::GetPathDownTo( TSharedRef<const SWidget> MarkerWidget ) const
 {
@@ -60,17 +57,24 @@ FWidgetPath FWidgetPath::GetPathDownTo( TSharedRef<const SWidget> MarkerWidget )
 	}	
 }
 
-const TSharedPtr<FVirtualPointerPosition>& FWidgetPath::GetCursorAt( int32 Index ) const
+const TSharedPtr<const FVirtualPointerPosition>& FWidgetPath::GetCursorAt( int32 Index ) const
+{
+	static TSharedPtr<const FVirtualPointerPosition> CursorAt;
+	return CursorAt;
+}
+
+
+bool FWidgetPath::ContainsWidget(TSharedRef<const SWidget> WidgetToFind) const
 {
-	return VirtualPointerPositions[Index];
+	return ContainsWidget(&WidgetToFind.Get());
 }
 
 
-bool FWidgetPath::ContainsWidget( TSharedRef<const SWidget> WidgetToFind ) const
+bool FWidgetPath::ContainsWidget( const SWidget* WidgetToFind ) const
 {
 	for(int32 WidgetIndex = 0; WidgetIndex < Widgets.Num(); ++WidgetIndex)
 	{
-		if ( Widgets[WidgetIndex].Widget == WidgetToFind )
+		if ( &Widgets[WidgetIndex].Widget.Get() == WidgetToFind )
 		{
 			return true;
 		}
@@ -93,6 +97,7 @@ TOptional<FArrangedWidget> FWidgetPath::FindArrangedWidget( TSharedRef<const SWi
 	return TOptional<FArrangedWidget>();
 }
 
+PRAGMA_DISABLE_DEPRECATION_WARNINGS
 TOptional<FWidgetAndPointer> FWidgetPath::FindArrangedWidgetAndCursor( TSharedRef<const SWidget> WidgetToFind ) const
 {
 	const int32 Index = Widgets.IndexOfByPredicate( [&WidgetToFind]( const FArrangedWidget& SomeWidget )
@@ -104,6 +109,7 @@ TOptional<FWidgetAndPointer> FWidgetPath::FindArrangedWidgetAndCursor( TSharedRe
 		? FWidgetAndPointer( Widgets[Index], VirtualPointerPositions[Index] )
 		: FWidgetAndPointer();
 }
+PRAGMA_ENABLE_DEPRECATION_WARNINGS
 
 TSharedRef<SWindow> FWidgetPath::GetWindow() const
 {
@@ -298,12 +304,9 @@ FWeakWidgetPath::EPathResolutionResult::Result FWeakWidgetPath::ToWidgetPath( FW
 			bPathUninterrupted = true;
 
 			FGeometry ParentGeometry = TopLevelWindowPtr->GetWindowGeometryInScreen();
-			PathWithGeometries.Add(FWidgetAndPointer(
-				FArrangedWidget(TopLevelWindowPtr.ToSharedRef(), ParentGeometry),
-				// @todo slate: this should be the cursor's virtual position in window space.
-				TSharedPtr<FVirtualPointerPosition>()));
+			PathWithGeometries.Add(FWidgetAndPointer(FArrangedWidget(TopLevelWindowPtr.ToSharedRef(), ParentGeometry)));
 
-			TSharedPtr<FVirtualPointerPosition> VirtualPointerPos;
+			TOptional<FVirtualPointerPosition> VirtualPointerPos;
 			// For every widget in the vertical slice...
 			for (int32 WidgetIndex = 0; bPathUninterrupted && WidgetIndex < WidgetPtrs.Num() - 1; ++WidgetIndex)
 			{
@@ -320,9 +323,9 @@ FWeakWidgetPath::EPathResolutionResult::Result FWeakWidgetPath::ToWidgetPath( FW
 
 						if (ChildWidgetPtr.IsValid() && ChildWidgetPtr->GetParentWidget() == CurWidgetRef)
 						{
-							if (PointerEvent && !VirtualPointerPos.IsValid())
+							if (PointerEvent && !VirtualPointerPos.IsSet())
 							{
-								VirtualPointerPos = CurWidget->TranslateMouseCoordinateForCustomHitTestChild(ChildWidgetPtr.ToSharedRef(), ParentGeometry, PointerEvent->GetScreenSpacePosition(), PointerEvent->GetLastScreenSpacePosition());
+								VirtualPointerPos = CurWidget->TranslateMouseCoordinateForCustomHitTestChild(*ChildWidgetPtr.Get(), ParentGeometry, PointerEvent->GetScreenSpacePosition(), PointerEvent->GetLastScreenSpacePosition());
 							}
 
 							bFoundChild = true;
@@ -358,14 +361,11 @@ FWeakWidgetPath::EPathResolutionResult::Result FWeakWidgetPath::ToWidgetPath( FW
 			bPathUninterrupted = true;
 
 			FGeometry ParentGeometry = TopLevelWindowPtr->GetWindowGeometryInScreen();
-			PathWithGeometries.Add(FWidgetAndPointer(
-				FArrangedWidget(TopLevelWindowPtr.ToSharedRef(), ParentGeometry),
-				// @todo slate: this should be the cursor's virtual position in window space.
-				TSharedPtr<FVirtualPointerPosition>()));
+			PathWithGeometries.Add(FWidgetAndPointer(FArrangedWidget(TopLevelWindowPtr.ToSharedRef(), ParentGeometry)));
 
 			FArrangedChildren ArrangedChildren(VisibilityFilter, true);
 
-			TSharedPtr<FVirtualPointerPosition> VirtualPointerPos;
+			TOptional<FVirtualPointerPosition> VirtualPointerPos;
 			// For every widget in the vertical slice...
 			for (int32 WidgetIndex = 0; bPathUninterrupted && WidgetIndex < WidgetPtrs.Num() - 1; ++WidgetIndex)
 			{
@@ -386,9 +386,9 @@ FWeakWidgetPath::EPathResolutionResult::Result FWeakWidgetPath::ToWidgetPath( FW
 
 						if (ArrangedWidget.Widget == WidgetPtrs[WidgetIndex + 1])
 						{
-							if (PointerEvent && !VirtualPointerPos.IsValid())
+							if (PointerEvent && !VirtualPointerPos.IsSet())
 							{
-								VirtualPointerPos = CurWidget->TranslateMouseCoordinateForCustomHitTestChild(ArrangedWidget.Widget, ParentGeometry, PointerEvent->GetScreenSpacePosition(), PointerEvent->GetLastScreenSpacePosition());
+								VirtualPointerPos = CurWidget->TranslateMouseCoordinateForCustomHitTestChild(ArrangedWidget.Widget.Get(), ParentGeometry, PointerEvent->GetScreenSpacePosition(), PointerEvent->GetLastScreenSpacePosition());
 							}
 
 							bFoundChild = true;
@@ -415,9 +415,14 @@ FWeakWidgetPath::EPathResolutionResult::Result FWeakWidgetPath::ToWidgetPath( FW
 
 bool FWeakWidgetPath::ContainsWidget( const TSharedRef< const SWidget >& SomeWidget ) const
 {
-	for ( int32 WidgetIndex=0; WidgetIndex<Widgets.Num(); ++WidgetIndex )
+	return ContainsWidget(&SomeWidget.Get());
+}
+
+bool FWeakWidgetPath::ContainsWidget(const SWidget* SomeWidget) const
+{
+	for (int32 WidgetIndex = 0; WidgetIndex < Widgets.Num(); ++WidgetIndex)
 	{
-		if (Widgets[WidgetIndex].Pin() == SomeWidget)
+		if (Widgets[WidgetIndex].Pin().Get() == SomeWidget)
 		{
 			return true;
 		}
diff --git a/Engine/Source/Runtime/SlateCore/Private/Types/PaintArgs.cpp b/Engine/Source/Runtime/SlateCore/Private/Types/PaintArgs.cpp
--- a/Engine/Source/Runtime/SlateCore/Private/Types/PaintArgs.cpp
+++ b/Engine/Source/Runtime/SlateCore/Private/Types/PaintArgs.cpp
@@ -24,8 +24,6 @@ FPaintArgs::FPaintArgs(const SWidget* PaintParent, FHittestGrid& InRootHittestGr
 
 FPaintArgs FPaintArgs::InsertCustomHitTestPath(const SWidget* Widget, TSharedRef<ICustomHitTestPath> CustomHitTestPath) const
 {
-	TSharedRef<SWidget> SafeWidget = const_cast<SWidget*>(Widget)->AsShared();
-
-	const_cast<FHittestGrid&>(CurrentGrid).InsertCustomHitTestPath(SafeWidget, CustomHitTestPath);
+	const_cast<FHittestGrid&>(CurrentGrid).InsertCustomHitTestPath(Widget, CustomHitTestPath);
 	return *this;
 }
diff --git a/Engine/Source/Runtime/SlateCore/Private/Widgets/SWidget.cpp b/Engine/Source/Runtime/SlateCore/Private/Widgets/SWidget.cpp
--- a/Engine/Source/Runtime/SlateCore/Private/Widgets/SWidget.cpp
+++ b/Engine/Source/Runtime/SlateCore/Private/Widgets/SWidget.cpp
@@ -621,6 +621,11 @@ TSharedPtr<struct FVirtualPointerPosition> SWidget::TranslateMouseCoordinateForC
 	return nullptr;
 }
 
+TOptional<FVirtualPointerPosition> SWidget::TranslateMouseCoordinateForCustomHitTestChild(const SWidget& ChildWidget, const FGeometry& MyGeometry, const FVector2D ScreenSpaceMouseCoordinate, const FVector2D LastScreenSpaceMouseCoordinate) const
+{
+	return TOptional<FVirtualPointerPosition>();
+}
+
 void SWidget::OnFinishedPointerInput()
 {
 
@@ -793,7 +798,7 @@ void SWidget::SetFastPathSortOrder(const FSlateInvalidationWidgetSortOrder InSor
 		{
 			if (FHittestGrid* HittestGrid = Root->GetHittestGrid())
 			{
-				HittestGrid->UpdateWidget(AsShared(), InSortOrder);
+				HittestGrid->UpdateWidget(this, InSortOrder);
 			}
 		}
 
@@ -847,7 +852,7 @@ void SWidget::UpdateFastPathVisibility(FSlateInvalidationWidgetVisibility Parent
 
 	if (HittestGridToRemoveFrom)
 	{
-		HittestGridToRemoveFrom->RemoveWidget(SharedThis(this));
+		HittestGridToRemoveFrom->RemoveWidget(this);
 	}
 
 	PersistentState.CachedElementHandle.ClearCachedElements();
@@ -872,7 +877,7 @@ void SWidget::UpdateFastPathWidgetRemoved(FHittestGrid* ParentHittestGrid)
 
 	if (HittestGridToRemoveFrom)
 	{
-		HittestGridToRemoveFrom->RemoveWidget(SharedThis(this));
+		HittestGridToRemoveFrom->RemoveWidget(this);
 	}
 
 	PersistentState.CachedElementHandle.RemoveFromCache();
@@ -1454,7 +1459,7 @@ int32 SWidget::Paint(const FPaintArgs& Args, const FGeometry& AllottedGeometry,
 
 	if (bOutgoingHittestability)
 	{
-		Args.GetHittestGrid().AddWidget(MutableThis->AsShared(), 0, LayerId, FastPathProxyHandle.GetWidgetSortOrder());
+		Args.GetHittestGrid().AddWidget(MutableThis, 0, LayerId, FastPathProxyHandle.GetWidgetSortOrder());
 	}
 
 	if (bClipToBounds)
diff --git a/Engine/Source/Runtime/SlateCore/Public/Input/HittestGrid.h b/Engine/Source/Runtime/SlateCore/Public/Input/HittestGrid.h
--- a/Engine/Source/Runtime/SlateCore/Public/Input/HittestGrid.h
+++ b/Engine/Source/Runtime/SlateCore/Public/Input/HittestGrid.h
@@ -21,7 +21,10 @@ public:
 
 	virtual void ArrangeCustomHitTestChildren( FArrangedChildren& ArrangedChildren ) const = 0;
 
-	virtual TSharedPtr<struct FVirtualPointerPosition> TranslateMouseCoordinateForCustomHitTestChild( const TSharedRef<SWidget>& ChildWidget, const FGeometry& ViewportGeometry, const FVector2D& ScreenSpaceMouseCoordinate, const FVector2D& LastScreenSpaceMouseCoordinate ) const = 0;
+	UE_DEPRECATED(5.0, "TranslateMouseCoordinateForCustomHitTestChild that returns a shared ptr is deprecated.")
+	virtual TSharedPtr<struct FVirtualPointerPosition> TranslateMouseCoordinateForCustomHitTestChild( const TSharedRef<SWidget>& ChildWidget, const FGeometry& ViewportGeometry, const FVector2D& ScreenSpaceMouseCoordinate, const FVector2D& LastScreenSpaceMouseCoordinate ) const { return TSharedPtr<struct FVirtualPointerPosition>(); }
+
+	virtual TOptional<FVirtualPointerPosition> TranslateMouseCoordinateForCustomHitTestChild(const SWidget& ChildWidget, const FGeometry& MyGeometry, const FVector2D ScreenSpaceMouseCoordinate, const FVector2D LastScreenSpaceMouseCoordinate) const = 0;
 };
 
 class SLATECORE_API FHittestGrid : public FNoncopyable
@@ -47,8 +50,12 @@ public:
 	bool SetHittestArea(const FVector2D& HittestPositionInDesktop, const FVector2D& HittestDimensions, const FVector2D& HitestOffsetInWindow = FVector2D::ZeroVector);
 
 	/** Insert custom hit test data for a widget already in the grid */
+	UE_DEPRECATED(5.0, "Deprecated. Use the InsertCustomHitTestPath with a pointer.")
 	void InsertCustomHitTestPath(const TSharedRef<SWidget> InWidget, TSharedRef<ICustomHitTestPath> CustomHitTestPath);
 
+	/** Insert custom hit test data for a widget already in the grid */
+	void InsertCustomHitTestPath(const SWidget* InWidget, const TSharedRef<ICustomHitTestPath>& CustomHitTestPath);
+
 	/** Sets the current slate user index that should be associated with any added widgets */
 	void SetUserIndex(int32 UserIndex) { CurrentUserIndex = UserIndex; }
 
@@ -83,16 +90,25 @@ public:
 	void AddWidget(const TSharedRef<SWidget>& InWidget, int32 InBatchPriorityGroup, int32 InLayerId, int32 InSecondarySort);
 
 	/** Add SWidget from the HitTest Grid */
+	UE_DEPRECATED(5.0, "Deprecated. Use the AddWidget with a pointer.")
 	void AddWidget(const TSharedRef<SWidget>& InWidget, int32 InBatchPriorityGroup, int32 InLayerId, FSlateInvalidationWidgetSortOrder InSecondarySort);
 
+	/** Add SWidget from the HitTest Grid */
+	void AddWidget(const SWidget* InWidget, int32 InBatchPriorityGroup, int32 InLayerId, FSlateInvalidationWidgetSortOrder InSecondarySort);
+
 	/** Remove SWidget from the HitTest Grid */
+	UE_DEPRECATED(5.0, "Deprecated. Use the RemoveWidget with a pointer.")
 	void RemoveWidget(const TSharedRef<SWidget>& InWidget);
 
 	/** Remove SWidget from the HitTest Grid */
 	void RemoveWidget(const SWidget* InWidget);
 
 	/** Update the widget SecondarySort without removing it and readding it again. */
+	UE_DEPRECATED(5.0, "Deprecated. Use the UpdateWidget with a pointer.")
 	void UpdateWidget(const TSharedRef<SWidget>& InWidget, FSlateInvalidationWidgetSortOrder InSecondarySort);
+	
+	/** Update the widget SecondarySort without removing it and readding it again. */
+	void UpdateWidget(const SWidget* InWidget, FSlateInvalidationWidgetSortOrder InSecondarySort);
 
 	/** Append an already existing grid that occupy the same space. */
 	UE_DEPRECATED(4.26, "Deprecated. Use the FHittestGrid::AddGrid method instead")
@@ -158,7 +174,7 @@ private:
 	 */
 	struct FWidgetData
 	{
-		FWidgetData(TSharedRef<SWidget> InWidget, const FIntPoint& InUpperLeftCell, const FIntPoint& InLowerRightCell, int64 InPrimarySort, FSlateInvalidationWidgetSortOrder InSecondarySort, int32 InUserIndex)
+		FWidgetData(const TWeakPtr<SWidget>& InWidget, const FIntPoint& InUpperLeftCell, const FIntPoint& InLowerRightCell, int64 InPrimarySort, FSlateInvalidationWidgetSortOrder InSecondarySort, int32 InUserIndex)
 			: WeakWidget(InWidget)
 			, UpperLeftCell(InUpperLeftCell)
 			, LowerRightCell(InLowerRightCell)
@@ -252,7 +268,7 @@ private:
 	FIndexAndDistance GetHitIndexFromCellIndex(const FGridTestingParams& Params) const;
 
 	/** @returns true if the child is a paint descendant of the provided Parent. */
-	bool IsDescendantOf(const TSharedRef<SWidget> Parent, const FWidgetData& ChildData) const;
+	bool IsDescendantOf(const SWidget* Parent, const FWidgetData& ChildData) const;
 
 	/** Utility function for searching for the next focusable widget. */
 	template<typename TCompareFunc, typename TSourceSideFunc, typename TDestSideFunc>
diff --git a/Engine/Source/Runtime/SlateCore/Public/Layout/ArrangedChildren.h b/Engine/Source/Runtime/SlateCore/Public/Layout/ArrangedChildren.h
--- a/Engine/Source/Runtime/SlateCore/Public/Layout/ArrangedChildren.h
+++ b/Engine/Source/Runtime/SlateCore/Public/Layout/ArrangedChildren.h
@@ -35,7 +35,7 @@ class SLATECORE_API FArrangedChildren
 	}
 
 	// @todo hittest2.0 : we should get rid of this eventually.
-	static FArrangedChildren Hittest2_FromArray(const TArray<FWidgetAndPointer>& InWidgets)
+	static FArrangedChildren Hittest2_FromArray(const TArrayView<FWidgetAndPointer> InWidgets)
 	{
 		FArrangedChildren Temp( EVisibility::All );
 		Temp.Array.Reserve(InWidgets.Num());
diff --git a/Engine/Source/Runtime/SlateCore/Public/Layout/ArrangedWidget.h b/Engine/Source/Runtime/SlateCore/Public/Layout/ArrangedWidget.h
--- a/Engine/Source/Runtime/SlateCore/Public/Layout/ArrangedWidget.h
+++ b/Engine/Source/Runtime/SlateCore/Public/Layout/ArrangedWidget.h
@@ -25,6 +25,12 @@ public:
 
 public:
 
+	/** The widget that is being arranged. */
+	SWidget* GetWidgetPtr() const
+	{
+		return &Widget.Get();
+	}
+
 	/**
 	 * Gets the string representation of the Widget and corresponding Geometry.
 	 *
@@ -54,11 +60,34 @@ public:
 	TSharedRef<SWidget> Widget;
 };
 
+PRAGMA_DISABLE_DEPRECATION_WARNINGS
 struct SLATECORE_API FWidgetAndPointer : public FArrangedWidget
 {
+public:
 	FWidgetAndPointer();
+	FWidgetAndPointer(const FArrangedWidget& InWidget);
+	FWidgetAndPointer(const FArrangedWidget& InWidget, TOptional<FVirtualPointerPosition> InPosition);
+
+	UE_DEPRECATED(5.0, "FWidgetAndPointer constructor that takes a shared ptr is deprecated.")
+	FWidgetAndPointer( const FArrangedWidget& InWidget, const TSharedPtr<const FVirtualPointerPosition>& InPosition );
 
-	FWidgetAndPointer( const FArrangedWidget& InWidget, const TSharedPtr<FVirtualPointerPosition>& InPosition );
+	TOptional<FVirtualPointerPosition> GetPointerPosition() const
+	{
+		return OptionalPointerPosition;
+	}
+
+	void SetPointerPosition(TOptional<FVirtualPointerPosition> InPosition)
+	{
+		OptionalPointerPosition = InPosition;
+	}
+
+public:
+#if WITH_EDITOR
+	UE_DEPRECATED(5.0, "Direct access to PointerPosition is now deprecated. Use the getter or setter.")
+	TSharedPtr<const FVirtualPointerPosition> PointerPosition;
+#endif
 
-	TSharedPtr<FVirtualPointerPosition> PointerPosition;
+private:
+	TOptional<FVirtualPointerPosition> OptionalPointerPosition;
 };
+PRAGMA_ENABLE_DEPRECATION_WARNINGS
diff --git a/Engine/Source/Runtime/SlateCore/Public/Layout/WidgetPath.h b/Engine/Source/Runtime/SlateCore/Public/Layout/WidgetPath.h
--- a/Engine/Source/Runtime/SlateCore/Public/Layout/WidgetPath.h
+++ b/Engine/Source/Runtime/SlateCore/Public/Layout/WidgetPath.h
@@ -31,6 +31,7 @@ struct FWidgetMatcher
  * The window is needed for its ability to determine its own geometry, from which the geometries of the rest
  * of the widget can be determined.
  */
+PRAGMA_DISABLE_DEPRECATION_WARNINGS
 class SLATECORE_API FWidgetPath
 {
 public:
@@ -38,7 +39,7 @@ public:
 
 	FWidgetPath( TSharedPtr<SWindow> InTopLevelWindow, const FArrangedChildren& InWidgetPath );
 
-	FWidgetPath( TArray<FWidgetAndPointer> InWidgetsAndPointers );
+	FWidgetPath( TArrayView<FWidgetAndPointer> InWidgetsAndPointers );
 
 	/**
 	 * @param MarkerWidget Copy the path up to and including this widget 
@@ -48,10 +49,23 @@ public:
 	FWidgetPath GetPathDownTo( TSharedRef<const SWidget> MarkerWidget ) const;
 	
 	/** Get the widget and associated virtual cursor at a given index. */
-	const TSharedPtr<FVirtualPointerPosition>& GetCursorAt( int32 Index ) const;
+	UE_DEPRECATED(5.0, "GetCursorAt is deprecated. Use GetVirtualPointerPosition instead.")
+	const TSharedPtr<const FVirtualPointerPosition>& GetCursorAt( int32 Index ) const;
+
+	/** Get the virtual representation of the mouse at each level in the widget path. */
+	PRAGMA_DISABLE_DEPRECATION_WARNINGS
+	TOptional<FVirtualPointerPosition> GetVirtualPointerPosition(int32 Index) const
+	{
+		return VirtualPointerPositions[Index];
+	}
+	PRAGMA_ENABLE_DEPRECATION_WARNINGS
 
 	/** @return true if the WidgetToFind is in this WidgetPath, false otherwise. */
+	UE_DEPRECATED(5.0, "ContainsWidget that takes a shared ptr is deprecated.")
 	bool ContainsWidget( TSharedRef<const SWidget> WidgetToFind ) const;
+	
+	/** @return true if the WidgetToFind is in this WidgetPath, false otherwise. */
+	bool ContainsWidget( const SWidget* WidgetToFind ) const;
 
 	TOptional<FArrangedWidget> FindArrangedWidget( TSharedRef<const SWidget> WidgetToFind ) const;
 
@@ -161,7 +175,8 @@ public:
 	TSharedPtr< SWindow > TopLevelWindow;
 
 	/** The virtual representation of the mouse at each level in the widget path.  Due to 3D widgets, the space you transition to can be completely arbitrary as you traverse the tree. */
-	TArray< TSharedPtr<FVirtualPointerPosition> > VirtualPointerPositions;
+	UE_DEPRECATED(5.0, "Direct access to VirtualPointerPositions is now deprecated. Use the getter.")
+	TArray< TOptional<FVirtualPointerPosition> > VirtualPointerPositions;
 
 private:
 
@@ -182,7 +197,7 @@ private:
 	template<typename MatchRuleType>
 	static bool SearchForWidgetRecursively_Reverse( const MatchRuleType& MatchRule, const FArrangedWidget& InCandidate, FArrangedChildren& OutReversedPath, EVisibility VisibilityFilter = EVisibility::Visible );
 };
-
+PRAGMA_ENABLE_DEPRECATION_WARNINGS
 
 /**
  * Just like a WidgetPath, but uses weak pointers and does not store geometry.
@@ -225,6 +240,15 @@ public:
 			Truncated
 		};
 	};
+
+
+	/** @return true if the WidgetToFind is in this WidgetPath, false otherwise. */
+	UE_DEPRECATED(5.0, "ContainsWidget that takes a shared ptr is deprecated.")
+	bool ContainsWidget(const TSharedRef<const SWidget>& WidgetToFind) const;
+
+	/** @return true if the WidgetToFind is in this WidgetPath, false otherwise. */
+	bool ContainsWidget(const SWidget* WidgetToFind) const;
+
 	/**
 	 * Make a non-weak WidgetPath out of this WeakWidgetPath. Do this by computing all the relevant geometries and converting the weak pointers to TSharedPtr.
 	 *
@@ -234,8 +258,6 @@ public:
 	 */
 	EPathResolutionResult::Result ToWidgetPath( FWidgetPath& WidgetPath, EInterruptedPathHandling::Type InterruptedPathHandling = EInterruptedPathHandling::Truncate, const FPointerEvent* PointerEvent = nullptr, const EVisibility VisibilityFilter = EVisibility::Visible) const;
 
-	bool ContainsWidget( const TSharedRef< const SWidget >& SomeWidget ) const;
-
 	/**
 	 * @param NavigationType      Direction in which to move the focus (only for use with EUINavigation::Next and EUINavigation::Previous).
 	 * 
diff --git a/Engine/Source/Runtime/SlateCore/Public/Widgets/SWidget.h b/Engine/Source/Runtime/SlateCore/Public/Widgets/SWidget.h
--- a/Engine/Source/Runtime/SlateCore/Public/Widgets/SWidget.h
+++ b/Engine/Source/Runtime/SlateCore/Public/Widgets/SWidget.h
@@ -647,11 +647,15 @@ public:
 	UE_DEPRECATED(4.26, "Renaming to TranslateMouseCoordinateForCustomHitTestChild")
 	TSharedPtr<FVirtualPointerPosition> TranslateMouseCoordinateFor3DChild(const TSharedRef<SWidget>& ChildWidget, const FGeometry& MyGeometry, const FVector2D& ScreenSpaceMouseCoordinate, const FVector2D& LastScreenSpaceMouseCoordinate) const
 	{
-		return TranslateMouseCoordinateForCustomHitTestChild(ChildWidget, MyGeometry, ScreenSpaceMouseCoordinate, LastScreenSpaceMouseCoordinate);
+		TOptional<FVirtualPointerPosition> Result = TranslateMouseCoordinateForCustomHitTestChild(ChildWidget.Get(), MyGeometry, ScreenSpaceMouseCoordinate, LastScreenSpaceMouseCoordinate);
+		return Result.IsSet() ? MakeShared<FVirtualPointerPosition>(Result.GetValue()) : TSharedPtr<FVirtualPointerPosition>();
 	}
 
+	UE_DEPRECATED(5.0, "TranslateMouseCoordinateForCustomHitTestChild that returns a shared ptr is deprecated.")
 	virtual TSharedPtr<FVirtualPointerPosition> TranslateMouseCoordinateForCustomHitTestChild(const TSharedRef<SWidget>& ChildWidget, const FGeometry& MyGeometry, const FVector2D& ScreenSpaceMouseCoordinate, const FVector2D& LastScreenSpaceMouseCoordinate) const;
 
+	virtual TOptional<FVirtualPointerPosition> TranslateMouseCoordinateForCustomHitTestChild(const SWidget& ChildWidget, const FGeometry& MyGeometry, const FVector2D ScreenSpaceMouseCoordinate, const FVector2D LastScreenSpaceMouseCoordinate) const;
+
 	/**
 	 * All the pointer (mouse, touch, stylus, etc.) events from this frame have been routed.
 	 * This is a widget's chance to act on any accumulated data.
diff --git a/Engine/Source/Runtime/UMG/Private/Components/WidgetComponent.cpp b/Engine/Source/Runtime/UMG/Private/Components/WidgetComponent.cpp
--- a/Engine/Source/Runtime/UMG/Private/Components/WidgetComponent.cpp
+++ b/Engine/Source/Runtime/UMG/Private/Components/WidgetComponent.cpp
@@ -202,7 +202,7 @@ public:
 		}
 	}
 
-	virtual TSharedPtr<struct FVirtualPointerPosition> TranslateMouseCoordinateForCustomHitTestChild( const TSharedRef<SWidget>& ChildWidget, const FGeometry& ViewportGeometry, const FVector2D& ScreenSpaceMouseCoordinate, const FVector2D& LastScreenSpaceMouseCoordinate ) const override
+	virtual TOptional<FVirtualPointerPosition> TranslateMouseCoordinateForCustomHitTestChild(const SWidget& ChildWidget, const FGeometry& ViewportGeometry, const FVector2D ScreenSpaceMouseCoordinate, const FVector2D LastScreenSpaceMouseCoordinate) const override
 	{
 		if ( World.IsValid() && ensure(World->IsGameWorld()) )
 		{
@@ -216,7 +216,7 @@ public:
 				{
 					UWidgetComponent* WidgetComponent = Component.Get();
 					// Check if visible;
-					if ( WidgetComponent && WidgetComponent->GetSlateWindow() == ChildWidget )
+					if ( WidgetComponent && WidgetComponent->GetSlateWindow().Get() == &ChildWidget )
 					{
 						if ( UPrimitiveComponent* HitComponent = GetHitResultAtScreenPositionAndCache(TargetPlayer->PlayerController, LocalMouseCoordinate) )
 						{
@@ -226,15 +226,9 @@ public:
 								{
 									if ( WidgetComponent == HitComponent )
 									{
-										TSharedPtr<FVirtualPointerPosition> VirtualCursorPos = MakeShareable(new FVirtualPointerPosition);
-
 										FVector2D LocalHitLocation;
 										WidgetComponent->GetLocalHitLocation(CachedHitResult.Location, LocalHitLocation);
-
-										VirtualCursorPos->CurrentCursorPosition = LocalHitLocation;
-										VirtualCursorPos->LastCursorPosition = LocalHitLocation;
-
-										return VirtualCursorPos;
+										return FVirtualPointerPosition(LocalHitLocation, LocalHitLocation);
 									}
 								}
 							}
@@ -244,7 +238,7 @@ public:
 			}
 		}
 
-		return nullptr;
+		return TOptional<FVirtualPointerPosition>();
 	}
 	// End ICustomHitTestPath
 
@@ -1991,12 +1985,8 @@ TArray<FWidgetAndPointer> UWidgetComponent::GetHitWidgetPath(FVector WorldHitLoc
 
 TArray<FWidgetAndPointer> UWidgetComponent::GetHitWidgetPath(FVector2D WidgetSpaceHitCoordinate, bool bIgnoreEnabledStatus, float CursorRadius /*= 0.0f*/)
 {
-	TSharedRef<FVirtualPointerPosition> VirtualMouseCoordinate = MakeShareable(new FVirtualPointerPosition);
-
 	const FVector2D& LocalHitLocation = WidgetSpaceHitCoordinate;
-
-	VirtualMouseCoordinate->CurrentCursorPosition = LocalHitLocation;
-	VirtualMouseCoordinate->LastCursorPosition = LastLocalHitLocation;
+	const FVirtualPointerPosition VirtualMouseCoordinate(LocalHitLocation, LastLocalHitLocation);
 
 	// Cache the location of the hit
 	LastLocalHitLocation = LocalHitLocation;
@@ -2010,7 +2000,7 @@ TArray<FWidgetAndPointer> UWidgetComponent::GetHitWidgetPath(FVector2D WidgetSpa
 
 		for( FWidgetAndPointer& ArrangedWidget : ArrangedWidgets )
 		{
-			ArrangedWidget.PointerPosition = VirtualMouseCoordinate;
+			ArrangedWidget.SetPointerPosition(VirtualMouseCoordinate);
 		}
 	}
 
diff --git a/Engine/Source/Runtime/UMG/Private/UserWidget.cpp b/Engine/Source/Runtime/UMG/Private/UserWidget.cpp
--- a/Engine/Source/Runtime/UMG/Private/UserWidget.cpp
+++ b/Engine/Source/Runtime/UMG/Private/UserWidget.cpp
@@ -1646,10 +1646,10 @@ void UUserWidget::NativeOnFocusChanging(const FWeakWidgetPath& PreviousFocusPath
 	TSharedPtr<SObjectWidget> SafeGCWidget = MyGCWidget.Pin();
 	if ( SafeGCWidget.IsValid() )
 	{
-		const bool bDecendantNewlyFocused = NewWidgetPath.ContainsWidget(SafeGCWidget.ToSharedRef());
+		const bool bDecendantNewlyFocused = NewWidgetPath.ContainsWidget(SafeGCWidget.Get());
 		if ( bDecendantNewlyFocused )
 		{
-			const bool bDecendantPreviouslyFocused = PreviousFocusPath.ContainsWidget(SafeGCWidget.ToSharedRef());
+			const bool bDecendantPreviouslyFocused = PreviousFocusPath.ContainsWidget(SafeGCWidget.Get());
 			if ( !bDecendantPreviouslyFocused )
 			{
 				NativeOnAddedToFocusPath( InFocusEvent );
