Slate: Convert SScrolBox & SResponsiveGridPanel FSlot's to the FSlotArgument syntax.

diff --git a/Engine/Source/Runtime/Slate/Private/Widgets/Layout/SResponsiveGridPanel.cpp b/Engine/Source/Runtime/Slate/Private/Widgets/Layout/SResponsiveGridPanel.cpp
--- a/Engine/Source/Runtime/Slate/Private/Widgets/Layout/SResponsiveGridPanel.cpp
+++ b/Engine/Source/Runtime/Slate/Private/Widgets/Layout/SResponsiveGridPanel.cpp
@@ -13,9 +13,15 @@ SResponsiveGridPanel::SResponsiveGridPanel()
 : Slots(this)
 {}
 
-SResponsiveGridPanel::FSlot& SResponsiveGridPanel::AddSlot(int32 Row)
+SResponsiveGridPanel::FScopedWidgetSlotArguments SResponsiveGridPanel::AddSlot(int32 Row)
 {
-	return InsertSlot(new FSlot(Row));
+	TWeakPtr<SResponsiveGridPanel> WeakPanel = SharedThis(this);
+
+	TUniquePtr<FSlot> NewSlot = MakeUnique<FSlot>(Row);
+	NewSlot->Panel = WeakPanel;
+
+	int32 InsertLocation = FindInsertSlotLocation(NewSlot.Get());
+	return FScopedWidgetSlotArguments{ MoveTemp(NewSlot), this->Slots, InsertLocation};
 }
 
 bool SResponsiveGridPanel::RemoveSlot(const TSharedRef<SWidget>& SlotWidget)
@@ -38,9 +44,24 @@ void SResponsiveGridPanel::Construct(const FArguments& InArgs, int32 InTotalColu
 
 	RowFillCoefficients = InArgs.RowFillCoefficients;
 
-	for (int32 SlotIdx = 0; SlotIdx < InArgs.Slots.Num(); ++SlotIdx)
+	// Populate the slots such that they are sorted by Layer (order preserved within layers)
+
+	TWeakPtr<SResponsiveGridPanel> WeakPanel = SharedThis(this);
+	Slots.Reserve(InArgs._Slots.Num());
+	for (int32 SlotIndex = 0; SlotIndex < InArgs._Slots.Num(); ++SlotIndex)
 	{
-		InsertSlot(InArgs.Slots[SlotIdx]);
+		int32 SlotLocation = FindInsertSlotLocation(InArgs._Slots[SlotIndex].GetSlot());
+		if (SlotLocation == INDEX_NONE)
+		{
+			SlotLocation = Slots.AddSlot(MoveTemp(const_cast<FSlot::FSlotArguments&>(InArgs._Slots[SlotIndex])));
+		}
+		else
+		{
+			Slots.InsertSlot(MoveTemp(const_cast<FSlot::FSlotArguments&>(InArgs._Slots[SlotIndex])), SlotLocation);
+		}
+
+		FSlot& NewSlot = Slots[SlotLocation];
+		NewSlot.Panel = WeakPanel;
 	}
 }
 
@@ -410,39 +431,18 @@ void SResponsiveGridPanel::SetRowFill(int32 RowId, float Coefficient)
 	RowFillCoefficients[RowId] = Coefficient;
 }
 
-SResponsiveGridPanel::FSlot& SResponsiveGridPanel::InsertSlot(SResponsiveGridPanel::FSlot* InSlot)
+int32 SResponsiveGridPanel::FindInsertSlotLocation(SResponsiveGridPanel::FSlot* InSlot) const
 {
-	InSlot->Panel = SharedThis(this);
-
-	bool bInserted = false;
-
 	// Insert the slot in the list such that slots are sorted by LayerOffset.
-	for (int32 SlotIndex = 0; !bInserted && SlotIndex < Slots.Num(); ++SlotIndex)
+	for (int32 SlotIndex = 0; SlotIndex < Slots.Num(); ++SlotIndex)
 	{
 		if (InSlot->RowParam < this->Slots[SlotIndex].RowParam)
 		{
-			Slots.Insert(InSlot, SlotIndex);
-			bInserted = true;
+			return SlotIndex;
 		}
 	}
 
-	// We haven't inserted yet, so add to the end of the list.
-	if (!bInserted)
-	{
-		Slots.Add(InSlot);
-	}
-
-	NotifySlotChanged(InSlot);
-
-	return *InSlot;
-}
-
-void SResponsiveGridPanel::NotifySlotChanged(SResponsiveGridPanel::FSlot* InSlot)
-{
-	//// Keep the size of the grid up to date.
-	//// We need an extra cell at the end for easily figuring out the size across any number of cells
-
-	//Rows.AddZeroed();
+	return INDEX_NONE;
 }
 
 int32 SResponsiveGridPanel::LayoutDebugPaint(const FGeometry& AllottedGeometry, const FSlateRect& MyCullingRect, FSlateWindowElementList& OutDrawElements, int32 LayerId ) const
diff --git a/Engine/Source/Runtime/Slate/Private/Widgets/Layout/SScrollBox.cpp b/Engine/Source/Runtime/Slate/Private/Widgets/Layout/SScrollBox.cpp
--- a/Engine/Source/Runtime/Slate/Private/Widgets/Layout/SScrollBox.cpp
+++ b/Engine/Source/Runtime/Slate/Private/Widgets/Layout/SScrollBox.cpp
@@ -10,9 +10,9 @@
 #include "Widgets/Images/SImage.h"
 
 
-SScrollBox::FSlot& SScrollBox::Slot()
+SScrollBox::FSlot::FSlotArguments SScrollBox::Slot()
 {
-	return *(new SScrollBox::FSlot());
+	return FSlot::FSlotArguments(MakeUnique<FSlot>());
 }
 
 void SScrollPanel::Construct(const FArguments& InArgs, const TArray<SScrollBox::FSlot*>& InSlots)
@@ -28,6 +28,16 @@ void SScrollPanel::Construct(const FArguments& InArgs, const TArray<SScrollBox::
 	FrontPadScrolling = InArgs._FrontPadScrolling;
 }
 
+
+void SScrollPanel::Construct(const FArguments& InArgs, TArray<SScrollBox::FSlot::FSlotArguments> InSlots)
+{
+	PhysicalOffset = 0;
+	Children.AddSlots(MoveTemp(InSlots));
+	Orientation = InArgs._Orientation;
+	BackPadScrolling = InArgs._BackPadScrolling;
+	FrontPadScrolling = InArgs._FrontPadScrolling;
+}
+
 void SScrollPanel::OnArrangeChildren(const FGeometry& AllottedGeometry, FArrangedChildren& ArrangedChildren) const
 {
 	float ScrollPadding = Orientation == Orient_Vertical ? AllottedGeometry.GetLocalSize().Y : AllottedGeometry.GetLocalSize().X;
@@ -167,7 +177,7 @@ void SScrollBox::Construct( const FArguments& InArgs )
 		bScrollBarIsExternal = false;
 	}
 
-	SAssignNew(ScrollPanel, SScrollPanel, InArgs.Slots)
+	SAssignNew(ScrollPanel, SScrollPanel, MoveTemp(const_cast<TArray<FSlot::FSlotArguments>&>(InArgs._Slots)))
 		.Clipping(InArgs._Clipping)
 		.Orientation(Orientation)
 		.BackPadScrolling(BackPadScrolling)
@@ -305,12 +315,9 @@ void SScrollBox::ConstructHorizontalLayout()
 }
 
 /** Adds a slot to SScrollBox */
-SScrollBox::FSlot& SScrollBox::AddSlot()
+SScrollBox::FScopedWidgetSlotArguments SScrollBox::AddSlot()
 {
-	SScrollBox::FSlot& NewSlot = *new SScrollBox::FSlot();
-	ScrollPanel->Children.Add( &NewSlot );
-
-	return NewSlot;
+	return FScopedWidgetSlotArguments{ MakeUnique<FSlot>(), ScrollPanel->Children, INDEX_NONE };
 }
 
 /** Removes a slot at the specified location */
diff --git a/Engine/Source/Runtime/Slate/Public/Widgets/Layout/SResponsiveGridPanel.h b/Engine/Source/Runtime/Slate/Public/Widgets/Layout/SResponsiveGridPanel.h
--- a/Engine/Source/Runtime/Slate/Public/Widgets/Layout/SResponsiveGridPanel.h
+++ b/Engine/Source/Runtime/Slate/Public/Widgets/Layout/SResponsiveGridPanel.h
@@ -40,8 +40,8 @@ public:
 
 	class FSlot : public TBasicLayoutWidgetSlot< FSlot >
 	{
-	public:
-
+	private:
+		friend SResponsiveGridPanel;
 		struct FColumnLayout
 		{
 			float LayoutSize;
@@ -57,54 +57,62 @@ public:
 			, ColumnLayouts()
 		{ }
 
-		/** The panel that contains this slot */
-		TWeakPtr<SResponsiveGridPanel> Panel;
 
-		/** The row index*/
-		int32 RowParam;
-
-		/** Layout information for the column */
-		TArray<FColumnLayout> ColumnLayouts;
-
-		/** How many columns this slot spans over */
-		FSlot& ColumnSpan(float LayoutSize, int32 ColumnSpan, int32 ColumnOffset = 0)
-		{
-			FColumnLayout ColumnLayout;
-			ColumnLayout.LayoutSize = LayoutSize;
-			ColumnLayout.Span = FMath::Max(0, ColumnSpan);
-			ColumnLayout.Offset = ColumnOffset;
-
-			bool Inserted = false;
-			for (int32 Index = 0; Index < ColumnLayouts.Num(); Index++)
+		SLATE_SLOT_BEGIN_ARGS(FSlot, TBasicLayoutWidgetSlot<FSlot>)
+			/** How many columns this slot spans over */
+			FSlot::FSlotArguments& ColumnSpan(float LayoutSize, int32 ColumnSpan, int32 ColumnOffset = 0)
 			{
-				if (ColumnLayout.LayoutSize < ColumnLayouts[Index].LayoutSize)
+				FColumnLayout ColumnLayout;
+				ColumnLayout.LayoutSize = LayoutSize;
+				ColumnLayout.Span = FMath::Max(0, ColumnSpan);
+				ColumnLayout.Offset = ColumnOffset;
+
+				bool Inserted = false;
+				for (int32 Index = 0; Index < ColumnLayouts.Num(); Index++)
 				{
-					ColumnLayouts.Insert(ColumnLayout, Index);
-					Inserted = true;
-					break;
+					if (ColumnLayout.LayoutSize < ColumnLayouts[Index].LayoutSize)
+					{
+						ColumnLayouts.Insert(ColumnLayout, Index);
+						Inserted = true;
+						break;
+					}
 				}
-			}
 
-			if (!Inserted)
-			{
-				ColumnLayouts.Add(ColumnLayout);
-			}
+				if (!Inserted)
+				{
+					ColumnLayouts.Add(ColumnLayout);
+				}
 
-			if (Panel.IsValid())
-			{
-				Panel.Pin()->NotifySlotChanged(this);
+				return *this;
 			}
+		
+			/** Layout information for the column */
+			TArray<FColumnLayout> ColumnLayouts;
+		SLATE_SLOT_END_ARGS()
 
-			return *this;
+		void Construct(const FChildren& SlotOwner, FSlotArguments&& InArgs)
+		{
+			TBasicLayoutWidgetSlot<FSlot>::Construct(SlotOwner, MoveTemp(InArgs));
+			ColumnLayouts = MoveTemp(InArgs.ColumnLayouts);
 		}
+
+	private:
+		/** The panel that contains this slot */
+		TWeakPtr<SResponsiveGridPanel> Panel;
+
+		/** The row index*/
+		int32 RowParam;
+
+		/** Layout information for the column */
+		TArray<FColumnLayout> ColumnLayouts;
 	};
 
 	/**
 	 * Used by declarative syntax to create a Slot
 	 */
-	static FSlot& Slot(int32 Row)
+	static FSlot::FSlotArguments Slot(int32 Row)
 	{
-		return *(new FSlot(Row));
+		return FSlot::FSlotArguments(MakeUnique<FSlot>(Row));
 	}
 
 	SLATE_BEGIN_ARGS(SResponsiveGridPanel)
@@ -117,7 +125,7 @@ public:
 		SLATE_ARGUMENT(float, ColumnGutter)
 		SLATE_ARGUMENT(float, RowGutter)
 
-		SLATE_SUPPORTS_SLOT(FSlot)
+		SLATE_SLOT_ARGUMENT(FSlot, Slots)
 
 		/** Specify a row to stretch instead of sizing to content. */
 		FArguments& FillRow(int32 RowId, float Coefficient)
@@ -142,7 +150,8 @@ public:
 	
 	void Construct( const FArguments& InArgs, int32 TotalColumns );
 
-	FSlot& AddSlot(int32 Row);
+	using FScopedWidgetSlotArguments = TPanelChildren<FSlot>::FScopedWidgetSlotArguments;
+	FScopedWidgetSlotArguments AddSlot(int32 Row);
 
 	bool RemoveSlot(const TSharedRef<SWidget>& SlotWidget);
 
@@ -176,12 +185,9 @@ private:
 	static void DistributeSizeContributions( float SizeContribution, TArray<float>& DistributeOverMe, int32 StartIndex, int32 UpperBound );
 
 	/**
-	 * Inserts the given slot into the list of Slots based on its LayerParam, such that Slots are sorter by layer.
-	 *
-	 * @param The newly-allocated slot to insert.
-	 * @return A reference to the added slot
+	 * Return the location where to insert the slot. INDEX_NONE if we insert it at the end.
 	 */
-	FSlot& InsertSlot(FSlot* InSlot);
+	int32 FindInsertSlotLocation(SResponsiveGridPanel::FSlot* InSlot) const;
 
 	/** Compute the sizes of columns and rows needed to fit all the slots in this grid. */
 	void ComputeDesiredCellSizes(float AvailableWidth, TArray<float>& OutColumns, TArray<float>& OutRows, TArray<float>& OutRowToSlot) const;
@@ -189,14 +195,6 @@ private:
 	/** Draw the debug grid of rows and colummns; useful for inspecting the GridPanel's logic. See OnPaint() for parameter meaning */
 	int32 LayoutDebugPaint(const FGeometry& AllottedGeometry, const FSlateRect& MyCullingRect, FSlateWindowElementList& OutDrawElements, int32 LayerId ) const;
 
-	/** 
-	 * Callback used to resize our internal arrays if a slot (or slot span) is changed. 
-	 *
-	 * @param The slot that has just changed.
-	 */
-	void NotifySlotChanged(FSlot* InSlot);
-
-
 private:
 
 	/** The rows*/
diff --git a/Engine/Source/Runtime/Slate/Public/Widgets/Layout/SScrollBox.h b/Engine/Source/Runtime/Slate/Public/Widgets/Layout/SScrollBox.h
--- a/Engine/Source/Runtime/Slate/Public/Widgets/Layout/SScrollBox.h
+++ b/Engine/Source/Runtime/Slate/Public/Widgets/Layout/SScrollBox.h
@@ -95,7 +95,7 @@ public:
 			_Clipping = EWidgetClipping::ClipToBounds;
 		}
 		
-		SLATE_SUPPORTS_SLOT( FSlot )
+		SLATE_SLOT_ARGUMENT( FSlot, Slots )
 
 		/** Style used to draw this scrollbox */
 		SLATE_STYLE_ARGUMENT( FScrollBoxStyle, Style )
@@ -149,12 +149,13 @@ public:
 	SScrollBox();
 
 	/** @return a new slot. Slots contain children for SScrollBox */
-	static FSlot& Slot();
+	static FSlot::FSlotArguments Slot();
 
 	void Construct( const FArguments& InArgs );
 
+	using FScopedWidgetSlotArguments = TPanelChildren<FSlot>::FScopedWidgetSlotArguments;
 	/** Adds a slot to SScrollBox */
-	SScrollBox::FSlot& AddSlot();
+	FScopedWidgetSlotArguments AddSlot();
 
 	/** Removes a slot at the specified location */
 	void RemoveSlot( const TSharedRef<SWidget>& WidgetToRemove );
@@ -435,8 +436,11 @@ public:
 	{
 	}
 
+	UE_DEPRECATED(5.0, "Direct construction of FSlot is deprecated")
 	void Construct(const FArguments& InArgs, const TArray<SScrollBox::FSlot*>& InSlots);
 
+	void Construct(const FArguments& InArgs, TArray<SScrollBox::FSlot::FSlotArguments> InSlots);
+
 public:
 
 	EOrientation GetOrientation()
diff --git a/Engine/Source/Runtime/UMG/Private/Components/ScrollBoxSlot.cpp b/Engine/Source/Runtime/UMG/Private/Components/ScrollBoxSlot.cpp
--- a/Engine/Source/Runtime/UMG/Private/Components/ScrollBoxSlot.cpp
+++ b/Engine/Source/Runtime/UMG/Private/Components/ScrollBoxSlot.cpp
@@ -16,10 +16,11 @@ UScrollBoxSlot::UScrollBoxSlot(const FObjectInitializer& ObjectInitializer)
 
 void UScrollBoxSlot::BuildSlot(TSharedRef<SScrollBox> ScrollBox)
 {
-	Slot = &ScrollBox->AddSlot()
+	ScrollBox->AddSlot()
 		.Padding(Padding)
 		.HAlign(HorizontalAlignment)
 		.VAlign(VerticalAlignment)
+		.Expose(Slot)
 		[
 			Content == nullptr ? SNullWidget::NullWidget : Content->TakeWidget()
 		];
@@ -30,7 +31,7 @@ void UScrollBoxSlot::SetPadding(FMargin InPadding)
 	Padding = InPadding;
 	if ( Slot )
 	{
-		Slot->Padding(InPadding);
+		Slot->SetPadding(InPadding);
 	}
 }
 
@@ -39,7 +40,7 @@ void UScrollBoxSlot::SetHorizontalAlignment(EHorizontalAlignment InHorizontalAli
 	HorizontalAlignment = InHorizontalAlignment;
 	if ( Slot )
 	{
-		Slot->HAlign(InHorizontalAlignment);
+		Slot->SetHorizontalAlignment(InHorizontalAlignment);
 	}
 }
 
@@ -48,7 +49,7 @@ void UScrollBoxSlot::SetVerticalAlignment(EVerticalAlignment InVerticalAlignment
 	VerticalAlignment = InVerticalAlignment;
 	if (Slot)
 	{
-		Slot->VAlign(InVerticalAlignment);
+		Slot->SetVerticalAlignment(InVerticalAlignment);
 	}
 }
 
