Slate - When the safezone changes margins it needs to invalidate layout.

diff --git a/Engine/Source/Runtime/Slate/Private/Widgets/Layout/SSafeZone.cpp b/Engine/Source/Runtime/Slate/Private/Widgets/Layout/SSafeZone.cpp
--- a/Engine/Source/Runtime/Slate/Private/Widgets/Layout/SSafeZone.cpp
+++ b/Engine/Source/Runtime/Slate/Private/Widgets/Layout/SSafeZone.cpp
@@ -85,7 +85,7 @@ void SSafeZone::SetTitleSafe( bool InIsTitleSafe )
 	UpdateSafeMargin();
 }
 
-void SSafeZone::UpdateSafeMargin() const
+void SSafeZone::UpdateSafeMargin()
 {
 	bSafeMarginNeedsUpdate = true;
 
@@ -127,6 +127,7 @@ void SSafeZone::UpdateSafeMargin() const
 
 	SafeMargin = FMargin(bPadLeft ? SafeMargin.Left : 0.0f, bPadTop ? SafeMargin.Top : 0.0f, bPadRight ? SafeMargin.Right : 0.0f, bPadBottom ? SafeMargin.Bottom : 0.0f);
 
+	Invalidate(EInvalidateWidgetReason::Layout);
 	bSafeMarginNeedsUpdate = false;
 }
 
@@ -160,7 +161,7 @@ FMargin SSafeZone::GetSafeMargin(float InLayoutScale) const
 {
 	if (bSafeMarginNeedsUpdate)
 	{
-		UpdateSafeMargin();
+		const_cast<SSafeZone*>(this)->UpdateSafeMargin();
 	}
 
 	const FMargin SlotPadding = Padding.Get() + (ComputeScaledSafeMargin(InLayoutScale) * SafeAreaScale);
diff --git a/Engine/Source/Runtime/Slate/Public/Widgets/Layout/SSafeZone.h b/Engine/Source/Runtime/Slate/Public/Widgets/Layout/SSafeZone.h
--- a/Engine/Source/Runtime/Slate/Public/Widgets/Layout/SSafeZone.h
+++ b/Engine/Source/Runtime/Slate/Public/Widgets/Layout/SSafeZone.h
@@ -119,7 +119,7 @@ public:
 
 private:
 
-	void UpdateSafeMargin() const;
+	void UpdateSafeMargin();
 	FMargin ComputeScaledSafeMargin(float Scale) const;
 
 	/** Cached values from the args */
