Add interface for extending tab context menu

diff --git a/Engine/Source/Runtime/Slate/Private/Framework/Docking/SDockingTabStack.cpp b/Engine/Source/Runtime/Slate/Private/Framework/Docking/SDockingTabStack.cpp
--- a/Engine/Source/Runtime/Slate/Private/Framework/Docking/SDockingTabStack.cpp
+++ b/Engine/Source/Runtime/Slate/Private/Framework/Docking/SDockingTabStack.cpp
@@ -685,6 +685,12 @@ TSharedRef<SWidget> SDockingTabStack::MakeContextMenu()
 			);
 		}
 		MenuBuilder.EndSection();
+
+		TSharedPtr<SDockTab> ForegroundTab = TabWell->GetForegroundTab();
+		if (ForegroundTab.IsValid())
+		{
+			ForegroundTab->ExtendContextMenu(MenuBuilder);
+		}
 	}
 
 	return MenuBuilder.MakeWidget();
diff --git a/Engine/Source/Runtime/Slate/Private/Widgets/Docking/SDockTab.cpp b/Engine/Source/Runtime/Slate/Private/Widgets/Docking/SDockTab.cpp
--- a/Engine/Source/Runtime/Slate/Private/Widgets/Docking/SDockTab.cpp
+++ b/Engine/Source/Runtime/Slate/Private/Widgets/Docking/SDockTab.cpp
@@ -416,6 +416,11 @@ void SDockTab::SetOnPersistVisualState( const FOnPersistVisualState& Handler )
 	OnPersistVisualState = Handler;
 }
 
+void SDockTab::SetOnExtendContextMenu( const FExtendContextMenu& Handler )
+{
+	OnExtendContextMenu = Handler;
+}
+
 void SDockTab::SetCanCloseTab( const FCanCloseTab& InOnTabClosing )
 {
 	OnCanCloseTab = InOnTabClosing;
@@ -537,6 +542,7 @@ void SDockTab::Construct( const FArguments& InArgs )
 	this->OnTabClosed = InArgs._OnTabClosed;
 	this->OnCanCloseTab = InArgs._OnCanCloseTab;
 	this->OnPersistVisualState = InArgs._OnPersistVisualState;
+	this->OnExtendContextMenu = InArgs._OnExtendContextMenu;
 	this->TabLabel = InArgs._Label;
 	this->TabLabelSuffix = InArgs._LabelSuffix;
 	this->bShouldAutosize = InArgs._ShouldAutosize;
@@ -795,6 +801,10 @@ const FSlateBrush* SDockTab::GetFlashOverlayImageBrush() const
 	return FStyleDefaults::GetNoBrush();
 }
 
+void SDockTab::ExtendContextMenu(FMenuBuilder& MenuBuilder)
+{
+	OnExtendContextMenu.ExecuteIfBound(MenuBuilder);
+}
 
 FSlateColor SDockTab::GetFlashColor() const
 {
diff --git a/Engine/Source/Runtime/Slate/Public/Widgets/Docking/SDockTab.h b/Engine/Source/Runtime/Slate/Public/Widgets/Docking/SDockTab.h
--- a/Engine/Source/Runtime/Slate/Public/Widgets/Docking/SDockTab.h
+++ b/Engine/Source/Runtime/Slate/Public/Widgets/Docking/SDockTab.h
@@ -40,6 +40,8 @@ enum ETabActivationCause : uint8
 	SetDirectly
 };
 
+class FMenuBuilder;
+
 /**
  * A tab widget that also holds on to some content that should be shown when this tab is selected.
  * Intended to be used in conjunction with SDockingTabStack.
@@ -65,6 +67,9 @@ public:
 	/** Delegate called before a tab is closed.  Returning false will prevent the tab from closing */
 	DECLARE_DELEGATE_RetVal( bool, FCanCloseTab );
 
+	/** Invoked to add entries to the tab context menu */
+	DECLARE_DELEGATE_OneParam(FExtendContextMenu, FMenuBuilder&);
+
 	SLATE_BEGIN_ARGS(SDockTab)
 		: _Content()
 		, _TabWellContentLeft()
@@ -104,6 +109,7 @@ public:
 		SLATE_ARGUMENT( bool, ShouldAutosize )
 		SLATE_EVENT( FCanCloseTab, OnCanCloseTab )
 		SLATE_EVENT( FOnPersistVisualState, OnPersistVisualState )
+		SLATE_EVENT( FExtendContextMenu, OnExtendContextMenu )
 		/** Invoked when a tab is closed from a drawer. This does not mean the tab or its contents is destroyed, just hidden. Use OnTabClosed for that */
 		SLATE_EVENT( FSimpleDelegate, OnTabDrawerClosed)
 		SLATE_ATTRIBUTE( FLinearColor, TabColorScale )
@@ -147,6 +153,9 @@ public:
 	/** @return the Foreground color that this widget sets; unset options if the widget does not set a foreground color */
 	virtual FSlateColor GetForegroundColor() const;
 
+	/** Add any entries specific to this tab to the tab context menu */
+	void ExtendContextMenu(FMenuBuilder& MenuBuilder);
+
 	/** Is this an MajorTab? A tool panel tab? */
 	ETabRole GetTabRole() const;
 
@@ -273,6 +282,9 @@ public:
 	/** Set the handler that will be invoked when the tab is closed from a drawer */
 	void SetOnTabDrawerClosed(const FSimpleDelegate InDelegate);
 
+	/** Set the handler for extending the tab context menu */
+	void SetOnExtendContextMenu( const FExtendContextMenu& Handler );
+
 	/** Get the tab manager currently managing this tab. Note that a user move the tab between Tab Managers, so this return value may change. */
 	TSharedRef<FTabManager> GetTabManager() const;
 
@@ -417,6 +429,8 @@ protected:
 
 	FSimpleDelegate OnTabDraggedOverDockArea;
 
+	FExtendContextMenu OnExtendContextMenu;
+
 	/**
 	 * Invoked during the Save Visual State pass; gives this tab a chance to save misc info about visual state.
 	 * e.g. Content Browser might save the current filters, current folder, whether some panel is collapsed, etc.
