TabManager change to allow document tabs to insert next to a particular tab id, but be created with a different id

diff --git a/Engine/Source/Runtime/Slate/Private/Framework/Docking/TabManager.cpp b/Engine/Source/Runtime/Slate/Private/Framework/Docking/TabManager.cpp
--- a/Engine/Source/Runtime/Slate/Private/Framework/Docking/TabManager.cpp
+++ b/Engine/Source/Runtime/Slate/Private/Framework/Docking/TabManager.cpp
@@ -1217,9 +1217,14 @@ void FTabManager::DrawAttention( const TSharedRef<SDockTab>& TabToHighlight )
 	}
 }
 
+void FTabManager::InsertNewDocumentTab(FName PlaceholderId, FName NewTabId, const FSearchPreference& SearchPreference, const TSharedRef<SDockTab>& UnmanagedTab)
+{
+	InsertDocumentTab(PlaceholderId, NewTabId, SearchPreference, UnmanagedTab, true);
+}
+
 void FTabManager::InsertNewDocumentTab(FName PlaceholderId, const FSearchPreference& SearchPreference, const TSharedRef<SDockTab>& UnmanagedTab)
 {
-	InsertDocumentTab(PlaceholderId, SearchPreference, UnmanagedTab, true);
+	InsertDocumentTab(PlaceholderId, PlaceholderId, SearchPreference, UnmanagedTab, true);
 }
 
 void FTabManager::InsertNewDocumentTab( FName PlaceholderId, ESearchPreference::Type SearchPreference, const TSharedRef<SDockTab>& UnmanagedTab )
@@ -1229,14 +1234,14 @@ void FTabManager::InsertNewDocumentTab( FName PlaceholderId, ESearchPreference::
 	case ESearchPreference::PreferLiveTab:
 	{
 		FLiveTabSearch Search;
-		InsertDocumentTab(PlaceholderId, Search, UnmanagedTab, true);
+		InsertDocumentTab(PlaceholderId, PlaceholderId, Search, UnmanagedTab, true);
 		break;
 	}
 
 	case ESearchPreference::RequireClosedTab:
 	{
 		FRequireClosedTab Search;
-		InsertDocumentTab(PlaceholderId, Search, UnmanagedTab, true);
+		InsertDocumentTab(PlaceholderId, PlaceholderId, Search, UnmanagedTab, true);
 		break;
 	}
 
@@ -1253,14 +1258,14 @@ void FTabManager::RestoreDocumentTab( FName PlaceholderId, ESearchPreference::Ty
 	case ESearchPreference::PreferLiveTab:
 	{
 		FLiveTabSearch Search;
-		InsertDocumentTab(PlaceholderId, Search, UnmanagedTab, false);
+		InsertDocumentTab(PlaceholderId, PlaceholderId, Search, UnmanagedTab, false);
 		break;
 	}
 
 	case ESearchPreference::RequireClosedTab:
 	{
 		FRequireClosedTab Search;
-		InsertDocumentTab(PlaceholderId, Search, UnmanagedTab, false);
+		InsertDocumentTab(PlaceholderId, PlaceholderId, Search, UnmanagedTab, false);
 		break;
 	}
 
@@ -1454,10 +1459,15 @@ void FTabManager::InvokeTabForMenu( FName TabId )
 }
 
 void FTabManager::InsertDocumentTab(FName PlaceholderId, const FSearchPreference& SearchPreference, const TSharedRef<SDockTab>& UnmanagedTab, bool bPlaySpawnAnim)
+{
+	InsertDocumentTab(PlaceholderId, PlaceholderId, SearchPreference, UnmanagedTab, bPlaySpawnAnim);
+}
+
+void FTabManager::InsertDocumentTab(FName PlaceholderId, FName NewTabId, const FSearchPreference& SearchPreference, const TSharedRef<SDockTab>& UnmanagedTab, bool bPlaySpawnAnim)
 {
 	bool bWasUnmanagedTabOpened = true;
 	const bool bTabNotManaged = ensure( ! FindTabInLiveAreas( FTabMatcher(UnmanagedTab->GetLayoutIdentifier()) ).IsValid() );
-	UnmanagedTab->SetLayoutIdentifier( FTabId(PlaceholderId, LastDocumentUID++) );
+	UnmanagedTab->SetLayoutIdentifier( FTabId(NewTabId, LastDocumentUID++) );
 	
 	if (bTabNotManaged)
 	{
diff --git a/Engine/Source/Runtime/Slate/Public/Framework/Docking/TabManager.h b/Engine/Source/Runtime/Slate/Public/Framework/Docking/TabManager.h
--- a/Engine/Source/Runtime/Slate/Public/Framework/Docking/TabManager.h
+++ b/Engine/Source/Runtime/Slate/Public/Framework/Docking/TabManager.h
@@ -814,6 +814,9 @@ class SLATE_API FTabManager : public TSharedFromThis<FTabManager>
 			};
 		};
 
+		/** Insert a new UnmanagedTab document tab next to an existing tab (closed or open) that has the PlaceholdId. Give the New tab NewTabId */
+		void InsertNewDocumentTab( FName PlaceholderId, FName NewTabId, const FSearchPreference& SearchPreference, const TSharedRef<SDockTab>& UnmanagedTab );
+
 		/** Insert a new UnmanagedTab document tab next to an existing tab (closed or open) that has the PlaceholdId. */
 		void InsertNewDocumentTab( FName PlaceholderId, const FSearchPreference& SearchPreference, const TSharedRef<SDockTab>& UnmanagedTab );
 		
@@ -923,6 +926,7 @@ class SLATE_API FTabManager : public TSharedFromThis<FTabManager>
 	protected:
 		
 		void InsertDocumentTab( FName PlaceholderId, const FSearchPreference& SearchPreference, const TSharedRef<SDockTab>& UnmanagedTab, bool bPlaySpawnAnim );
+		void InsertDocumentTab( FName PlaceholderId, FName NewTabId, const FSearchPreference& SearchPreference, const TSharedRef<SDockTab>& UnmanagedTab, bool bPlaySpawnAnim );
 
 		virtual void OpenUnmanagedTab(FName PlaceholderId, const FSearchPreference& SearchPreference, const TSharedRef<SDockTab>& UnmanagedTab);
 			
