Slate: Convert SButton attributes to SlateAttributes

diff --git a/Engine/Plugins/Editor/GeometryMode/Source/BspMode/Private/SBspPalette.cpp b/Engine/Plugins/Editor/GeometryMode/Source/BspMode/Private/SBspPalette.cpp
--- a/Engine/Plugins/Editor/GeometryMode/Source/BspMode/Private/SBspPalette.cpp
+++ b/Engine/Plugins/Editor/GeometryMode/Source/BspMode/Private/SBspPalette.cpp
@@ -53,7 +53,10 @@ class SBspButton: public SButton
 	{
 		if ( MouseEvent.GetEffectingButton() == EKeys::LeftMouseButton )
 		{
-			bIsPressed = true;
+			if (IsEnabled())
+			{
+				Press();
+			}
 
 			return FReply::Handled().DetectDrag( SharedThis( this ), MouseEvent.GetEffectingButton() );
 		}
diff --git a/Engine/Plugins/Experimental/CommonUI/Source/CommonUI/Private/CommonButtonTypes.cpp b/Engine/Plugins/Experimental/CommonUI/Source/CommonUI/Private/CommonButtonTypes.cpp
--- a/Engine/Plugins/Experimental/CommonUI/Source/CommonUI/Private/CommonButtonTypes.cpp
+++ b/Engine/Plugins/Experimental/CommonUI/Source/CommonUI/Private/CommonButtonTypes.cpp
@@ -116,9 +116,12 @@ FReply SCommonButton::OnKeyUp(const FGeometry& MyGeometry, const FKeyEvent& InKe
 	return SButton::OnKeyUp(MyGeometry, InKeyEvent);
 }
 
-bool SCommonButton::IsPressed() const
+void SCommonButton::Press()
 {
-	return IsInteractable() ? SButton::IsPressed() : false;
+	if (IsInteractable())
+	{
+		SButton::Press();
+	}
 }
 
 void SCommonButton::SetIsButtonEnabled(bool bInIsButtonEnabled)
@@ -128,7 +131,7 @@ void SCommonButton::SetIsButtonEnabled(bool bInIsButtonEnabled)
 
 void SCommonButton::SetIsButtonFocusable(bool bInIsButtonFocusable)
 {
-	bIsFocusable = bInIsButtonFocusable;
+	SetIsFocusable(bInIsButtonFocusable);
 }
 
 void SCommonButton::SetIsInteractionEnabled(bool bInIsInteractionEnabled)
@@ -147,14 +150,7 @@ void SCommonButton::SetIsInteractionEnabled(bool bInIsInteractionEnabled)
 	if (bWasHovered != bIsHoveredNow)
 	{
 		SetHover(bIsHoveredNow);
-		if (bIsHoveredNow)
-		{
-			OnHovered.ExecuteIfBound();
-		}
-		else
-		{
-			OnUnhovered.ExecuteIfBound();
-		}
+		ExecuteHoverStateChanged(false);
 	}
 }
 
diff --git a/Engine/Plugins/Experimental/CommonUI/Source/CommonUI/Private/CommonButtonTypes.h b/Engine/Plugins/Experimental/CommonUI/Source/CommonUI/Private/CommonButtonTypes.h
--- a/Engine/Plugins/Experimental/CommonUI/Source/CommonUI/Private/CommonButtonTypes.h
+++ b/Engine/Plugins/Experimental/CommonUI/Source/CommonUI/Private/CommonButtonTypes.h
@@ -90,8 +90,6 @@ public:
 
 	virtual FReply OnKeyUp(const FGeometry& MyGeometry, const FKeyEvent& InKeyEvent) override;
 
-	virtual bool IsPressed() const override;
-
 	void SetIsButtonEnabled(bool bInIsButtonEnabled);
 
 	void SetIsButtonFocusable(bool bInIsButtonFocusable);
@@ -105,6 +103,10 @@ public:
 
 	virtual int32 OnPaint(const FPaintArgs& Args, const FGeometry& AllottedGeometry, const FSlateRect& MyClippingRect, FSlateWindowElementList& OutDrawElements, int32 LayerId, const FWidgetStyle& InWidgetStyle, bool bParentEnabled) const override;
 
+protected:
+	/** Press the button */
+	virtual void Press() override;
+
 private:
 	FOnClicked OnDoubleClicked;
 
diff --git a/Engine/Source/Editor/AnimGraph/Private/PoseDriverDetails.cpp b/Engine/Source/Editor/AnimGraph/Private/PoseDriverDetails.cpp
--- a/Engine/Source/Editor/AnimGraph/Private/PoseDriverDetails.cpp
+++ b/Engine/Source/Editor/AnimGraph/Private/PoseDriverDetails.cpp
@@ -54,6 +54,8 @@ public:
 			.OnReleased(this, &FSoloToggleButton::OnButtonReleased)
 			.IsFocusable(false)
 		);
+
+		SetAppearPressed(SoloState);
 	}
 
 	FReply OnMouseButtonDoubleClick(const FGeometry& InMyGeometry, const FPointerEvent& InMouseEvent) override
@@ -62,12 +64,6 @@ public:
 		return FReply::Handled();
 	}
 
-protected:
-	bool IsPressed() const override
-	{
-		return SButton::IsPressed() || SoloState.Get(false);
-	}
-
 private:
 	void OnButtonPressed()
 	{
diff --git a/Engine/Source/Runtime/Slate/Private/Framework/MultiBox/SMenuEntryBlock.cpp b/Engine/Source/Runtime/Slate/Private/Framework/MultiBox/SMenuEntryBlock.cpp
--- a/Engine/Source/Runtime/Slate/Private/Framework/MultiBox/SMenuEntryBlock.cpp
+++ b/Engine/Source/Runtime/Slate/Private/Framework/MultiBox/SMenuEntryBlock.cpp
@@ -392,7 +392,10 @@ public:
 	{
 		if (ResponseToMouseUp == EResponseToMouseUp::Handle)
 		{
-			bIsPressed = true;
+			if (IsEnabled())
+			{
+				Press();
+			}
 		}
 
 		FReply Reply = SButton::OnMouseButtonUp(MyGeometry, MouseEvent);
@@ -407,7 +410,7 @@ public:
 			ResponseToMouseUp = EResponseToMouseUp::Handle;
 		}
 
-		bIsPressed = false;
+		Release();
 
 		SButton::OnMouseLeave(MouseEvent);
 	}
diff --git a/Engine/Source/Runtime/Slate/Private/Widgets/Input/SButton.cpp b/Engine/Source/Runtime/Slate/Private/Widgets/Input/SButton.cpp
--- a/Engine/Source/Runtime/Slate/Private/Widgets/Input/SButton.cpp
+++ b/Engine/Source/Runtime/Slate/Private/Widgets/Input/SButton.cpp
@@ -11,7 +11,49 @@
 
 static FName SButtonTypeName("SButton");
 
+SLATE_IMPLEMENT_WIDGET(SButton)
+void SButton::PrivateRegisterAttributes(FSlateAttributeInitializer& AttributeInitializer)
+{
+	SLATE_ADD_MEMBER_ATTRIBUTE_DEFINITION_WITH_NAME(AttributeInitializer, "BaseBorderForegroundColor", BorderForegroundColorAttribute, EInvalidateWidgetReason::Paint)
+		.OnValueChanged(FSlateAttributeDescriptor::FAttributeValueChangedDelegate::CreateLambda([](SWidget& Widget)
+			{
+				static_cast<SButton&>(Widget).UpdateForegroundColor();
+			}));
+
+	SLATE_ADD_MEMBER_ATTRIBUTE_DEFINITION_WITH_NAME(AttributeInitializer, "ContentPadding", ContentPaddingAttribute, EInvalidateWidgetReason::Layout)
+		.OnValueChanged(FSlateAttributeDescriptor::FAttributeValueChangedDelegate::CreateLambda([](SWidget& Widget)
+		{
+			static_cast<SButton&>(Widget).UpdatePadding();
+		}));
+	
+	SLATE_ADD_MEMBER_ATTRIBUTE_DEFINITION_WITH_NAME(AttributeInitializer, "AppearPressed", AppearPressedAttribute, EInvalidateWidgetReason::Paint)
+		.OnValueChanged(FSlateAttributeDescriptor::FAttributeValueChangedDelegate::CreateLambda([](SWidget& Widget)
+			{
+				static_cast<SButton&>(Widget).UpdatePressStateChanged();
+			}));
+
+	AttributeInitializer.OverrideInvalidationReason("EnabledState", FSlateAttributeDescriptor::FInvalidateWidgetReasonAttribute{EInvalidateWidgetReason::Layout|EInvalidateWidgetReason::Paint});
+
+	AttributeInitializer.OverrideOnValueChanged("EnabledState"
+		, FSlateAttributeDescriptor::ECallbackOverrideType::ExecuteAfterPrevious
+		, FSlateAttributeDescriptor::FAttributeValueChangedDelegate::CreateLambda([](SWidget& Widget)
+			{
+				static_cast<SButton&>(Widget).UpdateBorderImage();
+			}));
+
+	AttributeInitializer.OverrideOnValueChanged("Hovered"
+		, FSlateAttributeDescriptor::ECallbackOverrideType::ExecuteAfterPrevious
+		, FSlateAttributeDescriptor::FAttributeValueChangedDelegate::CreateLambda([](SWidget& Widget)
+			{
+				static_cast<SButton&>(Widget).UpdateBorderImage();
+				static_cast<SButton&>(Widget).UpdateForegroundColor();
+			}));
+}
+
 SButton::SButton()
+	: BorderForegroundColorAttribute(*this)
+	, ContentPaddingAttribute(*this)
+	, AppearPressedAttribute(*this)
 {
 #if WITH_ACCESSIBILITY
 	AccessibleBehavior = EAccessibleBehavior::Summary;
@@ -27,6 +69,22 @@ SButton::SButton()
 void SButton::Construct( const FArguments& InArgs )
 {
 	bIsPressed = false;
+	bIsFocusable = InArgs._IsFocusable;
+
+	BorderForegroundColorAttribute.Assign(*this, InArgs._ForegroundColor);
+
+	OnClicked = InArgs._OnClicked;
+	OnPressed = InArgs._OnPressed;
+	OnReleased = InArgs._OnReleased;
+	OnHovered = InArgs._OnHovered;
+	OnUnhovered = InArgs._OnUnhovered;
+
+	ClickMethod = InArgs._ClickMethod;
+	TouchMethod = InArgs._TouchMethod;
+	PressMethod = InArgs._PressMethod;
+
+	HoveredSound = InArgs._HoveredSoundOverride.Get(InArgs._ButtonStyle->HoveredSlateSound);
+	PressedSound = InArgs._PressedSoundOverride.Get(InArgs._ButtonStyle->PressedSlateSound);
 
 	// Text overrides button content. If nothing is specified, put an null widget in the button.
 	// Null content makes the button enter a special mode where it will ask to be as big as the image used for its border.
@@ -54,41 +112,21 @@ void SButton::Construct( const FArguments& InArgs )
 		.ContentScale(InArgs._ContentScale)
 		.DesiredSizeScale(InArgs._DesiredSizeScale)
 		.BorderBackgroundColor(InArgs._ButtonColorAndOpacity)
-		.ForegroundColor(InArgs._ForegroundColor)
-		.BorderImage( this, &SButton::GetBorder )
-		.HAlign( InArgs._HAlign )
-		.VAlign( InArgs._VAlign )
-		.Padding( TAttribute<FMargin>(this, &SButton::GetCombinedPadding) )
-		.ShowEffectWhenDisabled( TAttribute<bool>(this, &SButton::GetShowDisabledEffect) )
+		.HAlign(InArgs._HAlign)
+		.VAlign(InArgs._VAlign)
 		[
 			DetermineContent(InArgs)
 		]
 	);
 
+	SetContentPadding(InArgs._ContentPadding);
+	SetButtonStyle(InArgs._ButtonStyle);
+
 	// Only do this if we're exactly an SButton
 	if (GetType() == SButtonTypeName)
 	{
 		SetCanTick(false);
 	}
-
-	ContentPadding = InArgs._ContentPadding;
-
-	SetButtonStyle(InArgs._ButtonStyle);
-
-	bIsFocusable = InArgs._IsFocusable;
-
-	OnClicked = InArgs._OnClicked;
-	OnPressed = InArgs._OnPressed;
-	OnReleased = InArgs._OnReleased;
-	OnHovered = InArgs._OnHovered;
-	OnUnhovered = InArgs._OnUnhovered;
-
-	ClickMethod = InArgs._ClickMethod;
-	TouchMethod = InArgs._TouchMethod;
-	PressMethod = InArgs._PressMethod;
-
-	HoveredSound = InArgs._HoveredSoundOverride.Get(Style->HoveredSlateSound);
-	PressedSound = InArgs._PressedSoundOverride.Get(Style->PressedSlateSound);
 }
 
 int32 SButton::OnPaint(const FPaintArgs& Args, const FGeometry& AllottedGeometry, const FSlateRect& MyCullingRect, FSlateWindowElementList& OutDrawElements, int32 LayerId, const FWidgetStyle& InWidgetStyle, bool bParentEnabled) const
@@ -96,7 +134,7 @@ int32 SButton::OnPaint(const FPaintArgs& Args, const FGeometry& AllottedGeometry
 	bool bEnabled = ShouldBeEnabled(bParentEnabled);
 	bool bShowDisabledEffect = GetShowDisabledEffect();
 
-	const FSlateBrush* BrushResource = !bShowDisabledEffect && !bEnabled ? DisabledImage : GetBorder();
+	const FSlateBrush* BrushResource = !bShowDisabledEffect && !bEnabled ? &Style->Disabled : GetBorderImage();
 	
 	ESlateDrawEffect DrawEffects = bShowDisabledEffect && !bEnabled ? ESlateDrawEffect::DisabledEffect : ESlateDrawEffect::None;
 
@@ -118,36 +156,78 @@ int32 SButton::OnPaint(const FPaintArgs& Args, const FGeometry& AllottedGeometry
 FMargin SButton::GetCombinedPadding() const
 {
 	return ( IsPressed() )
-		? ContentPadding.Get() + PressedBorderPadding
-		: ContentPadding.Get() + BorderPadding;	
+		? ContentPaddingAttribute.Get() + Style->PressedPadding
+		: ContentPaddingAttribute.Get() + Style->NormalPadding;
+}
+
+//~ Update when { ContentPaddingAttribute, IsPressed, Style }
+void SButton::UpdatePadding()
+{
+	SetPadding(GetCombinedPadding());
 }
 
 bool SButton::GetShowDisabledEffect() const
 {
-	return DisabledImage->DrawAs == ESlateBrushDrawType::NoDrawType;
+	return Style->Disabled.DrawAs == ESlateBrushDrawType::NoDrawType;
 }
 
-/** @return An image that represents this button's border*/
-const FSlateBrush* SButton::GetBorder() const
+//~ Update when { Style }
+void SButton::UpdateShowDisabledEffect()
 {
-	if (!GetShowDisabledEffect() && !IsEnabled())
+	// Needs to be called when the style changed
+	SetShowEffectWhenDisabled(GetShowDisabledEffect());
+}
+
+//~ Update when { GetShowDisabledEffect(Style), IsEnable(EnabledState), Pressed, Hovered, Style }
+void SButton::UpdateBorderImage()
+{
+	if (!GetShowDisabledEffect() && !IsInteractable())
 	{
-		return DisabledImage;
+		SetBorderImage(&Style->Disabled);
 	}
-	else if ( IsPressed() )
+	else if (IsPressed())
 	{
-		return PressedImage;
+		SetBorderImage(&Style->Pressed);
 	}
 	else if (IsHovered())
 	{
-		return HoverImage;
+		SetBorderImage(&Style->Hovered);
+	}
+	else
+	{
+		SetBorderImage(&Style->Normal);
+	}
+}
+
+//~ Update when { DefaultForegroundColorAttribute, Pressed, Hovered, Style }
+void SButton::UpdateForegroundColor()
+{
+	if (BorderForegroundColorAttribute.Get() == FSlateColor::UseStyle())
+	{
+		if (IsPressed())
+		{
+			SetForegroundColor(Style->PressedForeground);
+		}
+		else if (IsHovered())
+		{
+			SetForegroundColor(Style->HoveredForeground);
+		}
+		else
+		{
+			SetForegroundColor(Style->NormalForeground);
+		}
 	}
 	else
 	{
-		return NormalImage;
+		SetForegroundColor(BorderForegroundColorAttribute.Get());
 	}
 }
 
+//~ Update when { Style }
+void SButton::UpdateDisabledForegroundColor()
+{
+	Invalidate(EInvalidateWidgetReason::Paint);
+}
 
 bool SButton::SupportsKeyboardFocus() const
 {
@@ -251,17 +331,10 @@ FReply SButton::OnMouseButtonDown( const FGeometry& MyGeometry, const FPointerEv
 		}
 	}
 
-	Invalidate(EInvalidateWidget::Layout);
-
 	//return the constructed reply
 	return Reply;
 }
 
-FReply SButton::OnMouseButtonDoubleClick( const FGeometry& InMyGeometry, const FPointerEvent& InMouseEvent )
-{
-	return OnMouseButtonDown( InMyGeometry, InMouseEvent );
-}
-
 FReply SButton::OnMouseButtonUp( const FGeometry& MyGeometry, const FPointerEvent& MouseEvent )
 {
 	FReply Reply = FReply::Unhandled();
@@ -275,7 +348,7 @@ FReply SButton::OnMouseButtonUp( const FGeometry& MyGeometry, const FPointerEven
 
 		if ( IsEnabled() )
 		{
-			if (InputClickMethod == EButtonClickMethod::MouseDown )
+			if ( InputClickMethod == EButtonClickMethod::MouseDown )
 			{
 				// NOTE: If we're configured to click on mouse-down/precise-tap, then we never capture the mouse thus
 				//       may never receive an OnMouseButtonUp() call.  We make sure that our bIsPressed
@@ -322,8 +395,6 @@ FReply SButton::OnMouseButtonUp( const FGeometry& MyGeometry, const FPointerEven
 		Reply.ReleaseMouseCapture();
 	}
 
-	Invalidate(EInvalidateWidget::Layout);
-
 	return Reply;
 }
 
@@ -339,16 +410,14 @@ FReply SButton::OnMouseMove( const FGeometry& MyGeometry, const FPointerEvent& M
 
 void SButton::OnMouseEnter( const FGeometry& MyGeometry, const FPointerEvent& MouseEvent )
 {
-	if ( IsEnabled() )
-	{
-		PlayHoverSound();
-	}
-	
-	SBorder::OnMouseEnter( MyGeometry, MouseEvent );
+	const bool bWasHovered = IsHovered();
 
-	OnHovered.ExecuteIfBound();
+	SBorder::OnMouseEnter( MyGeometry, MouseEvent );
 
-	Invalidate(EInvalidateWidget::Layout);
+	if (!bWasHovered && IsHovered())
+	{
+		ExecuteHoverStateChanged(true);
+	}
 }
 
 void SButton::OnMouseLeave( const FPointerEvent& MouseEvent )
@@ -365,12 +434,10 @@ void SButton::OnMouseLeave( const FPointerEvent& MouseEvent )
 		Release();
 	}
 
-	if (bWasHovered)
+	if (bWasHovered && !IsHovered())
 	{
-		OnUnhovered.ExecuteIfBound();
+		ExecuteHoverStateChanged(true);
 	}
-
-	Invalidate(EInvalidateWidget::Layout);
 }
 
 void SButton::OnMouseCaptureLost(const FCaptureLostEvent& CaptureLostEvent)
@@ -401,6 +468,7 @@ void SButton::Press()
 		bIsPressed = true;
 		PlayPressedSound();
 		OnPressed.ExecuteIfBound();
+		UpdatePressStateChanged();
 	}
 }
 
@@ -410,18 +478,20 @@ void SButton::Release()
 	{
 		bIsPressed = false;
 		OnReleased.ExecuteIfBound();
+		UpdatePressStateChanged();
 	}
 }
 
-bool SButton::IsInteractable() const
+void SButton::UpdatePressStateChanged()
 {
-	return IsEnabled();
+	UpdatePadding();
+	UpdateBorderImage();
+	UpdateForegroundColor();
 }
 
-bool SButton::ComputeVolatility() const
+bool SButton::IsInteractable() const
 {
-	// Note: we need to be careful with button volatility.  The parent SBorder class always has bound delegates to button but the following are the only thing that actually would be bound that would not be caught by an Invalidate call alone
-	return ContentPadding.IsBound();
+	return IsEnabled();
 }
 
 TEnumAsByte<EButtonClickMethod::Type> SButton::GetClickMethodFromInputType(const FPointerEvent& MouseEvent) const
@@ -463,7 +533,7 @@ FVector2D SButton::ComputeDesiredSize(float LayoutScaleMultiplier) const
 	// the border image specified by the style.
 	if (ChildSlot.GetWidget() == SNullWidget::NullWidget)
 	{
-		return GetBorder()->ImageSize;
+		return GetBorderImage()->ImageSize;
 	}
 	else
 	{
@@ -471,9 +541,9 @@ FVector2D SButton::ComputeDesiredSize(float LayoutScaleMultiplier) const
 	}
 }
 
-void SButton::SetContentPadding(const TAttribute<FMargin>& InContentPadding)
+void SButton::SetContentPadding(TAttribute<FMargin> InContentPadding)
 {
-	ContentPadding = InContentPadding;
+	ContentPaddingAttribute.Assign(*this, MoveTemp(InContentPadding));
 }
 
 void SButton::SetHoveredSound(TOptional<FSlateSound> InHoveredSound)
@@ -501,23 +571,41 @@ void SButton::SetOnUnhovered(FSimpleDelegate InOnUnhovered)
 	OnUnhovered = InOnUnhovered;
 }
 
-void SButton::SetButtonStyle(const FButtonStyle* ButtonStyle)
+void SButton::ExecuteHoverStateChanged(bool bPlaySound)
 {
-	/* Get pointer to the button style */
-	Style = ButtonStyle;
+	if (IsHovered())
+	{
+		if (bPlaySound)
+		{
+			PlayHoverSound();
+		}
+		OnHovered.ExecuteIfBound();
+	}
+	else
+	{
+		OnUnhovered.ExecuteIfBound();
+	}
+}
 
-	NormalImage = &Style->Normal;
-	HoverImage = &Style->Hovered;
-	PressedImage = &Style->Pressed;
-	DisabledImage = &Style->Disabled;
+void SButton::SetButtonStyle(const FButtonStyle* InButtonStyle)
+{
+	if (InButtonStyle == nullptr)
+	{
+		ensureAlwaysMsgf(false, TEXT("The Style is not valid."));
+		return;
+	}
 
-	BorderPadding = Style->NormalPadding;
-	PressedBorderPadding = Style->PressedPadding;
+	/* Get pointer to the button Style */
+	Style = InButtonStyle;
 
 	HoveredSound = Style->HoveredSlateSound;
 	PressedSound = Style->PressedSlateSound;
 
-	Invalidate(EInvalidateWidget::Layout);
+	UpdatePadding();
+	UpdateShowDisabledEffect();
+	UpdateBorderImage(); // Must be after UpdateShowDisabledEffect()
+	UpdateForegroundColor();
+	UpdateDisabledForegroundColor();
 }
 
 void SButton::SetClickMethod(EButtonClickMethod::Type InClickMethod)
@@ -537,30 +625,7 @@ void SButton::SetPressMethod(EButtonPressMethod::Type InPressMethod)
 
 FSlateColor SButton::GetDisabledForegroundColor() const
 {
-	return Style->DisabledForeground;	
-}
-
-FSlateColor SButton::GetForegroundColor() const
-{
-	FSlateColor UserColor = GetForegroundColorAttribute().Get();
-
-	if(UserColor == FSlateColor::UseStyle())
-	{
-		if (IsPressed())
-		{
-			return Style->PressedForeground;
-		}
-		else if (IsHovered())
-		{
-			return Style->HoveredForeground;
-		}
-		else
-		{
-			return Style->NormalForeground;
-		}
-	}
-
-	return UserColor;
+	return Style->DisabledForeground;
 }
 
 #if WITH_ACCESSIBILITY
diff --git a/Engine/Source/Runtime/Slate/Public/Widgets/Input/SButton.h b/Engine/Source/Runtime/Slate/Public/Widgets/Input/SButton.h
--- a/Engine/Source/Runtime/Slate/Public/Widgets/Input/SButton.h
+++ b/Engine/Source/Runtime/Slate/Public/Widgets/Input/SButton.h
@@ -23,9 +23,10 @@
 /**
  * Slate's Buttons are clickable Widgets that can contain arbitrary widgets as its Content().
  */
-class SLATE_API SButton
-	: public SBorder
+class SLATE_API SButton : public SBorder
 {
+	SLATE_DECLARE_WIDGET(SButton, SBorder)
+
 #if WITH_ACCESSIBILITY
 	// Allow the accessible button to "click" this button
 	friend class FSlateAccessibleButton;
@@ -119,25 +120,37 @@
 
 	SLATE_END_ARGS()
 
-	SButton();
+protected:
+		SButton();
 
-	/** @return An image that represents this button's border*/
-	virtual const FSlateBrush* GetBorder() const;
+public:
+	/**
+	 * @return An image that represents this button's border
+	 */
+	UE_DEPRECATED(5.0, "GetBorder is deprecated. Use SetBorderImage or GetBorderImage")
+	virtual const FSlateBrush* GetBorder() const
+	{
+		return GetBorderImage();
+	}
 
 	/** @return the Foreground color that this widget sets; unset options if the widget does not set a foreground color */
-	virtual FSlateColor GetForegroundColor() const;
+	virtual FSlateColor GetForegroundColor() const final
+	{
+		return Super::GetForegroundColor();
+	}
 
 	/** @return the Foreground color that this widget sets when this widget or any of its ancestors are disabled; unset options if the widget does not set a foreground color */
-	virtual FSlateColor GetDisabledForegroundColor() const;
+	virtual FSlateColor GetDisabledForegroundColor() const final;
 
 	/**
 	 * Returns true if this button is currently pressed
 	 *
 	 * @return	True if pressed, otherwise false
+	 * @note IsPressed used to be virtual. Use SetAppearPressed to assign an attribute if you need to override the default behavior.
 	 */
-	virtual bool IsPressed() const
+	bool IsPressed() const
 	{
-		return bIsPressed;
+		return bIsPressed || AppearPressedAttribute.Get();
 	}
 
 	/**
@@ -148,7 +161,7 @@
 	void Construct( const FArguments& InArgs );
 
 	/** See ContentPadding attribute */
-	void SetContentPadding(const TAttribute<FMargin>& InContentPadding);
+	void SetContentPadding(TAttribute<FMargin> InContentPadding);
 
 	/** See HoveredSound attribute */
 	void SetHoveredSound(TOptional<FSlateSound> InHoveredSound);
@@ -174,59 +187,117 @@
 
 public:
 
-	// SWidget overrides
+	//~ SWidget overrides
 	virtual int32 OnPaint(const FPaintArgs& Args, const FGeometry& AllottedGeometry, const FSlateRect& MyCullingRect, FSlateWindowElementList& OutDrawElements, int32 LayerId, const FWidgetStyle& InWidgetStyle, bool bParentEnabled) const override;
 	virtual bool SupportsKeyboardFocus() const override;
 	virtual void OnFocusLost( const FFocusEvent& InFocusEvent ) override;
 	virtual FReply OnKeyDown( const FGeometry& MyGeometry, const FKeyEvent& InKeyEvent ) override;
 	virtual FReply OnKeyUp( const FGeometry& MyGeometry, const FKeyEvent& InKeyEvent ) override;
 	virtual FReply OnMouseButtonDown( const FGeometry& MyGeometry, const FPointerEvent& MouseEvent ) override;
-	virtual FReply OnMouseButtonDoubleClick( const FGeometry& InMyGeometry, const FPointerEvent& InMouseEvent ) override;
 	virtual FReply OnMouseButtonUp( const FGeometry& MyGeometry, const FPointerEvent& MouseEvent ) override;
 	virtual FReply OnMouseMove( const FGeometry& MyGeometry, const FPointerEvent& MouseEvent ) override;
 	virtual void OnMouseEnter( const FGeometry& MyGeometry, const FPointerEvent& MouseEvent ) override;
 	virtual void OnMouseLeave( const FPointerEvent& MouseEvent ) override;
 	virtual void OnMouseCaptureLost(const FCaptureLostEvent& CaptureLostEvent) override;
 	virtual bool IsInteractable() const override;
-	virtual bool ComputeVolatility() const override;
 #if WITH_ACCESSIBILITY
 	virtual TSharedRef<FSlateAccessibleWidget> CreateAccessibleWidget() override;
 #endif
-	// SWidget
+private:
+	virtual FVector2D ComputeDesiredSize(float) const override;
+	//~ SWidget
 
 protected:
+	/** Press the button */
+	virtual void Press();
+
+	/** Release the button */
+	virtual void Release();
+
+	/** Execute the "OnClicked" delegate, and get the reply */
 	FReply ExecuteOnClick();
 
 	/** @return combines the user-specified margin and the button's internal margin. */
 	FMargin GetCombinedPadding() const;
 
-	/** @return True if the disab;ed effect should be shown. */
+	/** @return True if the disabled effect should be shown. */
 	bool GetShowDisabledEffect() const;
 
-	/** Padding specified by the user; it will be combind with the button's internal padding. */
-	TAttribute<FMargin> ContentPadding;
+	/** Utility function to translate other input click methods to regular ones. */
+	TEnumAsByte<EButtonClickMethod::Type> GetClickMethodFromInputType(const FPointerEvent& MouseEvent) const;
 
-	/** Padding that accounts for the button border */
-	FMargin BorderPadding;
+	/** Utility function to determine if the incoming mouse event is for a precise tap or click */
+	bool IsPreciseTapOrClick(const FPointerEvent& MouseEvent) const;
 
-	/** Padding that accounts for the button border when pressed */
-	FMargin PressedBorderPadding;
+	/** Play the pressed sound */
+	void PlayPressedSound() const;
 
-	/** The location in screenspace the button was pressed */
-	FVector2D PressedScreenSpacePosition;
+	/** Play the hovered sound */
+	void PlayHoverSound() const;
 
-	/** Style resource for the button */
-	const FButtonStyle* Style;
+	/** Set if this button can be focused */
+	void SetIsFocusable(bool bInIsFocusable)
+	{
+		bIsFocusable = bInIsFocusable;
+	}
 
+	void ExecuteHoverStateChanged(bool bPlaySound);
+
+protected:
+	/** @return the BorderForegroundColor attribute. */
+	TSlateAttributeRef<FSlateColor> GetBorderForegroundColorAttribute() const { return TSlateAttributeRef<FSlateColor>(SharedThis(this), BorderForegroundColorAttribute); }
+
+	/** @return the ContentPadding attribute. */
+	TSlateAttributeRef<FMargin> GetContentPaddingAttribute() const { return TSlateAttributeRef<FMargin>(SharedThis(this), ContentPaddingAttribute); }
+
+	/** Set the AppearPressed look. */
+	void SetAppearPressed(TAttribute<bool> InValue)
+	{
+		AppearPressedAttribute.Assign(*this, MoveTemp(InValue));
+	}
+
+	/** @return the AppearPressed attribute. */
+	TSlateAttributeRef<bool> GetAppearPressedAttribute() const { return TSlateAttributeRef<bool>(SharedThis(this), AppearPressedAttribute); }
+
+private:
+	void UpdatePressStateChanged();
+
+	void UpdatePadding();
+	void UpdateShowDisabledEffect();
+	void UpdateBorderImage();
+	void UpdateForegroundColor();
+	void UpdateDisabledForegroundColor();
+
+protected:
+#if WITH_EDITORONLY_DATA
+	UE_DEPRECATED(5.0, "Direct access to ContentPadding is now deprecated. Use the setter or getter.")
+	FSlateDeprecatedTAttribute<FMargin> ContentPadding;
 	/** Brush resource that represents a button */
+	UE_DEPRECATED(5.0, "NormalImage is now deprecated. Use the ButtonStyle.")
 	const FSlateBrush* NormalImage;
 	/** Brush resource that represents a button when it is hovered */
+	UE_DEPRECATED(5.0, "HoverImage is now deprecated. Use the ButtonStyle.")
 	const FSlateBrush* HoverImage;
 	/** Brush resource that represents a button when it is pressed */
+	UE_DEPRECATED(5.0, "PressedImage is now deprecated. Use the ButtonStyle.")
 	const FSlateBrush* PressedImage;
 	/** Brush resource that represents a button when it is disabled */
+	UE_DEPRECATED(5.0, "DisabledImage is now deprecated. Use the ButtonStyle.")
 	const FSlateBrush* DisabledImage;
+	/** Padding that accounts for the button border */
+	UE_DEPRECATED(5.0, "BorderPadding is now deprecated. Use the ButtonStyle.")
+	FMargin BorderPadding;
+	/** Padding that accounts for the button border when pressed */
+	UE_DEPRECATED(5.0, "PressedBorderPadding is now deprecated. Use the ButtonStyle.")
+	FMargin PressedBorderPadding;
+#endif
 
+private:
+	/** The location in screenspace the button was pressed */
+	FVector2D PressedScreenSpacePosition;
+
+	/** Style resource for the button */
+	const FButtonStyle* Style;
 	/** The delegate to execute when the button is clicked */
 	FOnClicked OnClicked;
 
@@ -236,28 +307,12 @@
 	/** The delegate to execute when the button is released */
 	FSimpleDelegate OnReleased;
 
+	/** The delegate to execute when the button is hovered */
 	FSimpleDelegate OnHovered;
 
+	/** The delegate to execute when the button exit the hovered state */
 	FSimpleDelegate OnUnhovered;
 
-	/** Press the button */
-	virtual void Press();
-
-	/** Release the button */
-	virtual void Release();
-
-	/** Utility function to translate other input click methods to regular ones. */
-	TEnumAsByte<EButtonClickMethod::Type> GetClickMethodFromInputType(const FPointerEvent& MouseEvent) const;
-
-	/** Utility function to determine if the incoming mouse event is for a precise tap or click */
-	bool IsPreciseTapOrClick(const FPointerEvent& MouseEvent) const;
-
-	/** Play the pressed sound */
-	void PlayPressedSound() const;
-
-	/** Play the hovered sound */
-	void PlayHoverSound() const;
-
 	/** The Sound to play when the button is hovered  */
 	FSlateSound HoveredSound;
 
@@ -280,6 +335,10 @@
 	uint8 bIsPressed:1;
 
 private:
-
-	virtual FVector2D ComputeDesiredSize(float) const override;
+	/** Optional foreground color that will be inherited by all of this widget's contents */
+	TSlateAttribute<FSlateColor> BorderForegroundColorAttribute;
+	/** Padding specified by the user; it will be combind with the button's internal padding. */
+	TSlateAttribute<FMargin> ContentPaddingAttribute;
+	/** Optional foreground color that will be inherited by all of this widget's contents */
+	TSlateAttribute<bool> AppearPressedAttribute;
 };
diff --git a/Engine/Source/Runtime/Slate/Public/Widgets/Input/SRichTextHyperlink.h b/Engine/Source/Runtime/Slate/Public/Widgets/Input/SRichTextHyperlink.h
--- a/Engine/Source/Runtime/Slate/Public/Widgets/Input/SRichTextHyperlink.h
+++ b/Engine/Source/Runtime/Slate/Public/Widgets/Input/SRichTextHyperlink.h
@@ -41,6 +41,7 @@ public:
 		ViewModel = InViewModel;
 
 		SetHover(TAttribute<bool>::CreateSP(this, &SRichTextHyperlink::ShouldAppearHovered));
+		SetAppearPressed(TAttribute<bool>::CreateSP(this, &SRichTextHyperlink::ShouldAppearPressed));
 
 		SHyperlink::Construct(
 			SHyperlink::FArguments()
@@ -56,19 +57,19 @@ public:
 	virtual void OnMouseEnter( const FGeometry& MyGeometry, const FPointerEvent& MouseEvent ) override
 	{
 		SHyperlink::OnMouseEnter( MyGeometry, MouseEvent );
-		ViewModel->SetIsHovered( true );
+		ViewModel->SetIsHovered( IsHovered() );
 	}
 
 	virtual void OnMouseLeave( const FPointerEvent& MouseEvent ) override
 	{
 		SHyperlink::OnMouseLeave( MouseEvent );
-		ViewModel->SetIsHovered( false );
+		ViewModel->SetIsHovered( IsHovered() );
 	}
 
 	virtual FReply OnMouseButtonDown( const FGeometry& MyGeometry, const FPointerEvent& MouseEvent ) override
 	{
 		FReply Reply = SHyperlink::OnMouseButtonDown( MyGeometry, MouseEvent );
-		ViewModel->SetIsPressed( bIsPressed );
+		ViewModel->SetIsPressed( IsPressed() );
 
 		return Reply;
 	}
@@ -76,21 +77,20 @@ public:
 	virtual FReply OnMouseButtonUp( const FGeometry& MyGeometry, const FPointerEvent& MouseEvent ) override
 	{
 		FReply Reply = SHyperlink::OnMouseButtonUp( MyGeometry, MouseEvent );
-		ViewModel->SetIsPressed( bIsPressed );
+		ViewModel->SetIsPressed( IsPressed() );
 
 		return Reply;
 	}
 
-	virtual bool IsPressed() const override
-	{
-		return ViewModel->IsPressed();
-	}
-
 private:
 	bool ShouldAppearHovered() const
 	{
 		return ViewModel->IsHovered();
 	}
+	bool ShouldAppearPressed() const
+	{
+		return ViewModel->IsPressed();
+	}
 
 private:
 
