Slate: Fix division by 0 in UniformWrapPanel layout code.

diff --git a/Engine/Source/Runtime/Slate/Private/Widgets/Layout/SUniformWrapPanel.cpp b/Engine/Source/Runtime/Slate/Private/Widgets/Layout/SUniformWrapPanel.cpp
--- a/Engine/Source/Runtime/Slate/Private/Widgets/Layout/SUniformWrapPanel.cpp
+++ b/Engine/Source/Runtime/Slate/Private/Widgets/Layout/SUniformWrapPanel.cpp
@@ -149,21 +149,23 @@ FVector2D SUniformWrapPanel::ComputeDesiredSize( float ) const
 
 	if (NumVisibleChildren > 0)
 	{
-		// Try to use the currente geometry .  If the preferred width or geometry isn't avaialble
+		// Try to use the current geometry. If the preferred width or geometry isn't available
 		// then try to make a square.
 		const FVector2D& LocalSize = GetTickSpaceGeometry().GetLocalSize();
 		if (!LocalSize.IsZero()) 
 		{
-			NumColumns = FMath::FloorToInt( LocalSize.X / MaxChildDesiredSize.X );
+			NumColumns = FMath::FloorToInt(LocalSize.X / MaxChildDesiredSize.X);
 		}
 		else 
 		{
 			NumColumns = FMath::CeilToInt(FMath::Sqrt((float)NumVisibleChildren));
 		}
 
-		NumRows = FMath::CeilToInt ( (float) NumVisibleChildren / (float) NumColumns );
-
-		return FVector2D( NumColumns*MaxChildDesiredSize.X, NumRows*MaxChildDesiredSize.Y );
+		if (NumColumns > 0)
+		{
+			NumRows = FMath::CeilToInt((float)NumVisibleChildren / (float)NumColumns);
+			return FVector2D(NumColumns * MaxChildDesiredSize.X, NumRows * MaxChildDesiredSize.Y);
+		}
 	}
 
 	return FVector2D::ZeroVector;
