disable spacebar based selection in lists by default. Can be enabled by setting HandleSpaceBarSelection in list views to true

diff --git a/Engine/Source/Runtime/Slate/Public/Widgets/Views/SListView.h b/Engine/Source/Runtime/Slate/Public/Widgets/Views/SListView.h
--- a/Engine/Source/Runtime/Slate/Public/Widgets/Views/SListView.h
+++ b/Engine/Source/Runtime/Slate/Public/Widgets/Views/SListView.h
@@ -102,6 +102,7 @@ public:
 		, _NavigationScrollOffset(0.5f)
 		, _HandleGamepadEvents( true )
 		, _HandleDirectionalNavigation( true )
+		, _HandleSpacebarSelection(false)
 		, _IsFocusable(true)
 		, _ReturnFocusToSelection()
 		, _OnItemToString_Debug()
@@ -166,6 +167,8 @@ public:
 
 		SLATE_ARGUMENT( bool, HandleDirectionalNavigation );
 
+		SLATE_ARGUMENT(bool, HandleSpacebarSelection);
+
 		SLATE_ATTRIBUTE(bool, IsFocusable)
 
 		SLATE_ARGUMENT(bool, ReturnFocusToSelection)
@@ -211,6 +214,7 @@ public:
 
 		this->bHandleGamepadEvents = InArgs._HandleGamepadEvents;
 		this->bHandleDirectionalNavigation = InArgs._HandleDirectionalNavigation;
+		this->bHandleSpacebarSelection = InArgs._HandleSpacebarSelection;
 		this->IsFocusable = InArgs._IsFocusable;
 
 		this->bReturnFocusToSelection = InArgs._ReturnFocusToSelection;
@@ -381,29 +385,29 @@ public:
 			else
 			{
 				// Change selected status of item.
-				if( TListTypeTraits<ItemType>::IsPtrValid(SelectorItem) && InKeyEvent.GetKey() == EKeys::SpaceBar )
+				if (bHandleSpacebarSelection && TListTypeTraits<ItemType>::IsPtrValid(SelectorItem) && InKeyEvent.GetKey() == EKeys::SpaceBar)
 				{
-					ItemType SelectorItemDereference( TListTypeTraits<ItemType>::NullableItemTypeConvertToItemType( SelectorItem ) );
+					ItemType SelectorItemDereference(TListTypeTraits<ItemType>::NullableItemTypeConvertToItemType(SelectorItem));
 
 					// Deselect.
-					if( InKeyEvent.IsControlDown() || SelectionMode.Get() == ESelectionMode::SingleToggle )
+					if (InKeyEvent.IsControlDown() || SelectionMode.Get() == ESelectionMode::SingleToggle)
 					{
-						this->Private_SetItemSelection( SelectorItemDereference, !( this->Private_IsItemSelected( SelectorItemDereference ) ), true );
-						this->Private_SignalSelectionChanged( ESelectInfo::OnKeyPress );
+						this->Private_SetItemSelection(SelectorItemDereference, !(this->Private_IsItemSelected(SelectorItemDereference)), true);
+						this->Private_SignalSelectionChanged(ESelectInfo::OnKeyPress);
 						bWasHandled = true;
 					}
 					else
 					{
 						// Already selected, don't handle.
-						if( this->Private_IsItemSelected( SelectorItemDereference ) )
+						if (this->Private_IsItemSelected(SelectorItemDereference))
 						{
 							bWasHandled = false;
 						}
 						// Select.
 						else
 						{
-							this->Private_SetItemSelection( SelectorItemDereference, true, true );
-							this->Private_SignalSelectionChanged( ESelectInfo::OnKeyPress );
+							this->Private_SetItemSelection(SelectorItemDereference, true, true);
+							this->Private_SignalSelectionChanged(ESelectInfo::OnKeyPress);
 							bWasHandled = true;
 						}
 					}
@@ -411,8 +415,8 @@ public:
 					RangeSelectionStart = SelectorItem;
 
 					// If the selector is not in the view, scroll it into view.
-					TSharedPtr<ITableRow> WidgetForItem = this->WidgetGenerator.GetWidgetForItem( SelectorItemDereference );
-					if ( !WidgetForItem.IsValid() )
+					TSharedPtr<ITableRow> WidgetForItem = this->WidgetGenerator.GetWidgetForItem(SelectorItemDereference);
+					if (!WidgetForItem.IsValid())
 					{
 						this->RequestScrollIntoView(SelectorItemDereference, InKeyEvent.GetUserIndex());
 					}
@@ -2184,6 +2188,9 @@ protected:
 	/** Should directional nav be supported */
 	bool bHandleDirectionalNavigation;
 
+	/** Should space bar based selection be supported */
+	bool bHandleSpacebarSelection = false;
+
 	/** If true, the focus will be returned to the last selected object in a list when navigated to. */
 	bool bReturnFocusToSelection;
 
