Slate: Convert SCanvas TAttribute to TSlateAttribute.

diff --git a/Engine/Source/Editor/CurveEditor/Private/SCurveEditorViewContainer.cpp b/Engine/Source/Editor/CurveEditor/Private/SCurveEditorViewContainer.cpp
--- a/Engine/Source/Editor/CurveEditor/Private/SCurveEditorViewContainer.cpp
+++ b/Engine/Source/Editor/CurveEditor/Private/SCurveEditorViewContainer.cpp
@@ -53,7 +53,7 @@ FVector2D SCurveEditorViewContainer::ComputeDesiredSize(float) const
 			// That will be the desired width of the whole panel.
 			MyDesiredSize.X = FMath::Max(MyDesiredSize.X, ChildDesiredSize.X + SlotPadding.GetTotalSpaceAlong<Orient_Horizontal>());
 
-			if (Child.GetSizeParam().SizeRule == FSizeParam::SizeRule_Stretch)
+			if (Child.GetSizeRule() == FSizeParam::SizeRule_Stretch)
 			{
 				++NumStretchPanels;
 			}
diff --git a/Engine/Source/Runtime/Slate/Private/Widgets/Layout/SGridPanel.cpp b/Engine/Source/Runtime/Slate/Private/Widgets/Layout/SGridPanel.cpp
--- a/Engine/Source/Runtime/Slate/Private/Widgets/Layout/SGridPanel.cpp
+++ b/Engine/Source/Runtime/Slate/Private/Widgets/Layout/SGridPanel.cpp
@@ -7,6 +7,13 @@
 #include "Layout/LayoutUtils.h"
 
 
+SLATE_IMPLEMENT_WIDGET(SGridPanel)
+void SGridPanel::PrivateRegisterAttributes(FSlateAttributeInitializer& AttributeInitializer)
+{
+	FSlateWidgetSlotAttributeInitializer Initializer = SLATE_ADD_PANELCHILDREN_DEFINITION(AttributeInitializer, Slots);
+	FSlot::RegisterAttributes(Initializer);
+}
+
 void SGridPanel::FSlot::Construct(const FChildren& SlotOwner, FSlotArguments&& InArgs)
 {
 	TBasicLayoutWidgetSlot<FSlot>::Construct(SlotOwner, MoveTemp(InArgs));
@@ -37,7 +44,7 @@ void SGridPanel::FSlot::Construct(const FChildren& SlotOwner, FSlotArguments&& I
 }
 
 SGridPanel::SGridPanel()
-	: Slots(this)
+	: Slots(this, GET_MEMBER_NAME_CHECKED(SGridPanel, Slots))
 {
 	SetCanTick(false);
 }
diff --git a/Engine/Source/Runtime/Slate/Private/Widgets/Layout/SResponsiveGridPanel.cpp b/Engine/Source/Runtime/Slate/Private/Widgets/Layout/SResponsiveGridPanel.cpp
--- a/Engine/Source/Runtime/Slate/Private/Widgets/Layout/SResponsiveGridPanel.cpp
+++ b/Engine/Source/Runtime/Slate/Private/Widgets/Layout/SResponsiveGridPanel.cpp
@@ -4,13 +4,20 @@
 #include "Layout/ArrangedChildren.h"
 #include "Layout/LayoutUtils.h"
 
+SLATE_IMPLEMENT_WIDGET(SResponsiveGridPanel)
+void SResponsiveGridPanel::PrivateRegisterAttributes(FSlateAttributeInitializer& AttributeInitializer)
+{
+	FSlateWidgetSlotAttributeInitializer Initializer = SLATE_ADD_PANELCHILDREN_DEFINITION(AttributeInitializer, Slots);
+	FSlot::RegisterAttributes(Initializer);
+}
+
 /**
  *  !!!!!!!!!!!!!!!!!   EXPERIMENTAL  !!!!!!!!!!!!!!!!!   
  * The SResponsiveGridPanel is still in development and the API may change drastically in the future
  * or maybe removed entirely.
  */
 SResponsiveGridPanel::SResponsiveGridPanel()
-: Slots(this)
+: Slots(this, GET_MEMBER_NAME_CHECKED(SResponsiveGridPanel, Slots))
 {}
 
 SResponsiveGridPanel::FScopedWidgetSlotArguments SResponsiveGridPanel::AddSlot(int32 Row)
diff --git a/Engine/Source/Runtime/Slate/Private/Widgets/Layout/SWrapBox.cpp b/Engine/Source/Runtime/Slate/Private/Widgets/Layout/SWrapBox.cpp
--- a/Engine/Source/Runtime/Slate/Private/Widgets/Layout/SWrapBox.cpp
+++ b/Engine/Source/Runtime/Slate/Private/Widgets/Layout/SWrapBox.cpp
@@ -3,8 +3,20 @@
 #include "Widgets/Layout/SWrapBox.h"
 #include "Layout/LayoutUtils.h"
 
+SLATE_IMPLEMENT_WIDGET(SWrapBox)
+void SWrapBox::PrivateRegisterAttributes(FSlateAttributeInitializer& AttributeInitializer)
+{
+	FSlateWidgetSlotAttributeInitializer SlotInitializer = SLATE_ADD_PANELCHILDREN_DEFINITION(AttributeInitializer, Slots);
+	FSlot::RegisterAttributes(SlotInitializer);
+
+	SLATE_ADD_MEMBER_ATTRIBUTE_DEFINITION(AttributeInitializer, PreferredSize, EInvalidateWidgetReason::Layout);
+	SLATE_ADD_MEMBER_ATTRIBUTE_DEFINITION(AttributeInitializer, HAlign, EInvalidateWidgetReason::Layout);
+}
+
 SWrapBox::SWrapBox()
-	: Slots(this)
+	: PreferredSize(*this, 100.f)
+	, HAlign(*this, EHorizontalAlignment::HAlign_Left)
+	, Slots(this, GET_MEMBER_NAME_CHECKED(SWrapBox, Slots))
 {
 }
 
@@ -18,30 +30,34 @@ int32 SWrapBox::RemoveSlot( const TSharedRef<SWidget>& SlotWidget )
 	return Slots.Remove(SlotWidget);
 }
 
+PRAGMA_DISABLE_DEPRECATION_WARNINGS
 void SWrapBox::Construct( const FArguments& InArgs )
 {
-	PreferredSize = InArgs._PreferredSize;
+	PreferredSize.Assign(*this, InArgs._PreferredSize);
 
 	// Handle deprecation of PreferredWidth
-	if (!PreferredSize.IsSet() && !PreferredSize.IsBound())
+	if (!InArgs._PreferredSize.IsSet() && !InArgs._PreferredSize.IsBound())
 	{
-		PreferredSize = InArgs._PreferredWidth;
+		PreferredSize.Assign(*this, InArgs._PreferredWidth);
 	}
 
 	InnerSlotPadding = InArgs._InnerSlotPadding;
 	bUseAllottedSize = InArgs._UseAllottedSize || InArgs._UseAllottedWidth;
 	Orientation = InArgs._Orientation;
-	HAlign = InArgs._HAlign;
+	HAlign.Assign(*this, InArgs._HAlign);
 
 	// Build the children from the declaration to the widget
 	Slots.AddSlots(MoveTemp(const_cast<TArray<FSlot::FSlotArguments>&>(InArgs._Slots)));
+
+	SetCanTick(bUseAllottedSize);
 }
+PRAGMA_ENABLE_DEPRECATION_WARNINGS
 
 void SWrapBox::Tick( const FGeometry& AllottedGeometry, const double InCurrentTime, const float InDeltaTime )
 {
 	if (bUseAllottedSize)
 	{
-		PreferredSize = Orientation == EOrientation::Orient_Vertical ? AllottedGeometry.GetLocalSize().Y : AllottedGeometry.GetLocalSize().X;
+		PreferredSize.Set(*this, Orientation == EOrientation::Orient_Vertical ? AllottedGeometry.GetLocalSize().Y : AllottedGeometry.GetLocalSize().X);
 	}
 }
 
@@ -408,7 +424,7 @@ FVector2D SWrapBox::ComputeDesiredSize( float ) const
 
 FChildren* SWrapBox::GetChildren()
 {
-	return &Slots;	
+	return &Slots;
 }
 
 void SWrapBox::SetInnerSlotPadding(FVector2D InInnerSlotPadding)
@@ -420,23 +436,19 @@ void SWrapBox::SetInnerSlotPadding(FVector2D InInnerSlotPadding)
 	}
 }
 
-void SWrapBox::SetWrapWidth(const TAttribute<float>& InWrapWidth)
+void SWrapBox::SetWrapWidth(TAttribute<float> InWrapWidth)
 {
-	SetAttribute(PreferredSize, InWrapWidth, EInvalidateWidgetReason::Layout);
+	SetWrapSize(MoveTemp(InWrapWidth));
 }
 
-void SWrapBox::SetWrapSize(const TAttribute<float>& InWrapSize)
+void SWrapBox::SetWrapSize(TAttribute<float> InWrapSize)
 {
-	SetAttribute(PreferredSize, InWrapSize, EInvalidateWidgetReason::Layout);
+	PreferredSize.Assign(*this, MoveTemp(InWrapSize));
 }
 
 void SWrapBox::SetUseAllottedWidth(bool bInUseAllottedWidth)
 {
-	if (bUseAllottedSize != bInUseAllottedWidth)
-	{
-		bUseAllottedSize = bInUseAllottedWidth;
-		Invalidate(EInvalidateWidgetReason::Layout);
-	}
+	SetUseAllottedSize(bInUseAllottedWidth);
 }
 
 void SWrapBox::SetUseAllottedSize(bool bInUseAllottedSize)
@@ -444,6 +456,7 @@ void SWrapBox::SetUseAllottedSize(bool bInUseAllottedSize)
 	if (bUseAllottedSize != bInUseAllottedSize)
 	{
 		bUseAllottedSize = bInUseAllottedSize;
+		SetCanTick(bUseAllottedSize);
 		Invalidate(EInvalidateWidgetReason::Layout);
 	}
 }
@@ -459,5 +472,5 @@ void SWrapBox::SetOrientation(EOrientation InOrientation)
 
 void SWrapBox::SetHorizontalAlignment(TAttribute<EHorizontalAlignment> InHAlignment)
 {
-	SetAttribute(HAlign, InHAlignment, EInvalidateWidgetReason::Layout);
+	HAlign.Assign(*this, MoveTemp(InHAlignment));
 }
\ No newline at end of file
diff --git a/Engine/Source/Runtime/Slate/Private/Widgets/SCanvas.cpp b/Engine/Source/Runtime/Slate/Private/Widgets/SCanvas.cpp
--- a/Engine/Source/Runtime/Slate/Private/Widgets/SCanvas.cpp
+++ b/Engine/Source/Runtime/Slate/Private/Widgets/SCanvas.cpp
@@ -5,25 +5,39 @@
 #include "Layout/ArrangedChildren.h"
 
 
-PRAGMA_DISABLE_DEPRECATION_WARNINGS
+SLATE_IMPLEMENT_WIDGET(SCanvas)
+void SCanvas::PrivateRegisterAttributes(FSlateAttributeInitializer& AttributeInitializer)
+{
+	FSlateWidgetSlotAttributeInitializer Initializer = SLATE_ADD_PANELCHILDREN_DEFINITION(AttributeInitializer, Children);
+	FSlot::RegisterAttributes(Initializer);
+}
+
+void SCanvas::FSlot::RegisterAttributes(FSlateWidgetSlotAttributeInitializer& AttributeInitializer)
+{
+	TWidgetSlotWithAttributeSupport::RegisterAttributes(AttributeInitializer);
+	SLATE_ADD_SLOT_ATTRIBUTE_DEFINITION_WITH_NAME(FSlot, AttributeInitializer, "Slot.Size", Size, EInvalidateWidgetReason::Paint);
+	SLATE_ADD_SLOT_ATTRIBUTE_DEFINITION_WITH_NAME(FSlot, AttributeInitializer, "Slot.Position", Position, EInvalidateWidgetReason::Paint)
+		.UpdatePrerequisite("Slot.Size");
+}
+
 void SCanvas::FSlot::Construct(const FChildren& SlotOwner, FSlotArguments&& InArg)
 {
-	TSlotBase<FSlot>::Construct(SlotOwner, MoveTemp(InArg));
-	if (InArg._Position.IsSet())
+	TWidgetSlotWithAttributeSupport<FSlot>::Construct(SlotOwner, MoveTemp(InArg));
+	if (InArg._Size.IsSet())
 	{
-		PositionAttr = MoveTemp(InArg._Position);
+		Size.Assign(*this, MoveTemp(InArg._Size));
 	}
-	if (InArg._Size.IsSet())
+	if (InArg._Position.IsSet())
 	{
-		SizeAttr = MoveTemp(InArg._Size);
+		Position.Assign(*this, MoveTemp(InArg._Position));
 	}
+
 	TAlignmentWidgetSlotMixin<FSlot>::ConstructMixin(SlotOwner, MoveTemp(InArg));
 }
-PRAGMA_ENABLE_DEPRECATION_WARNINGS
 
 
 SCanvas::SCanvas()
-	: Children(this)
+	: Children(this, GET_MEMBER_NAME_CHECKED(SCanvas, Children))
 {
 	SetCanTick(false);
 	bCanSupportFocus = false;
diff --git a/Engine/Source/Runtime/Slate/Public/Widgets/Input/SSegmentedControl.h b/Engine/Source/Runtime/Slate/Public/Widgets/Input/SSegmentedControl.h
--- a/Engine/Source/Runtime/Slate/Public/Widgets/Input/SSegmentedControl.h
+++ b/Engine/Source/Runtime/Slate/Public/Widgets/Input/SSegmentedControl.h
@@ -106,6 +106,7 @@ public:
 		TAttribute<const FSlateBrush*> _Icon;
 
 		OptionType _Value;
+		TWeakPtr<SCheckBox> _CheckBox;
 	};
 
 	static typename FSlot::FSlotArguments Slot(const OptionType& InValue)
@@ -144,6 +145,7 @@ public:
 
 	SSegmentedControl()
 		: Children(this)
+		, CurrentValue(*this)
 	{}
 
 	void Construct( const FArguments& InArgs )
@@ -153,7 +155,8 @@ public:
 		Style = InArgs._Style;
 		TextStyle = InArgs._TextStyle;
 
-		CurrentValue = InArgs._Value;
+		CurrentValueIsBound = InArgs._Value.IsBound();
+		CurrentValue.Assign(*this, InArgs._Value);
 		OnValueChanged = InArgs._OnValueChanged;
 
 		UniformPadding = InArgs._UniformPadding;
@@ -165,15 +168,24 @@ public:
 
 	void RebuildChildren()
 	{
-		TSharedPtr<SUniformGridPanel> UniformBox =
-			SNew(SUniformGridPanel);
+		TSharedPtr<SUniformGridPanel> UniformBox = SNew(SUniformGridPanel);
 
 		const int32 NumSlots = Children.Num();
 		for ( int32 SlotIndex = 0; SlotIndex < NumSlots; ++SlotIndex )
 		{
 			TSharedRef<SWidget> Child = Children[SlotIndex].GetWidget();
 			FSlot* ChildSlotPtr = &Children[SlotIndex];
-			OptionType& ChildValue = ChildSlotPtr->_Value;
+			const OptionType ChildValue = ChildSlotPtr->_Value;
+
+			TAttribute<FVector2D> SpacerLambda = FVector::ZeroVector;
+			if (ChildSlotPtr->_Icon.IsBound() || ChildSlotPtr->_Text.IsBound())
+			{
+				SpacerLambda = MakeAttributeLambda([ChildSlotPtr]() { return (ChildSlotPtr->_Icon.Get() != nullptr && !ChildSlotPtr->_Text.Get().IsEmpty()) ? FVector2D(8.0f, 1.0f) : FVector2D::ZeroVector; });
+			}
+			else
+			{
+				SpacerLambda = (ChildSlotPtr->_Icon.Get() != nullptr && !ChildSlotPtr->_Text.Get().IsEmpty()) ? FVector2D(8.0f, 1.0f) : FVector2D::ZeroVector;
+			}
 
 			if (Child == SNullWidget::NullWidget)
 			{
@@ -192,7 +204,7 @@ public:
 				.AutoWidth()
 				[
 					SNew(SSpacer)
-					.Size_Lambda( [ChildSlotPtr] () { return (ChildSlotPtr->_Icon.Get() != nullptr && !ChildSlotPtr->_Text.Get().IsEmpty()) ? FVector2D(8.0f, 1.0f) : FVector2D::ZeroVector; })
+					.Size(SpacerLambda)
 				]
 
 				+SHorizontalBox::Slot()
@@ -212,12 +224,12 @@ public:
 			// Note HAlignment is applied at the check box level because if it were applied here it would make the slots look physically disconnected from each other 
 			.VAlign(ChildSlotPtr->GetVerticalAlignment())
 			[
-				SNew(SCheckBox)
+				SAssignNew(ChildSlotPtr->_CheckBox, SCheckBox)
 				.Clipping(EWidgetClipping::ClipToBounds)
 				.HAlign(ChildSlotPtr->GetHorizontalAlignment())
 				.ToolTipText(ChildSlotPtr->_Tooltip)
 				.Style(ColumnIndex == 0 ? &Style->FirstControlStyle : ColumnIndex == (NumSlots - 1) ? &Style->LastControlStyle : &Style->ControlStyle)
-				.IsChecked(this, &SSegmentedControl::IsCurrentValue, ChildValue)
+				.IsChecked(GetCheckBoxStateAttribute(ChildValue))
 				.OnCheckStateChanged(this, &SSegmentedControl::CommitValue, ChildValue)
 				.Padding(UniformPadding)
 				[
@@ -259,39 +271,74 @@ public:
 		return Children.Num();
 	}
 
+	OptionType GetValue() const
+	{
+		return CurrentValue.Get();
+	}
+
 	/** See the Value attribute */
-	OptionType& GetValue() const { return CurrentValue.Get(); }
-	void SetValue(const TAttribute<OptionType>& InValue) 
+	void SetValue(TAttribute<OptionType> InValue) 
 	{
-		CurrentValue = InValue; 
+		CurrentValueIsBound = InValue.IsBound();
+		CurrentValue.Assign(*this, MoveTemp(InValue));
+
+		if (!CurrentValueIsBound)
+		{
+			for (int32 Index = 0; Index < Children.Num(); ++Index)
+			{
+				const FSlot& Slot = Children[Index];
+				if (TSharedPtr<SCheckBox> CheckBox = Slot._CheckBox.Pin())
+				{
+					CheckBox->SetIsChecked(Slot._Value == CurrentValue.Get() ? ECheckBoxState::Checked : ECheckBoxState::Unchecked);
+				}
+			}
+		}
 	}
 
-	ECheckBoxState IsCurrentValue(OptionType InValue) const
+private:
+	TAttribute<ECheckBoxState> GetCheckBoxStateAttribute(OptionType InValue) const
 	{
-		return InValue == CurrentValue.Get() ? ECheckBoxState::Checked : ECheckBoxState::Unchecked;
-	}	
+		auto Lambda = [this, InValue]()
+		{
+			return InValue == CurrentValue.Get() ? ECheckBoxState::Checked : ECheckBoxState::Unchecked;
+		};
+
+		if (CurrentValueIsBound)
+		{
+			return MakeAttributeLambda(Lambda);
+		}
+
+		return Lambda();
+	}
 
-	void CommitValue(const ECheckBoxState InCheckState, OptionType InValue) 
+	void CommitValue(const ECheckBoxState InCheckState, OptionType InValue)
 	{
 		if (InCheckState == ECheckBoxState::Checked)
 		{
-			// don't overwrite the bound attribute, but still notify that a
-			if (!CurrentValue.IsBound())
+			// don't overwrite the bound attribute, but still notify that the value was committed.
+			if (!CurrentValueIsBound)
 			{
-				CurrentValue.Set(InValue);
+				CurrentValue.Set(*this, InValue);
+				for (int32 Index = 0; Index < Children.Num(); ++Index)
+				{
+					const FSlot& Slot = Children[Index];
+					if (TSharedPtr<SCheckBox> CheckBox = Slot._CheckBox.Pin())
+					{
+						CheckBox->SetIsChecked(Slot._Value == CurrentValue.Get() ? ECheckBoxState::Checked : ECheckBoxState::Unchecked);
+					}
+			    }
 			}
 
 			OnValueChanged.ExecuteIfBound( InValue );
 		}
 	}
 
-
-protected:
+private:
 	TPanelChildren<FSlot> Children;
 
 	FOnValueChanged OnValueChanged;
 
-	TAttribute<OptionType> CurrentValue;
+	TSlateAttribute<OptionType, EInvalidateWidgetReason::Paint> CurrentValue;
 
 	TAttribute<FMargin> UniformPadding;
 
@@ -300,4 +347,6 @@ protected:
 	const FTextBlockStyle* TextStyle;
 
 	int32 MaxSegmentsPerLine = 0;
+
+	bool CurrentValueIsBound = false;
 };
diff --git a/Engine/Source/Runtime/Slate/Public/Widgets/Layout/SGridPanel.h b/Engine/Source/Runtime/Slate/Public/Widgets/Layout/SGridPanel.h
--- a/Engine/Source/Runtime/Slate/Public/Widgets/Layout/SGridPanel.h
+++ b/Engine/Source/Runtime/Slate/Public/Widgets/Layout/SGridPanel.h
@@ -15,9 +15,10 @@ class FArrangedChildren;
 class FPaintArgs;
 class FSlateWindowElementList;
 
-class SLATE_API SGridPanel
-	: public SPanel
+class SLATE_API SGridPanel : public SPanel
 {
+	SLATE_DECLARE_WIDGET(SGridPanel, SPanel)
+
 public:
 	// Used by the mandatory named parameter in FSlot
 	class Layer
diff --git a/Engine/Source/Runtime/Slate/Public/Widgets/Layout/SResponsiveGridPanel.h b/Engine/Source/Runtime/Slate/Public/Widgets/Layout/SResponsiveGridPanel.h
--- a/Engine/Source/Runtime/Slate/Public/Widgets/Layout/SResponsiveGridPanel.h
+++ b/Engine/Source/Runtime/Slate/Public/Widgets/Layout/SResponsiveGridPanel.h
@@ -33,9 +33,9 @@ namespace SResponsiveGridSize
  * The SResponsiveGridPanel is still in development and the API may change drastically in the future
  * or maybe removed entirely.
  */
-class SLATE_API SResponsiveGridPanel
-	: public SPanel
+class SLATE_API SResponsiveGridPanel : public SPanel
 {
+	SLATE_DECLARE_WIDGET(SResponsiveGridPanel, SPanel)
 public:
 
 	class FSlot : public TBasicLayoutWidgetSlot< FSlot >
@@ -159,7 +159,7 @@ public:
 
 public:
 
-	// SWidget interface
+	//~ SWidget interface
 
 	virtual int32 OnPaint( const FPaintArgs& Args, const FGeometry& AllottedGeometry, const FSlateRect& MyCullingRect, FSlateWindowElementList& OutDrawElements, int32 LayerId, const FWidgetStyle& InWidgetStyle, bool bParentEnabled ) const override;
 	virtual void OnArrangeChildren( const FGeometry& AllottedGeometry, FArrangedChildren& ArrangedChildren ) const override;
@@ -167,9 +167,9 @@ public:
 	virtual FChildren* GetChildren() override;
 
 protected:
-	// Begin SWidget overrides.
+	//~ Begin SWidget overrides.
 	virtual FVector2D ComputeDesiredSize(float) const override;
-	// End SWidget overrides.
+	//~ End SWidget overrides.
 
 private:
 
diff --git a/Engine/Source/Runtime/Slate/Public/Widgets/Layout/SWrapBox.h b/Engine/Source/Runtime/Slate/Public/Widgets/Layout/SWrapBox.h
--- a/Engine/Source/Runtime/Slate/Public/Widgets/Layout/SWrapBox.h
+++ b/Engine/Source/Runtime/Slate/Public/Widgets/Layout/SWrapBox.h
@@ -49,6 +49,7 @@ class FArrangedChildren;
 
 class SLATE_API SWrapBox : public SPanel
 {
+	SLATE_DECLARE_WIDGET(SWrapBox, SPanel)
 public:
 
 	/** A slot that support alignment of content and padding */
@@ -136,9 +137,9 @@ public:
 	PRAGMA_ENABLE_DEPRECATION_WARNINGS
 
 
+	PRAGMA_DISABLE_DEPRECATION_WARNINGS
 	SLATE_BEGIN_ARGS(SWrapBox)
-		: _PreferredWidth(100.f)
-		, _PreferredSize(100.f)
+		: _PreferredSize(100.f)
 		, _HAlign(HAlign_Left)
 		, _InnerSlotPadding(FVector2D::ZeroVector)
 		, _UseAllottedWidth(false)
@@ -152,6 +153,7 @@ public:
 		SLATE_SLOT_ARGUMENT( FSlot, Slots )
 
 		/** The preferred width, if not set will fill the space */
+		UE_DEPRECATED(5.0, "PreferredWidth is deprecated. Use PreferredSize instead.")
 		SLATE_ATTRIBUTE( float, PreferredWidth )
 
 		/** The preferred size, if not set will fill the space */
@@ -164,14 +166,16 @@ public:
 		SLATE_ARGUMENT( FVector2D, InnerSlotPadding )
 
 		/** if true, the PreferredWidth will always match the room available to the SWrapBox  */
+		UE_DEPRECATED(5.0, "UseAllottedWidth is deprecated. Use UseAllottedSize instead.")
 		SLATE_ARGUMENT( bool, UseAllottedWidth )
 
 		/** if true, the PreferredSize will always match the room available to the SWrapBox  */
 		SLATE_ARGUMENT( bool, UseAllottedSize )
 
 		/** Determines if the wrap box needs to arrange the slots left-to-right or top-to-bottom.*/
-		SLATE_ARGUMENT_DEFAULT(EOrientation, Orientation) { EOrientation::Orient_Horizontal };
+		SLATE_ARGUMENT(EOrientation, Orientation);
 	SLATE_END_ARGS()
+	PRAGMA_ENABLE_DEPRECATION_WARNINGS
 
 	SWrapBox();
 
@@ -207,10 +211,10 @@ public:
 
 	/** Set the width at which the wrap panel should wrap its content. */
 	UE_DEPRECATED(4.26, "Deprecated, please use SetWrapSize() instead")
-	void SetWrapWidth(const TAttribute<float>& InWrapWidth);
+	void SetWrapWidth(TAttribute<float> InWrapWidth);
 
 	/** Set the size at which the wrap panel should wrap its content. */
-	void SetWrapSize( const TAttribute<float>& InWrapSize );
+	void SetWrapSize(TAttribute<float> InWrapSize );
 
 	/** When true, use the WrapWidth property to determine where to wrap to the next line. */
 	UE_DEPRECATED(4.26, "Deprecated, please use SetUseAllottedSize() instead")
@@ -226,16 +230,11 @@ public:
 	void SetHorizontalAlignment(TAttribute<EHorizontalAlignment> InHAlignment);
 
 private:
-
-	/** How wide this panel should appear to be. Any widgets past this line will be wrapped onto the next line. */
-	UE_DEPRECATED(4.26, "Deprecated, please use PreferredSize instead")
-	TAttribute<float> PreferredWidth;
-
 	/** How wide or long, dependently of the orientation, this panel should appear to be. Any widgets past this line will be wrapped onto the next line. */
-	TAttribute<float> PreferredSize;
+	TSlateAttribute<float> PreferredSize;
 
 	/** How to distribute the elements among any extra space in a given row */
-	TAttribute< EHorizontalAlignment > HAlign;
+	TSlateAttribute< EHorizontalAlignment > HAlign;
 
 	/** The slots that contain this panel's children. */
 	TPanelChildren<FSlot> Slots;
@@ -243,15 +242,11 @@ private:
 	/** When two slots end up sharing a border, this will put that much padding between then, but otherwise wont. */
 	FVector2D InnerSlotPadding;
 
-	/** If true the box will have a preferred width equal to its alloted width  */
-	UE_DEPRECATED(4.26, "Deprecated, please use bUseAllottedSize instead")
-	bool bUseAllottedWidth;
-
 	/** If true the box will have a preferred size equal to its alloted size  */
 	bool bUseAllottedSize;
 
 	/** Determines if the wrap box needs to arrange the slots left-to-right or top-to-bottom.*/
-	EOrientation Orientation = EOrientation::Orient_Horizontal;
+	EOrientation Orientation;
 
 	class FChildArranger;
 	friend class SWrapBox::FChildArranger;
diff --git a/Engine/Source/Runtime/Slate/Public/Widgets/SCanvas.h b/Engine/Source/Runtime/Slate/Public/Widgets/SCanvas.h
--- a/Engine/Source/Runtime/Slate/Public/Widgets/SCanvas.h
+++ b/Engine/Source/Runtime/Slate/Public/Widgets/SCanvas.h
@@ -18,9 +18,10 @@ class FSlateWindowElementList;
 /**
  * Canvas is a layout widget that allows you to arbitrary position and size child widgets in a relative coordinate space
  */
-class SLATE_API SCanvas
-	: public SPanel
+class SLATE_API SCanvas : public SPanel
 {
+	SLATE_DECLARE_WIDGET(SCanvas, SPanel)
+
 public:
 
 	/**
@@ -49,8 +50,7 @@ public:
 	 * 
 	 *  Note: FILL is NOT supported.
 	 */
-	PRAGMA_DISABLE_DEPRECATION_WARNINGS
-	class SLATE_API FSlot : public TSlotBase<FSlot>, public TAlignmentWidgetSlotMixin<FSlot>
+	class SLATE_API FSlot : public TWidgetSlotWithAttributeSupport<FSlot>, public TAlignmentWidgetSlotMixin<FSlot>
 	{
 	public:
 		SLATE_SLOT_BEGIN_ARGS_OneMixin(FSlot, TSlotBase<FSlot>, TAlignmentWidgetSlotMixin<FSlot>)
@@ -60,43 +60,49 @@ public:
 
 		void SetPosition( TAttribute<FVector2D> InPosition )
 		{
-			PositionAttr = MoveTemp(InPosition);
+			Position.Assign(*this, MoveTemp(InPosition));
 		}
 		FVector2D GetPosition() const
 		{
-			return PositionAttr.Get();
+			return Position.Get();
 		}
 
 		void SetSize( TAttribute<FVector2D> InSize )
 		{
-			SizeAttr = MoveTemp(InSize);
+			Size.Assign(*this, MoveTemp(InSize));
 		}
 		FVector2D GetSize() const
 		{
-			return SizeAttr.Get();
+			return Size.Get();
 		}
 
 	public:
-		/** Position */
+#if WITH_EDITORONLY_DATA
 		UE_DEPRECATED(5.0, "Direct access to PositionAttr is now deprecated. Use the getter or setter.")
-		TAttribute<FVector2D> PositionAttr;
-
-		/** Size */
+		FSlateDeprecatedTAttribute<FVector2D> PositionAttr;
 		UE_DEPRECATED(5.0, "Direct access to SizeAttr is now deprecated. Use the getter or setter.")
-		TAttribute<FVector2D> SizeAttr;
+		FSlateDeprecatedTAttribute<FVector2D> SizeAttr;
+#endif
 
 	public:
 		/** Default values for a slot. */
 		FSlot()
-			: TSlotBase<FSlot>()
+			: TWidgetSlotWithAttributeSupport<FSlot>()
 			, TAlignmentWidgetSlotMixin<FSlot>(HAlign_Left, VAlign_Top)
-			, PositionAttr( FVector2D::ZeroVector )
-			, SizeAttr( FVector2D( 1.0f, 1.0f ) )
+			, Position(*this, FVector2D::ZeroVector)
+			, Size(*this, FVector2D(1.0f, 1.0f))
 		{ }
 
 		void Construct(const FChildren& SlotOwner, FSlotArguments&& InArg);
+		static void RegisterAttributes(FSlateWidgetSlotAttributeInitializer& AttributeInitializer);
+
+	private:
+		/** Position */
+		TSlateSlotAttribute<FVector2D> Position;
+
+		/** Size */
+		TSlateSlotAttribute<FVector2D> Size;
 	};
-	PRAGMA_ENABLE_DEPRECATION_WARNINGS
 
 	SLATE_BEGIN_ARGS( SCanvas )
 		{
@@ -140,16 +146,15 @@ public:
 
 public:
 
-	// SWidget overrides
-
+	//~ SWidget overrides
 	virtual void OnArrangeChildren( const FGeometry& AllottedGeometry, FArrangedChildren& ArrangedChildren ) const override;
 	virtual int32 OnPaint( const FPaintArgs& Args, const FGeometry& AllottedGeometry, const FSlateRect& MyCullingRect, FSlateWindowElementList& OutDrawElements, int32 LayerId, const FWidgetStyle& InWidgetStyle, bool bParentEnabled ) const override;
 	virtual FChildren* GetChildren() override;
 
 protected:
-	// Begin SWidget overrides.
+	//~ Begin SWidget overrides.
 	virtual FVector2D ComputeDesiredSize(float) const override;
-	// End SWidget overrides.
+	//~ End SWidget overrides.
 
 protected:
 
diff --git a/Engine/Source/Runtime/SlateCore/Private/Widgets/SBoxPanel.cpp b/Engine/Source/Runtime/SlateCore/Private/Widgets/SBoxPanel.cpp
--- a/Engine/Source/Runtime/SlateCore/Private/Widgets/SBoxPanel.cpp
+++ b/Engine/Source/Runtime/SlateCore/Private/Widgets/SBoxPanel.cpp
@@ -3,6 +3,23 @@
 #include "Widgets/SBoxPanel.h"
 #include "Layout/LayoutUtils.h"
 
+SLATE_IMPLEMENT_WIDGET(SBoxPanel)
+void SBoxPanel::PrivateRegisterAttributes(FSlateAttributeInitializer& AttributeInitializer)
+{
+	FSlateWidgetSlotAttributeInitializer Initializer = SLATE_ADD_PANELCHILDREN_DEFINITION(AttributeInitializer, Children);
+	FSlot::RegisterAttributes(Initializer);
+}
+
+SLATE_IMPLEMENT_WIDGET(SHorizontalBox)
+void SHorizontalBox::PrivateRegisterAttributes(FSlateAttributeInitializer& AttributeInitializer)
+{
+}
+
+SLATE_IMPLEMENT_WIDGET(SVerticalBox)
+void SVerticalBox::PrivateRegisterAttributes(FSlateAttributeInitializer& AttributeInitializer)
+{
+}
+
 void SHorizontalBox::Construct( const SHorizontalBox::FArguments& InArgs )
 {
 	Children.Reserve(InArgs._Slots.Num());
@@ -83,10 +100,10 @@ static void ArrangeChildrenAlong(EFlowDirection InLayoutFlow, const TPanelChildr
 				// All widgets contribute their margin to the fixed space requirement
 				FixedTotal += CurChild.GetPadding().template GetTotalSpaceAlong<Orientation>();
 
-				if ( CurChild.GetSizeParam().SizeRule == FSizeParam::SizeRule_Stretch )
+				if ( CurChild.GetSizeRule() == FSizeParam::SizeRule_Stretch )
 				{
 					// for stretch children we save sum up the stretch coefficients
-					StretchCoefficientTotal += CurChild.GetSizeParam().Value.Get();
+					StretchCoefficientTotal += CurChild.GetSizeValue();
 				}
 				else
 				{
@@ -129,12 +146,12 @@ static void ArrangeChildrenAlong(EFlowDirection InLayoutFlow, const TPanelChildr
 			if (ChildVisibility != EVisibility::Collapsed)
 			{
 				// The size of the widget depends on its size type
-				if (CurChild.GetSizeParam().SizeRule == FSizeParam::SizeRule_Stretch)
+				if (CurChild.GetSizeRule() == FSizeParam::SizeRule_Stretch)
 				{
 					if (StretchCoefficientTotal > 0)
 					{
 						// Stretch widgets get a fraction of the space remaining after all the fixed-space requirements are met
-						ChildSize = NonFixedSpace * CurChild.GetSizeParam().Value.Get() / StretchCoefficientTotal;
+						ChildSize = NonFixedSpace * CurChild.GetSizeValue() / StretchCoefficientTotal;
 					}
 				}
 				else
@@ -294,7 +311,7 @@ void SBoxPanel::ClearChildren()
 }
 
 SBoxPanel::SBoxPanel( EOrientation InOrientation )
-	: Children(this)
+	: Children(this, GET_MEMBER_NAME_CHECKED(SBoxPanel, Children))
 	, Orientation(InOrientation)
 {
 
diff --git a/Engine/Source/Runtime/SlateCore/Private/Widgets/SOverlay.cpp b/Engine/Source/Runtime/SlateCore/Private/Widgets/SOverlay.cpp
--- a/Engine/Source/Runtime/SlateCore/Private/Widgets/SOverlay.cpp
+++ b/Engine/Source/Runtime/SlateCore/Private/Widgets/SOverlay.cpp
@@ -6,9 +6,23 @@
 #include "Layout/LayoutUtils.h"
 #include "Rendering/DrawElements.h"
 
+SLATE_IMPLEMENT_WIDGET(SOverlay)
+void SOverlay::PrivateRegisterAttributes(FSlateAttributeInitializer& AttributeInitializer)
+{
+	FSlateWidgetSlotAttributeInitializer Initializer = SLATE_ADD_PANELCHILDREN_DEFINITION(AttributeInitializer, Children);
+	FOverlaySlot::RegisterAttributes(Initializer);
+}
+
+PRAGMA_DISABLE_DEPRECATION_WARNINGS
+void SOverlay::FOverlaySlot::Construct(const FChildren& SlotOwner, FSlotArguments&& InArgs)
+{
+	TBasicLayoutWidgetSlot<FOverlaySlot>::Construct(SlotOwner, MoveTemp(InArgs));
+	ZOrder = InArgs._ZOrder.Get(ZOrder);
+}
+PRAGMA_ENABLE_DEPRECATION_WARNINGS
 
 SOverlay::SOverlay()
-	: Children(this)
+	: Children(this, GET_MEMBER_NAME_CHECKED(SOverlay, Children))
 {
 	SetCanTick(false);
 	bCanSupportFocus = false;
diff --git a/Engine/Source/Runtime/SlateCore/Public/Layout/BasicLayoutWidgetSlot.h b/Engine/Source/Runtime/SlateCore/Public/Layout/BasicLayoutWidgetSlot.h
--- a/Engine/Source/Runtime/SlateCore/Public/Layout/BasicLayoutWidgetSlot.h
+++ b/Engine/Source/Runtime/SlateCore/Public/Layout/BasicLayoutWidgetSlot.h
@@ -223,9 +223,11 @@ template <typename MixedIntoType>
 class TPaddingWidgetSlotMixin
 {
 public:
-	TPaddingWidgetSlotMixin() = default;
+	TPaddingWidgetSlotMixin()
+		: SlotPaddingAttribute(*static_cast<MixedIntoType*>(this))
+	{}
 	TPaddingWidgetSlotMixin(const FMargin& Margin)
-		: SlotPadding(Margin)
+		: SlotPaddingAttribute(*static_cast<MixedIntoType*>(this), Margin)
 	{}
 
 public:
@@ -266,10 +268,15 @@ protected:
 	{
 		if (InArgs._Padding.IsSet())
 		{
-			SlotPadding = MoveTemp(InArgs._Padding);
+			SlotPaddingAttribute.Assign(*static_cast<MixedIntoType*>(this), MoveTemp(InArgs._Padding));
 		}
 	}
 
+	static void RegisterAttributesMixin(FSlateWidgetSlotAttributeInitializer& AttributeInitializer)
+	{
+		SLATE_ADD_SLOT_ATTRIBUTE_DEFINITION_WITH_NAME(MixedIntoType, AttributeInitializer, "Slot.Padding", SlotPaddingAttribute, EInvalidateWidgetReason::Layout);
+	}
+
 public:
 	UE_DEPRECATED(5.0, "Padding is now deprecated. Use the FSlotArgument or the SetPadding function.")
 	MixedIntoType& Padding(TAttribute<FMargin> InPadding)
@@ -302,18 +309,22 @@ public:
 public:
 	void SetPadding(TAttribute<FMargin> InPadding)
 	{
-		SlotPadding = MoveTemp(InPadding);
-		static_cast<MixedIntoType*>(this)->Invalidate(EInvalidateWidgetReason::LayoutAndVolatility);
+		SlotPaddingAttribute.Assign(static_cast<MixedIntoType&>(*this), MoveTemp(InPadding));
 	}
 
 	FMargin GetPadding() const
 	{
-		return SlotPadding.Get();
+		return SlotPaddingAttribute.Get();
 	}
 
 public:
 	UE_DEPRECATED(5.0, "Direct access to SlotPadding is now deprecated. Use the setter or getter.")
-	TAttribute<FMargin> SlotPadding;
+	FSlateDeprecatedTAttribute<FMargin> SlotPadding;
+
+private:
+	using SlotCompareType = TSlateAttributeComparePredicate<>;
+	using SlotType = ::SlateAttributePrivate::TSlateContainedAttribute<FMargin, ::SlateAttributePrivate::FSlateAttributeNoInvalidationReason, SlotCompareType>;
+	SlotType SlotPaddingAttribute;
 };
 PRAGMA_ENABLE_DEPRECATION_WARNINGS
 
@@ -486,6 +497,12 @@ public:
 		TPaddingWidgetSlotMixin<SlotType>::ConstructMixin(SlotOwner, MoveTemp(InArgs));
 		TAlignmentWidgetSlotMixin<SlotType>::ConstructMixin(SlotOwner, MoveTemp(InArgs));
 	}
+
+	static void RegisterAttributes(FSlateWidgetSlotAttributeInitializer& AttributeInitializer)
+	{
+		TWidgetSlotWithAttributeSupport<SlotType>::RegisterAttributes(AttributeInitializer);
+		TPaddingWidgetSlotMixin<SlotType>::RegisterAttributesMixin(AttributeInitializer);
+	}
 };
 
 
diff --git a/Engine/Source/Runtime/SlateCore/Public/Widgets/SBoxPanel.h b/Engine/Source/Runtime/SlateCore/Public/Widgets/SBoxPanel.h
--- a/Engine/Source/Runtime/SlateCore/Public/Widgets/SBoxPanel.h
+++ b/Engine/Source/Runtime/SlateCore/Public/Widgets/SBoxPanel.h
@@ -10,6 +10,7 @@
 #include "Widgets/DeclarativeSyntaxSupport.h"
 #include "SlotBase.h"
 #include "Widgets/SWidget.h"
+#include "Layout/BasicLayoutWidgetSlot.h"
 #include "Layout/Children.h"
 #include "Widgets/SPanel.h"
 #include "Layout/ArrangedChildren.h"
@@ -23,6 +24,8 @@ class FArrangedChildren;
  */
 class SLATECORE_API SBoxPanel : public SPanel
 {
+	SLATE_DECLARE_WIDGET(SBoxPanel, SPanel)
+
 protected:
 	/**
 	 * A BoxPanel contains one BoxPanel child and describes how that
@@ -41,35 +44,57 @@ protected:
 		/** Default values for a slot. */
 		TSlot()
 			: TBasicLayoutWidgetSlot<SlotType>(HAlign_Fill, VAlign_Fill)
-			, SizeParam(FStretch(1.f))
-			, MaxSize(0.0f)
+			, SizeRule(FSizeParam::SizeRule_Stretch)
+			, SizeValue(*this, 1.f)
+			, MaxSize(*this, 0.0f)
 		{ }
 
 	public:
 		void Construct(const FChildren& SlotOwner, FSlotArguments&& InArgs)
 		{
 			TBasicLayoutWidgetSlot<SlotType>::Construct(SlotOwner, MoveTemp(InArgs));
-			if (InArgs._SizeParam.IsSet())
+			if (InArgs._MaxSize.IsSet())
 			{
-				SizeParam = MoveTemp(InArgs._SizeParam.GetValue());
+				SetMaxSize(MoveTemp(InArgs._MaxSize));
 			}
-			if (InArgs._MaxSize.IsSet())
+			if (InArgs._SizeParam.IsSet())
 			{
-				MaxSize = MoveTemp(InArgs._MaxSize);
+				SetSizeParam(MoveTemp(InArgs._SizeParam.GetValue()));
 			}
 		}
 
-		/** Get the space param this slot should occupy along panel's direction. */
-		const FSizeParam& GetSizeParam() const { return SizeParam; }
+		static void RegisterAttributes(FSlateWidgetSlotAttributeInitializer& AttributeInitializer)
+		{
+			TBasicLayoutWidgetSlot<SlotType>::RegisterAttributes(AttributeInitializer);
+			SLATE_ADD_SLOT_ATTRIBUTE_DEFINITION_WITH_NAME(TSlot<SlotType>, AttributeInitializer, "Slot.MaxSize", MaxSize, EInvalidateWidgetReason::Layout);
+			SLATE_ADD_SLOT_ATTRIBUTE_DEFINITION_WITH_NAME(TSlot<SlotType>, AttributeInitializer, "Slot.SizeValue", SizeValue, EInvalidateWidgetReason::Layout)
+				.UpdatePrerequisite("Slot.MaxSize");
+		}
+
+		/** Get the space rule this slot should occupy along panel's direction. */
+		FSizeParam::ESizeRule GetSizeRule() const
+		{
+			return SizeRule;
+		}
+
+		/** Get the space rule value this slot should occupy along panel's direction. */
+		float GetSizeValue() const
+		{
+			return SizeValue.Get();
+		}
+
 		/** Get the max size the slot can be.*/
-		float GetMaxSize() const { return MaxSize.Get(); }
+		float GetMaxSize() const
+		{
+			return MaxSize.Get();
+		}
 
 	public:
 		/** Set the size Param of the slot, It could be a FStretch or a FAuto. */
 		void SetSizeParam(FSizeParam InSizeParam)
 		{
-			SizeParam = MoveTemp(InSizeParam);
-			FSlotBase::Invalidate(EInvalidateWidget::Layout);
+			SizeRule = InSizeParam.SizeRule;
+			SizeValue.Assign(*this, MoveTemp(InSizeParam.Value));
 		}
 
 		/** The widget's DesiredSize will be used as the space required. */
@@ -87,8 +112,7 @@ protected:
 		/** Set the max size in SlateUnit this slot can be. */
 		void SetMaxSize(TAttribute<float> InMaxSize)
 		{
-			MaxSize = MoveTemp(InMaxSize);
-			FSlotBase::Invalidate(EInvalidateWidget::Layout);
+			MaxSize.Assign(*this, MoveTemp(InMaxSize));
 		}
 
 	private:
@@ -99,10 +123,15 @@ protected:
 		 * peer Widgets depending on the Value property. Available space is space remaining after all the
 		 * peers' SizeRule_Auto requirements have been satisfied.
 		 */
-		FSizeParam SizeParam;
+
+		/** The sizing rule to use. */
+		FSizeParam::ESizeRule SizeRule;
+
+		/** The actual value this size parameter stores. */
+		typename TBasicLayoutWidgetSlot<SlotType>::template TSlateSlotAttribute<float> SizeValue;
 
 		/** The max size that this slot can be (0 if no max) */
-		TAttribute<float> MaxSize;
+		typename TBasicLayoutWidgetSlot<SlotType>::template TSlateSlotAttribute<float> MaxSize;
 	};
 
 public:
@@ -192,6 +221,7 @@ protected:
 /** A Horizontal Box Panel. See SBoxPanel for more info. */
 class SLATECORE_API SHorizontalBox : public SBoxPanel
 {
+	SLATE_DECLARE_WIDGET(SHorizontalBox, SBoxPanel)
 public:
 	class FSlot : public SBoxPanel::TSlot<FSlot>
 	{
@@ -294,6 +324,7 @@ public:
 /** A Vertical Box Panel. See SBoxPanel for more info. */
 class SLATECORE_API SVerticalBox : public SBoxPanel
 {
+	SLATE_DECLARE_WIDGET(SVerticalBox, SBoxPanel)
 public:
 	class FSlot : public SBoxPanel::TSlot<FSlot>
 	{
diff --git a/Engine/Source/Runtime/SlateCore/Public/Widgets/SOverlay.h b/Engine/Source/Runtime/SlateCore/Public/Widgets/SOverlay.h
--- a/Engine/Source/Runtime/SlateCore/Public/Widgets/SOverlay.h
+++ b/Engine/Source/Runtime/SlateCore/Public/Widgets/SOverlay.h
@@ -31,9 +31,9 @@ class FSlateWindowElementList;
  *
  *		Note that SWidget3 will be drawn on top of SWidget2 and SWidget1.
  */
-class SLATECORE_API SOverlay
-	: public SPanel
+class SLATECORE_API SOverlay : public SPanel
 {
+	SLATE_DECLARE_WIDGET(SOverlay, SPanel)
 public:	
 
 	/** A slot that support alignment of content and padding and z-order */
@@ -50,14 +50,13 @@ public:
 			SLATE_ARGUMENT(TOptional<int32>, ZOrder)
 		SLATE_SLOT_END_ARGS()
 
-		void Construct(const FChildren& SlotOwner, FSlotArguments&& InArgs)
+		void Construct(const FChildren& SlotOwner, FSlotArguments&& InArgs);
+
+		int32 GetZOrder() const
 		{
-			TBasicLayoutWidgetSlot<FOverlaySlot>::Construct(SlotOwner, MoveTemp(InArgs));
-			ZOrder = InArgs._ZOrder.Get(ZOrder);
+			return ZOrder;
 		}
 
-		int32 GetZOrder() const { return ZOrder; }
-
 		void SetZOrder(int32 InOrder)
 		{
 			if (ZOrder != InOrder)
@@ -117,16 +116,16 @@ public:
 	/** @return a new slot. Slots contain children for SOverlay */
 	static FOverlaySlot::FSlotArguments Slot();
 
-	// SWidget interface
+	//~ Begin of SWidget interface
 	virtual void OnArrangeChildren( const FGeometry& AllottedGeometry, FArrangedChildren& ArrangedChildren ) const override;
 	virtual FChildren* GetChildren() override;
 	virtual int32 OnPaint( const FPaintArgs& Args, const FGeometry& AllottedGeometry, const FSlateRect& MyCullingRect, FSlateWindowElementList& OutDrawElements, int32 LayerId, const FWidgetStyle& InWidgetStyle, bool bParentEnabled ) const override;
-	// End of SWidget interface
+	//~ End of SWidget interface
 
 protected:
-	// Begin SWidget overrides.
+	//~ Begin SWidget overrides.
 	virtual FVector2D ComputeDesiredSize(float) const override;
-	// End SWidget overrides.
+	//~ End SWidget overrides.
 
 protected:
 	/** The SOverlay's slots; each slot contains a child widget. */
