Initialize SBackgroundBlur's variables to prevent comparation with NaN value.

diff --git a/Engine/Source/Runtime/Slate/Private/Widgets/Layout/SBackgroundBlur.cpp b/Engine/Source/Runtime/Slate/Private/Widgets/Layout/SBackgroundBlur.cpp
--- a/Engine/Source/Runtime/Slate/Private/Widgets/Layout/SBackgroundBlur.cpp
+++ b/Engine/Source/Runtime/Slate/Private/Widgets/Layout/SBackgroundBlur.cpp
@@ -24,9 +24,9 @@ void SBackgroundBlur::PrivateRegisterAttributes(FSlateAttributeInitializer& Attr
 }
 
 SBackgroundBlur::SBackgroundBlur()
-	: BlurStrengthAttribute(*this)
+	: BlurStrengthAttribute(*this, 0.f)
 	, BlurRadiusAttribute(*this)
-	, CornerRadiusAttribute(*this)
+	, CornerRadiusAttribute(*this, FVector4(0.f, 0.f, 0.f, 0.f))
 {
 }
 
@@ -66,7 +66,7 @@ void SBackgroundBlur::SetApplyAlphaToBlur(bool bInApplyAlphaToBlur)
 
 void SBackgroundBlur::SetBlurRadius(TAttribute<TOptional<int32>> InBlurRadius)
 {
-	BlurRadiusAttribute.Assign(*this, InBlurRadius);
+	BlurRadiusAttribute.Assign(*this, MoveTemp(InBlurRadius));
 }
 
 void SBackgroundBlur::SetBlurStrength(TAttribute<float> InStrength)
@@ -85,7 +85,7 @@ void SBackgroundBlur::SetLowQualityBackgroundBrush(const FSlateBrush* InBrush)
 
 void SBackgroundBlur::SetCornerRadius(TAttribute<FVector4> InCornerRadius)
 {
-	CornerRadiusAttribute.Assign(*this, InCornerRadius);
+	CornerRadiusAttribute.Assign(*this, MoveTemp(InCornerRadius));
 }
 
 void SBackgroundBlur::SetHAlign(EHorizontalAlignment HAlign)
