Slate: Convert Scalebox's Attribute to SlateAttribute

diff --git a/Engine/Source/Runtime/Slate/Private/Widgets/Layout/SScaleBox.cpp b/Engine/Source/Runtime/Slate/Private/Widgets/Layout/SScaleBox.cpp
--- a/Engine/Source/Runtime/Slate/Private/Widgets/Layout/SScaleBox.cpp
+++ b/Engine/Source/Runtime/Slate/Private/Widgets/Layout/SScaleBox.cpp
@@ -10,17 +10,38 @@
 
 /* SScaleBox interface
  *****************************************************************************/
+SLATE_IMPLEMENT_WIDGET(SScaleBox)
+void SScaleBox::PrivateRegisterAttributes(FSlateAttributeInitializer& AttributeInitializer)
+{
+	SLATE_ADD_MEMBER_ATTRIBUTE_DEFINITION_WITH_NAME(AttributeInitializer, "StretchDirection", StretchDirectionAttribute, EInvalidateWidgetReason::Layout | EInvalidateWidgetReason::Prepass);
+	SLATE_ADD_MEMBER_ATTRIBUTE_DEFINITION_WITH_NAME(AttributeInitializer, "UserSpecifiedScale", UserSpecifiedScaleAttribute, EInvalidateWidgetReason::Layout | EInvalidateWidgetReason::Prepass);
+	SLATE_ADD_MEMBER_ATTRIBUTE_DEFINITION_WITH_NAME(AttributeInitializer, "IgnoreInheritedScale", IgnoreInheritedScaleAttribute, EInvalidateWidgetReason::Layout | EInvalidateWidgetReason::Prepass);
+	SLATE_ADD_MEMBER_ATTRIBUTE_DEFINITION_WITH_NAME(AttributeInitializer, "Stretch", StretchAttribute, EInvalidateWidgetReason::Layout | EInvalidateWidgetReason::Prepass)
+		.OnValueChanged(FSlateAttributeDescriptor::FAttributeValueChangedDelegate::CreateLambda([](SWidget& Widget)
+			{
+				static_cast<SScaleBox&>(Widget).RefreshSafeZoneScale();
+			}));
+}
+
+SScaleBox::SScaleBox()
+	: StretchDirectionAttribute(*this, EStretchDirection::Both)
+	, StretchAttribute(*this, EStretch::None)
+	, UserSpecifiedScaleAttribute(*this, 1.0f)
+	, IgnoreInheritedScaleAttribute(*this, false)
+{
+	SetCanTick(false);
+	bCanSupportFocus = false;
+}
 
 void SScaleBox::Construct(const SScaleBox::FArguments& InArgs)
 {
 	bHasCustomPrepass = true;
 	bHasRelativeLayoutScale = true;
 
-	Stretch = InArgs._Stretch;
-
-	StretchDirection = InArgs._StretchDirection;
-	UserSpecifiedScale = InArgs._UserSpecifiedScale;
-	IgnoreInheritedScale = InArgs._IgnoreInheritedScale;
+	StretchAttribute.Assign(*this, InArgs._Stretch);
+	StretchDirectionAttribute.Assign(*this, InArgs._StretchDirection);
+	UserSpecifiedScaleAttribute.Assign(*this, InArgs._UserSpecifiedScale, 1.0f);
+	IgnoreInheritedScaleAttribute.Assign(*this, InArgs._IgnoreInheritedScale, false);
 
 	LastFinalOffset = FVector2D(0, 0);
 
@@ -37,7 +58,7 @@ void SScaleBox::Construct(const SScaleBox::FArguments& InArgs)
 #endif
 
 	RefreshSafeZoneScale();
-	OnSafeFrameChangedHandle = FCoreDelegates::OnSafeFrameChangedEvent.AddSP(this, &SScaleBox::RefreshSafeZoneScale);
+	OnSafeFrameChangedHandle = FCoreDelegates::OnSafeFrameChangedEvent.AddSP(this, &SScaleBox::HandleSafeFrameChangedEvent);
 }
 
 SScaleBox::~SScaleBox()
@@ -95,7 +116,7 @@ bool SScaleBox::CustomPrepass(float LayoutScaleMultiplier)
 	// Extract the incoming scale out of the layout scale if 
 	if (NewComputedContentScale.IsSet())
 	{
-		if (IgnoreInheritedScale.Get(false) && LayoutScaleMultiplier != 0)
+		if (IgnoreInheritedScaleAttribute.Get() && LayoutScaleMultiplier != 0)
 		{
 			NewComputedContentScale = NewComputedContentScale.GetValue() / LayoutScaleMultiplier;
 		}
@@ -108,7 +129,7 @@ bool SScaleBox::CustomPrepass(float LayoutScaleMultiplier)
 
 bool SScaleBox::DoesScaleRequireNormalizingPrepassOrLocalGeometry() const
 {
-	const EStretch::Type CurrentStretch = Stretch.Get();
+	const EStretch::Type CurrentStretch = StretchAttribute.Get();
 	switch (CurrentStretch)
 	{
 	case EStretch::None:
@@ -123,7 +144,7 @@ bool SScaleBox::DoesScaleRequireNormalizingPrepassOrLocalGeometry() const
 
 bool SScaleBox::IsDesiredSizeDependentOnAreaAndScale() const
 {
-	const EStretch::Type CurrentStretch = Stretch.Get();
+	const EStretch::Type CurrentStretch = StretchAttribute.Get();
 	switch (CurrentStretch)
 	{
 	case EStretch::ScaleToFitX:
@@ -136,15 +157,15 @@ bool SScaleBox::IsDesiredSizeDependentOnAreaAndScale() const
 
 float SScaleBox::ComputeContentScale(const FGeometry& PaintGeometry) const
 {
-	const EStretch::Type CurrentStretch = Stretch.Get();
-	const EStretchDirection::Type CurrentStretchDirection = StretchDirection.Get();
+	const EStretch::Type CurrentStretch = StretchAttribute.Get();
+	const EStretchDirection::Type CurrentStretchDirection = StretchDirectionAttribute.Get();
 
 	switch (CurrentStretch)
 	{
 	case EStretch::ScaleBySafeZone:
 		return SafeZoneScale;
 	case EStretch::UserSpecified:
-		return UserSpecifiedScale.Get(1.0f);
+		return UserSpecifiedScaleAttribute.Get();
 	}
 
 	float FinalScale = 1;
@@ -207,8 +228,8 @@ void SScaleBox::OnArrangeChildren(const FGeometry& AllottedGeometry, FArrangedCh
 		const FVector2D CurrentWidgetDesiredSize = ChildSlot.GetWidget()->GetDesiredSize();
 		FVector2D SlotWidgetDesiredSize = CurrentWidgetDesiredSize;
 
-		const EStretch::Type CurrentStretch = Stretch.Get();
-		const EStretchDirection::Type CurrentStretchDirection = StretchDirection.Get();
+		const EStretch::Type CurrentStretch = StretchAttribute.Get();
+		const EStretchDirection::Type CurrentStretchDirection = StretchDirectionAttribute.Get();
 
 		if (CurrentStretch == EStretch::Fill)
 		{
@@ -275,7 +296,7 @@ int32 SScaleBox::OnPaint( const FPaintArgs& Args, const FGeometry& AllottedGeome
 
 	if (GetClipping() == EWidgetClipping::Inherit)
 	{
-		const EStretch::Type CurrentStretch = Stretch.Get();
+		const EStretch::Type CurrentStretch = StretchAttribute.Get();
 
 		// There are a few stretch modes that require we clip, even if the user didn't set the property.
 		switch (CurrentStretch)
@@ -325,36 +346,22 @@ void SScaleBox::SetVAlign(EVerticalAlignment VAlign)
 
 void SScaleBox::SetStretchDirection(EStretchDirection::Type InStretchDirection)
 {
-	if (SetAttribute(StretchDirection, TAttribute<EStretchDirection::Type>(InStretchDirection), EInvalidateWidgetReason::Layout))
-	{
-		Invalidate(EInvalidateWidgetReason::Prepass);
-	}
+	StretchDirectionAttribute.Set(*this, InStretchDirection);
 }
 
 void SScaleBox::SetStretch(EStretch::Type InStretch)
 {
-	if (SetAttribute(Stretch, TAttribute<EStretch::Type>(InStretch), EInvalidateWidgetReason::Layout))
-	{
-		// This function invalidates the prepass.
-		RefreshSafeZoneScale();
-		check(NeedsPrepass());
-	}
+	StretchAttribute.Set(*this, InStretch);
 }
 
 void SScaleBox::SetUserSpecifiedScale(float InUserSpecifiedScale)
 {
-	if (SetAttribute(UserSpecifiedScale, TAttribute<float>(InUserSpecifiedScale), EInvalidateWidgetReason::Prepass))
-	{
-		Invalidate(EInvalidateWidgetReason::Prepass);
-	}
+	UserSpecifiedScaleAttribute.Set(*this, InUserSpecifiedScale);
 }
 
 void SScaleBox::SetIgnoreInheritedScale(bool InIgnoreInheritedScale)
 {
-	if (SetAttribute(IgnoreInheritedScale, TAttribute<bool>(InIgnoreInheritedScale), EInvalidateWidgetReason::Layout))
-	{
-		Invalidate(EInvalidateWidgetReason::Prepass);
-	}
+	IgnoreInheritedScaleAttribute.Set(*this, InIgnoreInheritedScale);
 }
 
 FVector2D SScaleBox::ComputeDesiredSize(float InScale) const
@@ -372,7 +379,7 @@ FVector2D SScaleBox::ComputeDesiredSize(float InScale) const
 				// expected scale, if we can get that extra space, awesome.
 				if (ComputedContentScale.IsSet() && ComputedContentScale.GetValue() != 0)
 				{
-					const EStretch::Type CurrentStretch = Stretch.Get();
+					const EStretch::Type CurrentStretch = StretchAttribute.Get();
 
 					switch (CurrentStretch)
 					{
@@ -422,7 +429,7 @@ void SScaleBox::RefreshSafeZoneScale()
 	else
 #endif
 	{
-		if (Stretch.Get() == EStretch::ScaleBySafeZone)
+		if (StretchAttribute.Get() == EStretch::ScaleBySafeZone)
 		{
 			TSharedPtr<SViewport> GameViewport = FSlateApplication::Get().GetGameViewport();
 			if (GameViewport.IsValid())
@@ -446,7 +453,11 @@ void SScaleBox::RefreshSafeZoneScale()
 	ScaleDownBy = FMath::Max(SafeZoneScaleX, SafeZoneScaleY);
 
 	SafeZoneScale = 1.f - ScaleDownBy;
+}
 
+void SScaleBox::HandleSafeFrameChangedEvent()
+{
+	RefreshSafeZoneScale();
 	Invalidate(EInvalidateWidgetReason::Prepass);
 }
 
@@ -454,13 +465,13 @@ void SScaleBox::RefreshSafeZoneScale()
 
 void SScaleBox::DebugSafeAreaUpdated(const FMargin& NewSafeZone, bool bShouldRecacheMetrics)
 {
-	RefreshSafeZoneScale();
+	HandleSafeFrameChangedEvent();
 }
 
 void SScaleBox::SetOverrideScreenInformation(TOptional<FVector2D> InScreenSize)
 {
 	OverrideScreenSize = InScreenSize;
-	RefreshSafeZoneScale();
+	HandleSafeFrameChangedEvent();
 }
 
 #endif
\ No newline at end of file
diff --git a/Engine/Source/Runtime/Slate/Public/Widgets/Layout/SScaleBox.h b/Engine/Source/Runtime/Slate/Public/Widgets/Layout/SScaleBox.h
--- a/Engine/Source/Runtime/Slate/Public/Widgets/Layout/SScaleBox.h
+++ b/Engine/Source/Runtime/Slate/Public/Widgets/Layout/SScaleBox.h
@@ -73,6 +73,8 @@ namespace EStretch
  */
 class SLATE_API SScaleBox : public SCompoundWidget
 {
+	SLATE_DECLARE_WIDGET(SScaleBox, SCompoundWidget)
+
 public:
 	SLATE_BEGIN_ARGS(SScaleBox)
 	: _Content()
@@ -111,13 +113,11 @@ public:
 
 	SLATE_END_ARGS()
 
+protected:
 	/** Constructor */
-	SScaleBox()
-	{
-		SetCanTick(false);
-		bCanSupportFocus = false;
-	}
+	SScaleBox();
 
+public:
 	virtual ~SScaleBox();
 
 	void Construct(const FArguments& InArgs);
@@ -164,23 +164,36 @@ protected:
 	float ComputeContentScale(const FGeometry& PaintGeometry) const;
 
 	void RefreshSafeZoneScale();
+	void HandleSafeFrameChangedEvent();
 
 #if WITH_EDITOR
 	void DebugSafeAreaUpdated(const FMargin& NewSafeZone, bool bShouldRecacheMetrics);
 #endif
 
 protected:
+#if WITH_EDITORONLY_DATA
+	UE_DEPRECATED(5.0, "Direct access to StretchDirection is now deprecated. Use the setter or getter.")
+	TSlateDeprecatedTAttribute<EStretchDirection::Type> StretchDirection;
+	UE_DEPRECATED(5.0, "Direct access to Stretch is now deprecated. Use the setter or getter.")
+	TSlateDeprecatedTAttribute<EStretch::Type> Stretch;
+	UE_DEPRECATED(5.0, "Direct access to UserSpecifiedScale is now deprecated. Use the setter or getter.")
+	TSlateDeprecatedTAttribute<float> UserSpecifiedScale;
+	UE_DEPRECATED(5.0, "Direct access to IgnoreInheritedScale is now deprecated. Use the setter or getter.")
+	TSlateDeprecatedTAttribute<bool> IgnoreInheritedScale;
+#endif
+
+private:
 	/** The allowed direction of stretching of the content */
-	TAttribute<EStretchDirection::Type> StretchDirection;
+	TSlateAttribute<EStretchDirection::Type> StretchDirectionAttribute;
 
 	/** The method of scaling that is applied to the content. */
-	TAttribute<EStretch::Type> Stretch;
+	TSlateAttribute<EStretch::Type> StretchAttribute;
 
 	/** Optional scale that can be specified by the User */
-	TAttribute<float> UserSpecifiedScale;
+	TSlateAttribute<float> UserSpecifiedScaleAttribute;
 
 	/** Optional bool to ignore the inherited scale */
-	TAttribute<bool> IgnoreInheritedScale;
+	TSlateAttribute<bool> IgnoreInheritedScaleAttribute;
 
 	/** Computed scale when scaled by safe zone padding */
 	float SafeZoneScale;
