Slate: Add check to see if the InvalidationWidget is still valid after a child order. Add debug command to verify that relationship between widgets parent/children.

diff --git a/Engine/Source/Runtime/Slate/Private/Framework/Application/SlateApplication.cpp b/Engine/Source/Runtime/Slate/Private/Framework/Application/SlateApplication.cpp
--- a/Engine/Source/Runtime/Slate/Private/Framework/Application/SlateApplication.cpp
+++ b/Engine/Source/Runtime/Slate/Private/Framework/Application/SlateApplication.cpp
@@ -41,6 +41,7 @@
 #include "Math/UnitConversion.h"
 #include "HAL/LowLevelMemTracker.h"
 #include "ProfilingDebugging/CsvProfiler.h"
+#include "Types/ReflectionMetadata.h"
 #include "Trace/SlateTrace.h"
 #include "Styling/StarshipCoreStyle.h"
 #include "Styling/UMGCoreStyle.h"
@@ -62,6 +63,23 @@ CSV_DECLARE_CATEGORY_MODULE_EXTERN(SLATECORE_API, Slate);
 
 //////////////////////////////////////////////////////////////////////////
 
+#if WITH_SLATE_DEBUGGING
+bool GSlateVerifyParentChildrenRelationship = false;
+static FAutoConsoleVariableRef CVarSlateVerifyParentChildrenRelationship(
+	TEXT("Slate.VerifyParentChildrenRelationship"),
+	GSlateVerifyParentChildrenRelationship,
+	TEXT("Each frame, verify that a widget has only one parent.")
+);
+namespace UE {
+	namespace Slate
+	{
+		void VerifyParentChildrenRelationship(const TSharedRef<SWindow>& WindowToDraw);
+	}
+}
+#endif
+
+//////////////////////////////////////////////////////////////////////////
+
 /** 
  * The cursor given to any additional input-providing SlateUsers beyond the first one (since the first one owns the platform cursor)
  * Or, if we're on a platform without a cursor, User 0 gets one of these too
@@ -1046,6 +1064,13 @@ struct FDrawWindowArgs
 
 void FSlateApplication::DrawWindowAndChildren( const TSharedRef<SWindow>& WindowToDraw, FDrawWindowArgs& DrawWindowArgs )
 {
+#if WITH_SLATE_DEBUGGING
+	if (GSlateVerifyParentChildrenRelationship)
+	{
+		UE::Slate::VerifyParentChildrenRelationship(WindowToDraw);
+	}
+#endif
+
 	// On Mac, where child windows can be on screen even if their parent is hidden or minimized, we want to always draw child windows.
 	// On other platforms we set bDrawChildWindows to true only if we draw the current window.
 	bool bDrawChildWindows = PLATFORM_MAC;
@@ -6973,3 +6998,66 @@ bool FSlateApplication::InputPreProcessorsHelper::PreProcessInput(ESlateDebuggin
 
 	return bShouldExit;
 }
+
+
+namespace UE
+{
+namespace Slate
+{
+#if WITH_SLATE_DEBUGGING
+namespace Private
+{
+	bool VerifyParentChildrenRelationship_Recursive(SWidget* Parent, TMap<const SWidget*, const SWidget*>& AllWidgets)
+	{
+		bool bResult = true;
+		Parent->GetAllChildren()->ForEachWidget([&bResult, Parent, &AllWidgets](SWidget& ChildWidget)
+			{
+				if (&ChildWidget != &SNullWidget::NullWidget.Get())
+				{
+					if (AllWidgets.Find(&ChildWidget))
+					{
+						UE_LOG(LogSlate, Warning, TEXT("The widget '%s' is owned by more than one parent. 1:'%s' 2:'%s'.")
+							, *FReflectionMetaData::GetWidgetDebugInfo(ChildWidget)
+							, *FReflectionMetaData::GetWidgetDebugInfo(AllWidgets[&ChildWidget])
+							, *FReflectionMetaData::GetWidgetDebugInfo(Parent));
+					}
+					else
+					{
+						AllWidgets.Add(&ChildWidget, Parent);
+					}
+
+					if (ChildWidget.GetParentWidget().Get() != Parent)
+					{
+						bResult = false;
+						UE_LOG(LogSlate, Warning, TEXT("The widget '%s' ."), *FReflectionMetaData::GetWidgetDebugInfo(ChildWidget));
+					}
+					else
+					{
+						VerifyParentChildrenRelationship_Recursive(&ChildWidget, AllWidgets);
+					}
+				}
+			});
+
+		return bResult;
+	}
+}//Private
+
+
+void VerifyParentChildrenRelationship(const TSharedRef<SWindow>& WindowToDraw)
+{
+	if (WindowToDraw != SNullWidget::NullWidget)
+	{
+		TMap<const SWidget*, const SWidget*> AllWidgets;
+		AllWidgets.Add(&WindowToDraw.Get(), nullptr);
+		if (!Private::VerifyParentChildrenRelationship_Recursive(&WindowToDraw.Get(), AllWidgets))
+		{
+			GSlateVerifyParentChildrenRelationship = false;
+			ensureMsgf(false, TEXT("VerifyParentChildrenRelationship failed. See log for more info."));
+		}
+	}
+}
+
+#endif
+}// Slate
+}//UE
+
diff --git a/Engine/Source/Runtime/SlateCore/Private/FastUpdate/SlateInvalidationRoot.cpp b/Engine/Source/Runtime/SlateCore/Private/FastUpdate/SlateInvalidationRoot.cpp
--- a/Engine/Source/Runtime/SlateCore/Private/FastUpdate/SlateInvalidationRoot.cpp
+++ b/Engine/Source/Runtime/SlateCore/Private/FastUpdate/SlateInvalidationRoot.cpp
@@ -675,6 +675,10 @@ void FSlateInvalidationRoot::ProcessPreUpdate()
 
 						// We need to keep it to run the layout calculation in FWidgetProxy::ProcessPostInvalidation
 						//EnumRemoveFlags(InvalidationWidget.CurrentInvalidateReason, EInvalidateWidgetReason::ChildOrder);
+
+						// The child may change and may become invalid but not the parent.
+						checkf(WidgetPtr->GetProxyHandle().IsValid(WidgetPtr) && WidgetPtr->GetProxyHandle().GetWidgetIndex() == InvalidationWidget.Index
+							, TEXT("The widget '%s' should be valid after a ProcessChildOrderInvalidation."), *FReflectionMetaData::GetWidgetDebugInfo(WidgetPtr));
 					}
 
 					if (EnumHasAnyFlags(InvalidationWidget.CurrentInvalidateReason, EInvalidateWidgetReason::AttributeRegistration))
