Slate: Fixed an issue where right-clicking outside of the current selection in editable text boxes behaved unexpectedly
- Before: It deselected but didn't update the visuals, so the menu options were greyed out and broken seeming
- Now: Deselects and moves the cursor to the right-click location (matching the behavior of the widget in windows, chrome, etc...)

diff --git a/Engine/Source/Runtime/Slate/Private/Widgets/Text/SlateEditableTextLayout.cpp b/Engine/Source/Runtime/Slate/Private/Widgets/Text/SlateEditableTextLayout.cpp
--- a/Engine/Source/Runtime/Slate/Private/Widgets/Text/SlateEditableTextLayout.cpp
+++ b/Engine/Source/Runtime/Slate/Private/Widgets/Text/SlateEditableTextLayout.cpp
@@ -1183,6 +1183,7 @@ FReply FSlateEditableTextLayout::HandleMouseButtonDown(const FGeometry& MyGeomet
 				{
 					// Deselect any text that was selected
 					ClearSelection();
+					MoveCursor(FMoveCursor::ViaScreenPointer(MyGeometry.AbsoluteToLocal(InMouseEvent.GetScreenSpacePosition()), MyGeometry.Scale, ECursorAction::MoveCursor));
 				}
 			}
 
