Fix stacked notifications on high dpi not positioning correctly.

diff --git a/Engine/Source/Runtime/Slate/Private/Framework/Notifications/NotificationManager.cpp b/Engine/Source/Runtime/Slate/Private/Framework/Notifications/NotificationManager.cpp
--- a/Engine/Source/Runtime/Slate/Private/Framework/Notifications/NotificationManager.cpp
+++ b/Engine/Source/Runtime/Slate/Private/Framework/Notifications/NotificationManager.cpp
@@ -45,31 +45,34 @@ void FSlateNotificationManager::FRegionalNotificationList::Arrange()
 	// If the app has a status bar push up notifications so that they don't overlap the status bar
 	static const float StausBarHeight = FAppStyle::Get().GetFloat("StatusBar.Height", nullptr, 0.0f);
 
+	const FVector2D ScaledOffset = NotificationManagerConstants::NotificationOffset;
 	FVector2D AnchorPoint(
-		Region.Right - NotificationManagerConstants::NotificationOffset.X,
-		Region.Bottom - (NotificationManagerConstants::NotificationOffset.Y + StausBarHeight));
+		Region.Right - ScaledOffset.X,
+		Region.Bottom - (ScaledOffset.Y + StausBarHeight));
 
 	for (int32 ListIndex = Notifications.Num() - 1; ListIndex >= 0; --ListIndex)
 	{
 		TSharedPtr<SWindow> PinnedWindow = Notifications[ListIndex]->ParentWindowPtr.Pin();
 		if( PinnedWindow.IsValid() )
 		{
+			const float StackOffset = ScaledOffset.Y * ((Notifications.Num() - 1) - ListIndex);
 			const FVector2D DesiredSize = PinnedWindow->GetDesiredSize();
-			const FVector2D NewPosition(AnchorPoint.X - DesiredSize.X, AnchorPoint.Y - DesiredSize.Y);
-			if( NewPosition != PinnedWindow->GetPositionInScreen() && DesiredSize != PinnedWindow->GetSizeInScreen() )
+			const FVector2D NewPosition(AnchorPoint.X - DesiredSize.X, AnchorPoint.Y - DesiredSize.Y - StackOffset);
+
+			if( NewPosition != PinnedWindow->GetPositionInScreen() && DesiredSize != PinnedWindow->GetSizeInScreen())
 			{
 				PinnedWindow->ReshapeWindow( NewPosition, DesiredSize );
 			}
-			else if( NewPosition != PinnedWindow->GetPositionInScreen() )
+			else if (NewPosition != PinnedWindow->GetPositionInScreen())
 			{
-				float StackOffset = NotificationManagerConstants::NotificationOffset.Y * ((Notifications.Num()-1) - ListIndex);
-				PinnedWindow->MoveWindowTo(NewPosition - FVector2D(0, StackOffset));
+				PinnedWindow->MoveWindowTo(NewPosition);
 			}
 			AnchorPoint.Y -= DesiredSize.Y;
 		}
 	}
 }
 
+
 FSlateNotificationManager& FSlateNotificationManager::Get()
 {
 	static FSlateNotificationManager Instance;
@@ -272,7 +275,7 @@ void FSlateNotificationManager::Tick()
 {
 	// Ensure that the region rectangles still match the screen work areas.
 	// This is necessary if the desktop configuration has changed
-	for (auto& RegionList : RegionalLists)
+	for (FRegionalNotificationList& RegionList : RegionalLists)
 	{
 		RegionList.Region = FSlateApplication::Get().GetWorkArea(RegionList.Region);
 	}
