Fix mouse offset at non-standard resolutions.

diff --git a/Engine/Source/Runtime/Slate/Private/Framework/Application/SlateApplication.cpp b/Engine/Source/Runtime/Slate/Private/Framework/Application/SlateApplication.cpp
--- a/Engine/Source/Runtime/Slate/Private/Framework/Application/SlateApplication.cpp
+++ b/Engine/Source/Runtime/Slate/Private/Framework/Application/SlateApplication.cpp
@@ -535,6 +535,13 @@ FAutoConsoleVariableRef CVarRequireFocusForGamepadInput(
 	TEXT("Whether gamepad input should be ignored by the engine if the application is not currently active")
 );
 
+static bool TransformFullscreenMouseInput = true;
+FAutoConsoleVariableRef CVarSlateTransformFullscreenMouseInput(
+	TEXT("Slate.Transform.FullscreenMouseInput"),
+	TransformFullscreenMouseInput,
+	TEXT("Set true to transform mouse input to account for viewport stretching at fullscreen resolutions not natively supported by the monitor.")
+);
+
 #if PLATFORM_UI_NEEDS_TOOLTIPS
 static bool bEnableTooltips = true;
 #else
@@ -1713,7 +1720,19 @@ FWidgetPath FSlateApplication::LocateWidgetInWindow(FVector2D ScreenspaceMouseCo
 	const bool bAcceptsInput = Window->IsVisible() && (Window->AcceptsInput() || IsWindowHousingInteractiveTooltip(Window));
 	if (bAcceptsInput && Window->IsScreenspaceMouseWithin(ScreenspaceMouseCoordinate))
 	{
-		TArray<FWidgetAndPointer> WidgetsAndCursors = Window->GetHittestGrid().GetBubblePath(ScreenspaceMouseCoordinate, GetCursorRadius(), bIgnoreEnabledStatus, UserIndex);
+		FVector2D CursorPosition = ScreenspaceMouseCoordinate;
+
+		if (TransformFullscreenMouseInput && !GIsEditor && Window->GetWindowMode() == EWindowMode::Fullscreen)
+		{
+			// Screen space mapping scales everything. When window resolution doesn't match platform resolution, 
+			// this causes offset cursor hit-tests in fullscreen. Correct in slate since we are first window-aware slate processor.
+			FVector2D WindowSize = Window->GetSizeInScreen();
+			FVector2D DisplaySize = { (float)CachedDisplayMetrics.PrimaryDisplayWidth, (float)CachedDisplayMetrics.PrimaryDisplayHeight };
+
+			CursorPosition *= WindowSize / DisplaySize;
+		}
+
+		TArray<FWidgetAndPointer> WidgetsAndCursors = Window->GetHittestGrid().GetBubblePath(CursorPosition, GetCursorRadius(), bIgnoreEnabledStatus, UserIndex);
 		return FWidgetPath(MoveTemp(WidgetsAndCursors));
 	}
 	else
@@ -2438,6 +2457,11 @@ void FSlateApplication::ActivateGameViewport()
 	}
 }
 
+bool FSlateApplication::GetTransformFullscreenMouseInput() const
+{
+	return TransformFullscreenMouseInput;
+}
+
 bool FSlateApplication::SetUserFocus(uint32 UserIndex, const TSharedPtr<SWidget>& WidgetToFocus, EFocusCause ReasonFocusIsChanging /* = EFocusCause::SetDirectly*/)
 {
 	TSharedPtr<FSlateUser> CurrentUser = GetUser(UserIndex);
diff --git a/Engine/Source/Runtime/Slate/Private/Framework/Application/SlateUser.cpp b/Engine/Source/Runtime/Slate/Private/Framework/Application/SlateUser.cpp
--- a/Engine/Source/Runtime/Slate/Private/Framework/Application/SlateUser.cpp
+++ b/Engine/Source/Runtime/Slate/Private/Framework/Application/SlateUser.cpp
@@ -897,7 +897,7 @@ void FSlateUser::LockCursorInternal(const FWidgetPath& WidgetPath)
 	if (NativeWindow && NativeWindow->IsForegroundWindow())
 	{
 		// The last widget in the path should be the widget we are locking the cursor to
-		const FSlateRect SlateClipRect = WidgetPath.Widgets.Last().Geometry.GetLayoutBoundingRect();
+		FSlateRect SlateClipRect = WidgetPath.Widgets.Last().Geometry.GetLayoutBoundingRect();
 		LastComputedLockBounds = SlateClipRect;
 		LockingWidgetPath = WidgetPath;
 
@@ -908,6 +908,20 @@ void FSlateUser::LockCursorInternal(const FWidgetPath& WidgetPath)
 #else
 		const int32 ClipRectAdjustment = 0;
 #endif
+		// Screen space mapping scales everything. When viewport resolution doesn't match platform resolution, 
+		// this causes offset cursor hit-tests in fullscreen. Correct when capturing mouse as viewport widget may be smaller than screen in pixels.
+		if (FSlateApplication::Get().GetTransformFullscreenMouseInput() && !GIsEditor && NativeWindow->GetWindowMode() == EWindowMode::Fullscreen)
+		{
+			FDisplayMetrics CachedDisplayMetrics;
+			FSlateApplication::Get().GetCachedDisplayMetrics(CachedDisplayMetrics);
+			FVector2D DisplaySize = { (float)CachedDisplayMetrics.PrimaryDisplayWidth, (float)CachedDisplayMetrics.PrimaryDisplayHeight };
+			FVector2D DisplayDistortion = SlateClipRect.GetSize() / DisplaySize;
+
+			SlateClipRect.Left /= DisplayDistortion.X;
+			SlateClipRect.Top /= DisplayDistortion.Y;
+			SlateClipRect.Right /= DisplayDistortion.X;
+			SlateClipRect.Bottom /= DisplayDistortion.Y;
+		}
 
 		// Note: We round the upper left coordinate of the clip rect so we guarantee the rect is inside the geometry of the widget.  If we truncated when there is a half pixel we would cause the clip
 		// rect to be half a pixel larger than the geometry and cause the mouse to go outside of the geometry.
diff --git a/Engine/Source/Runtime/Slate/Private/Widgets/SViewport.cpp b/Engine/Source/Runtime/Slate/Private/Widgets/SViewport.cpp
--- a/Engine/Source/Runtime/Slate/Private/Widgets/SViewport.cpp
+++ b/Engine/Source/Runtime/Slate/Private/Widgets/SViewport.cpp
@@ -102,6 +102,21 @@ void SViewport::SetActive(bool bActive)
 	}
 }
 
+FVector2D SViewport::TransformToViewport(const FVector2D& CursorPosition) const
+{
+	// Screen space mapping scales everything. When viewport resolution doesn't match platform resolution, 
+	// this causes offset cursor hit-tests in fullscreen. Correct in slate since we are first viewport-aware processor.
+	if (FSlateApplication::Get().GetTransformFullscreenMouseInput() && !GIsEditor && CachedParentWindow.IsValid() && CachedParentWindow.Pin()->GetWindowMode() == EWindowMode::Fullscreen)
+	{
+		FDisplayMetrics CachedDisplayMetrics;
+		FSlateApplication::Get().GetCachedDisplayMetrics(CachedDisplayMetrics);
+		FVector2D DisplaySize = { (float)CachedDisplayMetrics.PrimaryDisplayWidth, (float)CachedDisplayMetrics.PrimaryDisplayHeight };
+		return CursorPosition * GetTickSpaceGeometry().GetAbsoluteSize() / DisplaySize;
+	}
+
+	return CursorPosition;
+}
+
 EActiveTimerReturnType SViewport::EnsureTick(double InCurrentTime, float InDeltaTime)
 {
 	return EActiveTimerReturnType::Continue;
@@ -228,7 +243,11 @@ void SViewport::Tick( const FGeometry& AllottedGeometry, const double InCurrentT
 
 FCursorReply SViewport::OnCursorQuery( const FGeometry& MyGeometry, const FPointerEvent& CursorEvent ) const
 {
-	return ViewportInterface.IsValid() ? ViewportInterface.Pin()->OnCursorQuery(MyGeometry, CursorEvent) : FCursorReply::Unhandled();
+	FPointerEvent ViewportOffsetMouseEvent(
+		CursorEvent,
+		TransformToViewport(CursorEvent.GetScreenSpacePosition()),
+		TransformToViewport(CursorEvent.GetLastScreenSpacePosition()));
+	return ViewportInterface.IsValid() ? ViewportInterface.Pin()->OnCursorQuery(MyGeometry, ViewportOffsetMouseEvent) : FCursorReply::Unhandled();
 }
 
 TOptional<TSharedRef<SWidget>> SViewport::OnMapCursor(const FCursorReply& CursorReply) const
@@ -238,12 +257,20 @@ TOptional<TSharedRef<SWidget>> SViewport::OnMapCursor(const FCursorReply& Cursor
 
 FReply SViewport::OnMouseButtonDown( const FGeometry& MyGeometry, const FPointerEvent& MouseEvent )
 {
-	return ViewportInterface.IsValid() ? ViewportInterface.Pin()->OnMouseButtonDown(MyGeometry, MouseEvent) : FReply::Unhandled();
+	FPointerEvent ViewportOffsetMouseEvent(
+		MouseEvent, 
+		TransformToViewport(MouseEvent.GetScreenSpacePosition()), 
+		TransformToViewport(MouseEvent.GetLastScreenSpacePosition()));
+	return ViewportInterface.IsValid() ? ViewportInterface.Pin()->OnMouseButtonDown(MyGeometry, ViewportOffsetMouseEvent) : FReply::Unhandled();
 }
 
 FReply SViewport::OnMouseButtonUp( const FGeometry& MyGeometry, const FPointerEvent& MouseEvent )
 {
-	return ViewportInterface.IsValid() ? ViewportInterface.Pin()->OnMouseButtonUp(MyGeometry, MouseEvent) : FReply::Unhandled();
+	FPointerEvent ViewportOffsetMouseEvent(
+		MouseEvent,
+		TransformToViewport(MouseEvent.GetScreenSpacePosition()),
+		TransformToViewport(MouseEvent.GetLastScreenSpacePosition()));
+	return ViewportInterface.IsValid() ? ViewportInterface.Pin()->OnMouseButtonUp(MyGeometry, ViewportOffsetMouseEvent) : FReply::Unhandled();
 }
 
 void SViewport::OnMouseEnter( const FGeometry& MyGeometry, const FPointerEvent& MouseEvent )
@@ -252,7 +279,11 @@ void SViewport::OnMouseEnter( const FGeometry& MyGeometry, const FPointerEvent&
 
 	if (ViewportInterface.IsValid())
 	{
-		ViewportInterface.Pin()->OnMouseEnter(MyGeometry, MouseEvent);
+		FPointerEvent ViewportOffsetMouseEvent(
+			MouseEvent,
+			TransformToViewport(MouseEvent.GetScreenSpacePosition()),
+			TransformToViewport(MouseEvent.GetLastScreenSpacePosition()));
+		ViewportInterface.Pin()->OnMouseEnter(MyGeometry, ViewportOffsetMouseEvent);
 	}
 }
 
@@ -262,23 +293,39 @@ void SViewport::OnMouseLeave( const FPointerEvent& MouseEvent )
 
 	if (ViewportInterface.IsValid())
 	{
-		ViewportInterface.Pin()->OnMouseLeave(MouseEvent);
+		FPointerEvent ViewportOffsetMouseEvent(
+			MouseEvent,
+			TransformToViewport(MouseEvent.GetScreenSpacePosition()),
+			TransformToViewport(MouseEvent.GetLastScreenSpacePosition()));
+		ViewportInterface.Pin()->OnMouseLeave(ViewportOffsetMouseEvent);
 	}
 }
 
 FReply SViewport::OnMouseMove( const FGeometry& MyGeometry, const FPointerEvent& MouseEvent )
 {
-	return ViewportInterface.IsValid() ? ViewportInterface.Pin()->OnMouseMove(MyGeometry, MouseEvent) : FReply::Unhandled();
+	FPointerEvent ViewportOffsetMouseEvent(
+		MouseEvent,
+		TransformToViewport(MouseEvent.GetScreenSpacePosition()),
+		TransformToViewport(MouseEvent.GetLastScreenSpacePosition()));
+	return ViewportInterface.IsValid() ? ViewportInterface.Pin()->OnMouseMove(MyGeometry, ViewportOffsetMouseEvent) : FReply::Unhandled();
 }
 
 FReply SViewport::OnMouseWheel( const FGeometry& MyGeometry, const FPointerEvent& MouseEvent )
 {
-	return ViewportInterface.IsValid() ? ViewportInterface.Pin()->OnMouseWheel(MyGeometry, MouseEvent) : FReply::Unhandled();
+	FPointerEvent ViewportOffsetMouseEvent(
+		MouseEvent,
+		TransformToViewport(MouseEvent.GetScreenSpacePosition()),
+		TransformToViewport(MouseEvent.GetLastScreenSpacePosition()));
+	return ViewportInterface.IsValid() ? ViewportInterface.Pin()->OnMouseWheel(MyGeometry, ViewportOffsetMouseEvent) : FReply::Unhandled();
 }
 
 FReply SViewport::OnMouseButtonDoubleClick( const FGeometry& MyGeometry, const FPointerEvent& MouseEvent )
 {
-	return ViewportInterface.IsValid() ? ViewportInterface.Pin()->OnMouseButtonDoubleClick(MyGeometry, MouseEvent) : FReply::Unhandled();
+	FPointerEvent ViewportOffsetMouseEvent(
+		MouseEvent,
+		TransformToViewport(MouseEvent.GetScreenSpacePosition()),
+		TransformToViewport(MouseEvent.GetLastScreenSpacePosition()));
+	return ViewportInterface.IsValid() ? ViewportInterface.Pin()->OnMouseButtonDoubleClick(MyGeometry, ViewportOffsetMouseEvent) : FReply::Unhandled();
 }
 
 FReply SViewport::OnKeyDown( const FGeometry& MyGeometry, const FKeyEvent& KeyEvent )
@@ -346,6 +393,7 @@ void SViewport::OnWindowClosed( const TSharedRef<SWindow>& WindowBeingClosed )
 
 FReply SViewport::OnViewportActivated(const FWindowActivateEvent& InActivateEvent)
 {
+	CachedParentWindow = FSlateApplication::Get().FindWidgetWindow(SharedThis(this));
 	return ViewportInterface.IsValid() ? ViewportInterface.Pin()->OnViewportActivated(InActivateEvent) : FReply::Unhandled();
 }
 
diff --git a/Engine/Source/Runtime/Slate/Public/Framework/Application/SlateApplication.h b/Engine/Source/Runtime/Slate/Public/Framework/Application/SlateApplication.h
--- a/Engine/Source/Runtime/Slate/Public/Framework/Application/SlateApplication.h
+++ b/Engine/Source/Runtime/Slate/Public/Framework/Application/SlateApplication.h
@@ -543,6 +543,11 @@ public:
 	 */
 	void ActivateGameViewport();
 
+	/**
+	 * True if transforming mouse input coordinates to account for fullscreen distortions
+	 */
+	bool GetTransformFullscreenMouseInput() const; 
+
 	/**
 	 * Sets specified user focus to the SWidget passed in.
 	 *
diff --git a/Engine/Source/Runtime/Slate/Public/Widgets/SViewport.h b/Engine/Source/Runtime/Slate/Public/Widgets/SViewport.h
--- a/Engine/Source/Runtime/Slate/Public/Widgets/SViewport.h
+++ b/Engine/Source/Runtime/Slate/Public/Widgets/SViewport.h
@@ -145,7 +145,7 @@ public:
 	void OnWindowClosed( const TSharedRef<SWindow>& InWindowBeingClosed );
 
 	/**
-	 * A delegate called when the viewports top level window is deactivated
+	 * A delegate called when the viewports top level window is activated
 	 */
 	FReply OnViewportActivated(const FWindowActivateEvent& InActivateEvent);
 
@@ -201,6 +201,9 @@ public:
 	UE_DEPRECATED(4.11, "SetWidgetToFocusOnActivate is no longer needed, remove this call.")
 	void SetWidgetToFocusOnActivate(TSharedPtr<SWidget> WidgetToFocus) { }
 
+	/** Transform a cursor position to the viewport */
+	FVector2D TransformToViewport(const FVector2D& CursorPosition) const;
+
 public:
 
 	// SWidget interface
@@ -252,6 +255,9 @@ protected:
 	
 private:
 
+	/** The parent window during this viewport's last activation */
+	TWeakPtr<SWindow> CachedParentWindow;
+
 	/** The handle to the active EnsureTick() timer */
 	TWeakPtr<FActiveTimerHandle> ActiveTimerHandle;
 
diff --git a/Engine/Source/Runtime/SlateCore/Public/Input/Events.h b/Engine/Source/Runtime/SlateCore/Public/Input/Events.h
--- a/Engine/Source/Runtime/SlateCore/Public/Input/Events.h
+++ b/Engine/Source/Runtime/SlateCore/Public/Input/Events.h
@@ -811,6 +811,17 @@ public:
 		, bIsTouchForceChanged(false)
 		, bIsTouchFirstMove(false)
 	{ }
+
+	/** A constructor to alter cursor positions */
+	FPointerEvent(
+		const FPointerEvent& Other,
+		const FVector2D& InScreenSpacePosition,
+		const FVector2D& InLastScreenSpacePosition)
+	{
+		*this = Other;
+		ScreenSpacePosition = InScreenSpacePosition;
+		LastScreenSpacePosition = InLastScreenSpacePosition;
+	}
 	
 public:
 
