Slate: Fix live issue with bound attribute to IsReadOnly.

diff --git a/Engine/Source/Runtime/Slate/Private/Widgets/Text/SlateEditableTextLayout.cpp b/Engine/Source/Runtime/Slate/Private/Widgets/Text/SlateEditableTextLayout.cpp
--- a/Engine/Source/Runtime/Slate/Private/Widgets/Text/SlateEditableTextLayout.cpp
+++ b/Engine/Source/Runtime/Slate/Private/Widgets/Text/SlateEditableTextLayout.cpp
@@ -2989,12 +2989,11 @@ bool FSlateEditableTextLayout::HasTextChangedFromOriginal() const
 
 void FSlateEditableTextLayout::BeginEditTransation()
 {
-	// Never change text on read only controls! 
-	check(!OwnerWidget->IsTextReadOnly());
-
-	if (StateBeforeChangingText.IsSet())
+	if (StateBeforeChangingText.IsSet() || OwnerWidget->IsTextReadOnly())
 	{
 		// Already within a translation - don't open another
+		//Or never change text on read only controls.
+		//The TextReadOnly is an attribute, the return value may have changed since the last time it was checked.
 		return;
 	}
 
