PR #8260: fix Slate to ensure instead of assert when keys are invalid(Contributed by [REDACTED_PERSON])

diff --git a/Engine/Source/Runtime/Slate/Private/Framework/Application/SlateApplication.cpp b/Engine/Source/Runtime/Slate/Private/Framework/Application/SlateApplication.cpp
--- a/Engine/Source/Runtime/Slate/Private/Framework/Application/SlateApplication.cpp
+++ b/Engine/Source/Runtime/Slate/Private/Framework/Application/SlateApplication.cpp
@@ -5828,49 +5828,37 @@ bool FSlateApplication::ExecuteNavigation(const FWidgetPath& NavigationSource, T
 bool FSlateApplication::OnControllerAnalog( FGamepadKeyNames::Type KeyName, int32 ControllerId, float AnalogValue )
 {
 	FKey Key(KeyName);
-	check(Key.IsValid());
-
 	TOptional<int32> UserIndex = GetUserIndexForController(ControllerId, Key);
-	if (!UserIndex.IsSet())
+	if(UserIndex.IsSet() && ensureMsgf(Key.IsValid(), TEXT("OnControllerAnalog(KeyName=%s,ControllerId=%d,AnalogValue=%f) key is invalid"), *KeyName.ToString(), ControllerId, AnalogValue))
 	{
-		return false;
+		FAnalogInputEvent AnalogInputEvent(Key, PlatformApplication->GetModifierKeys(), UserIndex.GetValue(), false, 0, 0, AnalogValue);
+		return ProcessAnalogInputEvent(AnalogInputEvent);
 	}
-
-	FAnalogInputEvent AnalogInputEvent(Key, PlatformApplication->GetModifierKeys(), UserIndex.GetValue(), false, 0, 0, AnalogValue);
-
-	return ProcessAnalogInputEvent(AnalogInputEvent);
+	return false;
 }
 
 bool FSlateApplication::OnControllerButtonPressed(FGamepadKeyNames::Type KeyName, int32 ControllerId, bool IsRepeat)
 {
 	FKey Key(KeyName);
-	check(Key.IsValid());
-
 	TOptional<int32> UserIndex = GetUserIndexForController(ControllerId, Key);
-	if (!UserIndex.IsSet())
+	if(UserIndex.IsSet() && ensureMsgf(Key.IsValid(), TEXT("OnControllerButtonPressed(KeyName=%s,ControllerId=%d,IsRepeat=%b) key is invalid"), *KeyName.ToString(), ControllerId, IsRepeat))
 	{
-		return false;
+		FKeyEvent KeyEvent(Key, PlatformApplication->GetModifierKeys(), UserIndex.GetValue(), IsRepeat, 0, 0);
+		return ProcessKeyDownEvent(KeyEvent);
 	}
-
-	FKeyEvent KeyEvent(Key, PlatformApplication->GetModifierKeys(), UserIndex.GetValue(), IsRepeat, 0, 0);
-
-	return ProcessKeyDownEvent(KeyEvent);
+	return false;
 }
 
 bool FSlateApplication::OnControllerButtonReleased(FGamepadKeyNames::Type KeyName, int32 ControllerId, bool IsRepeat)
 {
 	FKey Key(KeyName);
-	check(Key.IsValid());
-
 	TOptional<int32> UserIndex = GetUserIndexForController(ControllerId, Key);
-	if (!UserIndex.IsSet())
+	if(UserIndex.IsSet() && ensureMsgf(Key.IsValid(), TEXT("OnControllerButtonReleased(KeyName=%s,ControllerId=%d,IsRepeat=%b) key is invalid"), *KeyName.ToString(), ControllerId, IsRepeat))
 	{
-		return false;
+		FKeyEvent KeyEvent(Key, PlatformApplication->GetModifierKeys(), UserIndex.GetValue(), IsRepeat, 0, 0);
+		return ProcessKeyUpEvent(KeyEvent);
 	}
-
-	FKeyEvent KeyEvent(Key, PlatformApplication->GetModifierKeys(), UserIndex.GetValue(), IsRepeat, 0, 0);
-
-	return ProcessKeyUpEvent(KeyEvent);
+	return false;
 }
 
 bool FSlateApplication::OnTouchGesture( EGestureEvent GestureType, const FVector2D &Delta, const float MouseWheelDelta, bool bIsDirectionInvertedFromDevice )
