Use visual tab role for tab context menu commands so that nomad tabs being visualized as major tabs can participate in Close Other Tabs and Close Tabs to the Right

diff --git a/Engine/Source/Runtime/Slate/Private/Framework/Docking/SDockingTabStack.cpp b/Engine/Source/Runtime/Slate/Private/Framework/Docking/SDockingTabStack.cpp
--- a/Engine/Source/Runtime/Slate/Private/Framework/Docking/SDockingTabStack.cpp
+++ b/Engine/Source/Runtime/Slate/Private/Framework/Docking/SDockingTabStack.cpp
@@ -546,10 +546,11 @@ void SDockingTabStack::CloseTabsToRightOfForegroundTab(ETabsToClose TabsToClose)
 		{
 			const TSharedRef<SDockTab>& Tab = TabWell->GetTabs()[DestroyIndex];
 
+			const ETabRole VisualTabRole = Tab->GetVisualTabRole();
 			const bool bCanClose =
 				(TabsToClose == CloseAllTabs) ||
-				(TabsToClose == CloseDocumentTabs && Tab->GetTabRole() == ETabRole::DocumentTab) ||
-				(TabsToClose == CloseDocumentAndMajorTabs && (Tab->GetTabRole() == ETabRole::DocumentTab || Tab->GetTabRole() == ETabRole::MajorTab));
+				(TabsToClose == CloseDocumentTabs && VisualTabRole == ETabRole::DocumentTab) ||
+				(TabsToClose == CloseDocumentAndMajorTabs && (VisualTabRole == ETabRole::DocumentTab || VisualTabRole == ETabRole::MajorTab));
 
 			if (!bCanClose || !Tab->RequestCloseTab())
 			{
@@ -569,10 +570,11 @@ void SDockingTabStack::CloseAllButForegroundTab(ETabsToClose TabsToClose)
 		{
 			const TSharedRef<SDockTab>& Tab = TabWell->GetTabs()[DestroyIndex];
 
+			const ETabRole VisualTabRole = Tab->GetVisualTabRole();
 			const bool bCanClose = 
 				(TabsToClose == CloseAllTabs) ||
-				(TabsToClose == CloseDocumentTabs && Tab->GetTabRole() == ETabRole::DocumentTab) ||
-				(TabsToClose == CloseDocumentAndMajorTabs && (Tab->GetTabRole() == ETabRole::DocumentTab || Tab->GetTabRole() == ETabRole::MajorTab));
+				(TabsToClose == CloseDocumentTabs && VisualTabRole == ETabRole::DocumentTab) ||
+				(TabsToClose == CloseDocumentAndMajorTabs && (VisualTabRole == ETabRole::DocumentTab || VisualTabRole == ETabRole::MajorTab));
 
 			if ((Tab == ForegroundTab) || !bCanClose || !Tab->RequestCloseTab())
 			{
@@ -1023,7 +1025,13 @@ bool SDockingTabStack::CanCloseForegroundTab() const
 bool SDockingTabStack::CanCloseTabsToRightOfForegroundTab() const
 {
 	TSharedPtr<SDockTab> ForegroundTabPtr = TabWell->GetForegroundTab();
-	if (ForegroundTabPtr.IsValid() && (ForegroundTabPtr->GetTabRole() == ETabRole::DocumentTab || ForegroundTabPtr->GetTabRole() == ETabRole::MajorTab) && (TabWell->GetNumTabs() > 1) && (TabWell->GetForegroundTabIndex() != TabWell->GetNumTabs() - 1))
+	if (!ForegroundTabPtr.IsValid())
+	{
+		return false;
+	}
+
+	const ETabRole VisualTabRole = ForegroundTabPtr->GetVisualTabRole();
+	if ((VisualTabRole == ETabRole::DocumentTab || VisualTabRole == ETabRole::MajorTab) && (TabWell->GetNumTabs() > 1) && (TabWell->GetForegroundTabIndex() != TabWell->GetNumTabs() - 1))
 	{
 		const TArray< TSharedRef<SDockTab> > MyTabs = this->GetTabs().AsArrayCopy();
 		for (int32 TabIndex = TabWell->GetForegroundTabIndex() + 1; TabIndex < MyTabs.Num(); ++TabIndex)
@@ -1042,7 +1050,13 @@ bool SDockingTabStack::CanCloseAllButForegroundTab() const
 {
 	// If the active tab is a document tab or major tab and there is at least 1 other closeable tab, offer to close the others
 	TSharedPtr<SDockTab> ForegroundTabPtr = TabWell->GetForegroundTab();
-	if (ForegroundTabPtr.IsValid() && (ForegroundTabPtr->GetTabRole() == ETabRole::DocumentTab || ForegroundTabPtr->GetTabRole() == ETabRole::MajorTab) && (TabWell->GetNumTabs() > 1))
+	if (!ForegroundTabPtr.IsValid())
+	{
+		return false;
+	}
+
+	const ETabRole VisualTabRole = ForegroundTabPtr->GetVisualTabRole();
+	if ((VisualTabRole == ETabRole::DocumentTab || VisualTabRole == ETabRole::MajorTab) && (TabWell->GetNumTabs() > 1))
 	{
 		const TArray< TSharedRef<SDockTab> > MyTabs = this->GetTabs().AsArrayCopy();
 		for (int32 TabIndex = 0; TabIndex < MyTabs.Num(); ++TabIndex)
@@ -1077,7 +1091,7 @@ bool SDockingTabStack::CanMoveTabToSideBar(TSharedRef<SDockTab> Tab) const
 bool SDockingTabStack::IsTabAllowedInSidebar(TSharedPtr<SDockTab> Tab) const
 {
 	// Major tabs are not allowed to be sidebared
-	return Tab.IsValid() && Tab->GetTabRole() != ETabRole::MajorTab && Tab->GetTabManager()->IsTabAllowedInSidebar(Tab->GetLayoutIdentifier());
+	return Tab.IsValid() && Tab->GetVisualTabRole() != ETabRole::MajorTab && Tab->GetTabManager()->IsTabAllowedInSidebar(Tab->GetLayoutIdentifier());
 }
 
 SSplitter::ESizeRule SDockingTabStack::GetSizeRule() const
