Improve tab sidebar readability when the sidebar is small or there are many tabs

- Reduce sidebar tab size when there's not enough room to show all tabs at full size
- Add tooltip for sidebar tabs
- Ellipsize sidebar tab titles when they don't fit in the tab (requires changes to support ellipsizing rotated text)

diff --git a/Engine/Source/Runtime/Slate/Private/Framework/Docking/STabSidebar.cpp b/Engine/Source/Runtime/Slate/Private/Framework/Docking/STabSidebar.cpp
--- a/Engine/Source/Runtime/Slate/Private/Framework/Docking/STabSidebar.cpp
+++ b/Engine/Source/Runtime/Slate/Private/Framework/Docking/STabSidebar.cpp
@@ -1,22 +1,124 @@
 // Copyright Epic Games, Inc. All Rights Reserved.
 
 #include "STabSidebar.h"
-#include "Widgets/Text/STextBlock.h"
 #include "Widgets/Input/SButton.h"
-#include "Widgets/SBoxPanel.h"
 #include "SDockingTabWell.h"
 #include "Widgets/Layout/SBox.h"
 #include "Widgets/Images/SImage.h"
 #include "Framework/Application/SlateApplication.h"
 #include "Framework/Docking/STabDrawer.h"
 #include "Framework/MultiBox/MultiBoxBuilder.h"
+#include "Framework/Text/PlainTextLayoutMarshaller.h"
 #include "Misc/App.h"
 #include "Widgets/Colors/SComplexGradient.h"
+#include "Widgets/Text/SlateTextBlockLayout.h"
 
 #define LOCTEXT_NAMESPACE "TabSidebar"
 
 DECLARE_DELEGATE_OneParam(FOnTabDrawerButtonClicked, TSharedRef<SDockTab>);
 
+/**
+ * Vertical text block for use in the tab drawer button.
+ * Text is aligned to the top of the widget if it fits without clipping;
+ * otherwise it is ellipsized and fills the widget height.
+ */
+class STabDrawerTextBlock : public SLeafWidget
+{
+public:
+	enum class ERotation
+	{
+		Clockwise,
+		CounterClockwise,
+	};
+
+	SLATE_BEGIN_ARGS(STabDrawerTextBlock)
+		: _Text()
+		, _TextStyle(&FCoreStyle::Get().GetWidgetStyle<FTextBlockStyle>("NormalText"))
+		, _Rotation(ERotation::Clockwise)
+		, _OverflowPolicy()
+		{}
+		SLATE_ATTRIBUTE(FText, Text)
+		SLATE_STYLE_ARGUMENT(FTextBlockStyle, TextStyle)
+		SLATE_ATTRIBUTE(ERotation, Rotation)
+		SLATE_ARGUMENT(TOptional<ETextOverflowPolicy>, OverflowPolicy)
+	SLATE_END_ARGS()
+
+	void Construct(const FArguments& InArgs)
+	{
+		Text = InArgs._Text;
+		TextStyle = *InArgs._TextStyle;
+		Rotation = InArgs._Rotation;
+		TextLayoutCache = MakeUnique<FSlateTextBlockLayout>(
+			this, FTextBlockStyle::GetDefault(), TOptional<ETextShapingMethod>(), TOptional<ETextFlowDirection>(),
+			FCreateSlateTextLayout(), FPlainTextLayoutMarshaller::Create(), nullptr);
+		TextLayoutCache->SetTextOverflowPolicy(InArgs._OverflowPolicy.IsSet() ? InArgs._OverflowPolicy : TextStyle.OverflowPolicy);
+	}
+
+	int32 OnPaint(const FPaintArgs& Args, const FGeometry& AllottedGeometry, const FSlateRect& MyCullingRect, FSlateWindowElementList& OutDrawElements, int32 LayerId, const FWidgetStyle& InWidgetStyle, bool bParentEnabled) const
+	{
+		// We're going to figure out the bounds of the corresponding horizontal text, and then rotate it into a vertical orientation.
+		const FVector2D LocalSize = AllottedGeometry.GetLocalSize();
+		const FVector2D DesiredHorizontalTextSize = TextLayoutCache->GetDesiredSize();
+		const FVector2D ActualHorizontalTextSize(FMath::Min(DesiredHorizontalTextSize.X, LocalSize.Y), FMath::Min(DesiredHorizontalTextSize.Y, LocalSize.X));
+
+		// Now determine the center of the vertical text by rotating the dimensions of the horizontal text.
+		// The center should align it to the top of the widget.
+		const FVector2D VerticalTextSize(ActualHorizontalTextSize.Y, ActualHorizontalTextSize.X);
+		const FVector2D VerticalTextCenter = VerticalTextSize / 2.0f;
+
+		// Now determine where the horizontal text should be positioned so that it is centered on the vertical text:
+		//      +-+
+		//      |v|
+		//      |e|
+		// [ horizontal ]
+		//      |r|
+		//      |t|
+		//      +-+
+		const FVector2D HorizontalTextPosition = VerticalTextCenter - ActualHorizontalTextSize / 2.0f;
+
+		// Define the text's geometry using the horizontal bounds, then rotate it 90/-90 degrees into place to become vertical.
+		const FSlateRenderTransform RotationTransform(FSlateRenderTransform(FQuat2D(FMath::DegreesToRadians(Rotation.Get() == ERotation::Clockwise ? 90 : -90))));
+		const FGeometry TextGeometry = AllottedGeometry.MakeChild(ActualHorizontalTextSize, FSlateLayoutTransform(HorizontalTextPosition), RotationTransform, FVector2D(0.5f, 0.5f));
+
+		return TextLayoutCache->OnPaint(Args, TextGeometry, MyCullingRect, OutDrawElements, LayerId, InWidgetStyle, ShouldBeEnabled(bParentEnabled));
+	}
+
+	virtual FVector2D ComputeDesiredSize(float LayoutScaleMultiplier) const override
+	{
+		// The text's desired size reflects the horizontal/untransformed text.
+		// Switch the dimensions for vertical text.
+		const FVector2D DesiredHorizontalTextSize = TextLayoutCache->ComputeDesiredSize(
+			FSlateTextBlockLayout::FWidgetDesiredSizeArgs(
+				Text.Get(),
+				FText(),
+				0.0f,
+				false,
+				ETextWrappingPolicy::DefaultWrapping,
+				ETextTransformPolicy::None,
+				FMargin(),
+				1.0f,
+				ETextJustify::Left),
+			LayoutScaleMultiplier, TextStyle);
+		return FVector2D(DesiredHorizontalTextSize.Y, DesiredHorizontalTextSize.X);
+	}
+
+	void SetText(TAttribute<FText> InText)
+	{
+		Text = InText;
+	}
+
+	void SetRotation(TAttribute<ERotation> InRotation)
+	{
+		Rotation = InRotation;
+	}
+
+private:
+	TAttribute<FText> Text;
+	FTextBlockStyle TextStyle;
+	TAttribute<ERotation> Rotation;
+	TUniquePtr<FSlateTextBlockLayout> TextLayoutCache;
+};
+
 class STabDrawerButton : public SCompoundWidget
 {
 
@@ -49,11 +151,13 @@ public:
 		.Padding(0, 0, 0, 0)
 		[
 			SNew(SBox)
-			.WidthOverride(Size.X)
-			.HeightOverride(Size.Y)
+			.WidthOverride(Size.Y) // Swap desired dimensions for a vertical tab
+			.HeightOverride(Size.X)
 			.Clipping(EWidgetClipping::ClipToBounds)
 			[
 				SAssignNew(MainButton, SButton)
+				.ToolTip(ForTab->GetToolTip() ? ForTab->GetToolTip() : TAttribute<TSharedPtr<IToolTip>>())
+				.ToolTipText(ForTab->GetToolTip() ? TAttribute<FText>() : ForTab->GetTabLabel())
 				.ContentPadding(FMargin(0.0f, DockTabStyle->TabPadding.Top, 0.0f, DockTabStyle->TabPadding.Bottom))
 				.OnClicked_Lambda([this](){OnDrawerButtonClicked.ExecuteIfBound(Tab.ToSharedRef()); return FReply::Handled(); })
 				.ForegroundColor(FSlateColor::UseForeground())
@@ -68,6 +172,8 @@ public:
 						.Orientation(EOrientation::Orient_Horizontal)
 					]
 					+ SOverlay::Slot()
+					.VAlign(VAlign_Fill)
+					.HAlign(HAlign_Center)
 					[
 						SNew(SVerticalBox)
 						+ SVerticalBox::Slot()
@@ -84,15 +190,14 @@ public:
 							//.RenderTransformPivot(FVector2D(.5f, .5f))
 						]
 						+ SVerticalBox::Slot()
-						.Padding(0, 10, 0, 0)
-						.HAlign(HAlign_Left)
-						.VAlign(VAlign_Top)
+						.Padding(0.0f, 5.0f, 0.0f, 0.0f)
+						.FillHeight(1.0f)
 						[
-							SAssignNew(Label, STextBlock)
-							.TextStyle(&DockTabStyle->TabTextStyle)
-							.Text(ForTab->GetTabLabel())
-							.Clipping(EWidgetClipping::Inherit)
-							.RenderTransformPivot(FVector2D(0.5f, 0.5f))
+							SAssignNew(Label, STabDrawerTextBlock)
+								.TextStyle(&DockTabStyle->TabTextStyle)
+								.Text(ForTab->GetTabLabel())
+								.OverflowPolicy(ETextOverflowPolicy::Ellipsis)
+								.Clipping(EWidgetClipping::ClipToBounds)
 						]
 					]
 				]
@@ -106,17 +211,21 @@ public:
 	{
 		bool bShouldAppearOpened = OpenedDrawer.IsValid();
 
+		STabDrawerTextBlock::ERotation Rotation;
 		switch (Location)
 		{
 		case ESidebarLocation::Left:
-			bShouldAppearOpened ? Rotate90CounterClockwise() : Rotate90Clockwise();
+			Rotation = bShouldAppearOpened ? STabDrawerTextBlock::ERotation::CounterClockwise : STabDrawerTextBlock::ERotation::Clockwise;
 			break;
 		case ESidebarLocation::Right:
 		default:
-			bShouldAppearOpened ? Rotate90Clockwise() : Rotate90CounterClockwise();
+			Rotation = bShouldAppearOpened ? STabDrawerTextBlock::ERotation::Clockwise : STabDrawerTextBlock::ERotation::CounterClockwise;
 			break;
 		}
 
+		check(Label);
+		Label->SetRotation(Rotation);
+
 		if (OpenedDrawer == Tab)
 		{
 			// this button is the one with the tab that is actually opened so show the opened indicator
@@ -135,15 +244,16 @@ public:
 		if (ensure(ForTab == Tab))
 		{
 			Label->SetText(ForTab->GetTabLabel());
-		}
-	}
-
-	virtual FVector2D ComputeDesiredSize(float LayoutScaleMultiplier) const override
-	{
-		const FVector2D Size = FDockingConstants::GetMaxTabSizeFor(ETabRole::PanelTab);
 
-		FVector2D LocalDesiredSize = SCompoundWidget::ComputeDesiredSize(LayoutScaleMultiplier);
-		return FVector2D(LocalDesiredSize.Y, FMath::Min(LocalDesiredSize.X, Size.X));
+			if (TSharedPtr<IToolTip> ToolTip = ForTab->GetToolTip())
+			{
+				MainButton->SetToolTip(ToolTip);
+			}
+			else
+			{
+				MainButton->SetToolTipText(ForTab->GetTabLabel());
+			}
+		}
 	}
 
 	virtual FReply OnMouseButtonDown(const FGeometry& MyGeometry, const FPointerEvent& MouseEvent)
@@ -169,21 +279,10 @@ public:
 
 		return FSlateColor::UseStyle();
 	}
-private:
-	void Rotate90Clockwise()
-	{
-		Label->SetJustification(ETextJustify::Left);
-		Label->SetRenderTransform(FSlateRenderTransform(FQuat2D(FMath::DegreesToRadians(90))));
-	}
 
-	void Rotate90CounterClockwise()
-	{
-		Label->SetJustification(ETextJustify::Right);
-		Label->SetRenderTransform(FSlateRenderTransform(FQuat2D(FMath::DegreesToRadians(-90))));
-	}
 private:
 	TSharedPtr<SDockTab> Tab;
-	TSharedPtr<STextBlock> Label;
+	TSharedPtr<STabDrawerTextBlock> Label;
 	TSharedPtr<SWidget> OpenIndicator;
 	TSharedPtr<SButton> MainButton;
 	FOnGetContent OnGetContextMenuContent;
@@ -250,7 +349,9 @@ void STabSidebar::AddTab(TSharedRef<SDockTab> Tab)
 		}
 
 		TabBox->AddSlot()
-			.AutoHeight()
+			// Make the tabs evenly fill the sidebar until they reach the max size
+			.FillHeight(1.0f)
+			.MaxHeight(FDockingConstants::GetMaxTabSizeFor(ETabRole::PanelTab).X)
 			.HAlign(HAlign_Left)
 			[
 				TabButton
diff --git a/Engine/Source/Runtime/SlateCore/Private/Rendering/ElementBatcher.cpp b/Engine/Source/Runtime/SlateCore/Private/Rendering/ElementBatcher.cpp
--- a/Engine/Source/Runtime/SlateCore/Private/Rendering/ElementBatcher.cpp
+++ b/Engine/Source/Runtime/SlateCore/Private/Rendering/ElementBatcher.cpp
@@ -1208,6 +1208,38 @@ void FSlateElementBatcher::AddTextElement(const FSlateDrawElement& DrawElement)
 	}
 }
 
+/**
+ * Determines if the 2x2 matrix represents a rotation that will keep an axis-aligned rect
+ * axis-aligned (i.e. a rotation of 90-degree increments). This allows both "proper rotations"
+ * (those without a reflection) and "improper rotations" (rotations combined with a reflection
+ * over a single axis).
+ */
+static bool IsAxisAlignedRotation(const FMatrix2x2& Matrix)
+{
+	constexpr float Tolerance = KINDA_SMALL_NUMBER;
+
+	float A, B, C, D;
+	Matrix.GetMatrix(A, B, C, D);
+
+	// The 90- and 270-degree rotation matrices have zeroes on the main diagonal e.g.
+	// [0 n]
+	// [n 0] with n = 1 or -1
+	if (FMath::IsNearlyZero(A, Tolerance) && FMath::IsNearlyZero(D, Tolerance))
+	{
+		return FMath::IsNearlyEqual(1.0f, FMath::Abs(B), Tolerance) && FMath::IsNearlyEqual(1.0f, FMath::Abs(C), Tolerance);
+	}
+
+	// The 0- and 180-degree rotation matrices have zeroes on the secondary diagonal e.g.
+	// [n 0]
+	// [0 n] with n = 1 or -1
+	if (FMath::IsNearlyZero(B, Tolerance) && FMath::IsNearlyZero(C, Tolerance))
+	{
+		return FMath::IsNearlyEqual(1.0f, FMath::Abs(A), Tolerance) && FMath::IsNearlyEqual(1.0f, FMath::Abs(D), Tolerance);
+	}
+
+	return false;
+}
+
 template<ESlateVertexRounding Rounding>
 void FSlateElementBatcher::AddShapedTextElement( const FSlateDrawElement& DrawElement )
 {
@@ -1292,20 +1324,21 @@ void FSlateElementBatcher::AddShapedTextElement( const FSlateDrawElement& DrawEl
 		{
 			const FSlateClippingState* ClippingState = ResolveClippingState(DrawElement);
 
-			if (ClippingState && ClippingState->ScissorRect.IsSet() && ClippingState->ScissorRect->IsAxisAligned() && RenderTransform.GetMatrix().IsIdentity())
+			if (ClippingState && ClippingState->ScissorRect.IsSet() && ClippingState->ScissorRect->IsAxisAligned() && IsAxisAlignedRotation(RenderTransform.GetMatrix()))
 			{
-				// Non-render transformed box;
+				// Non-render transformed box or rotation is axis-aligned at 90-degree increments
 				const FSlateRect ScissorRectBox = ClippingState->ScissorRect->GetBoundingBox();
 
+				// We know that this will be axis-aligned because the scissor rect is axis-aligned and we already
+				// checked that the render transform is axis-aligned as well
 				const FSlateRect LocalClipBoundingBox = TransformRect(RenderTransform.Inverse(), ScissorRectBox);
 				BuildContext.LocalClipBoundingBoxLeft = LocalClipBoundingBox.Left;
 				BuildContext.LocalClipBoundingBoxRight = LocalClipBoundingBox.Right - (BuildContext.bForceEllipsis ? OverflowGlyphSequence->GetMeasuredWidth() : 0);
 
-				const float TransformOffset = RenderTransform.GetTranslation().X - DrawElement.GetPosition().X;
-
-				if (OverflowGlyphSequence && BuildContext.LocalClipBoundingBoxLeft <= 0 && BuildContext.LocalClipBoundingBoxRight > ShapedGlyphSequence->GetMeasuredWidth() + TransformOffset)
+				// In checks below, ignore floating-point differences caused by transforming and untransforming the clip rect
+				if (OverflowGlyphSequence && FMath::FloorToInt(BuildContext.LocalClipBoundingBoxLeft) <= 0 && FMath::CeilToInt(BuildContext.LocalClipBoundingBoxRight) >= ShapedGlyphSequence->GetMeasuredWidth())
 				{
-					// Override overflow if the text is smaller than the clipping rect and wont be clipped
+					// Override overflow if the text is smaller than (or is the same size as) the clipping rect and wont be clipped
 					BuildContext.OverflowDirection = ETextOverflowDirection::NoOverflow;
 				}
 				else if(!OverflowGlyphSequence)
@@ -1315,7 +1348,7 @@ void FSlateElementBatcher::AddShapedTextElement( const FSlateDrawElement& DrawEl
 			}
 			else
 			{
-				// Overflow not supported on non-identity transforms
+				// Overflow not supported on non-identity transforms (except for 90-degree rotations)
 				BuildContext.OverflowDirection = ETextOverflowDirection::NoOverflow;
 			}
 		}
