Implemented recursive menu search
  - Modified the multibox to recursively explore, pull and flatten its sub-menu items to the top level multibox.Flatten items are invisible until they are matched by a user search.

Fixed a focus issue that prevented the 'Create' button/menu item action from firing when the item was clicked in the search result list.

diff --git a/Engine/Source/Runtime/Slate/Private/Framework/MultiBox/MultiBox.cpp b/Engine/Source/Runtime/Slate/Private/Framework/MultiBox/MultiBox.cpp
--- a/Engine/Source/Runtime/Slate/Private/Framework/MultiBox/MultiBox.cpp
+++ b/Engine/Source/Runtime/Slate/Private/Framework/MultiBox/MultiBox.cpp
@@ -31,6 +31,58 @@
 
 #define LOCTEXT_NAMESPACE "MultiBox"
 
+namespace UE
+{
+namespace MultiBoxUtils
+{
+
+/**
+ * Returns a path-like text containing the sub-menu(s) location of the block before it was pulled, flatten and added to its top-level multibox. Searchable
+ * blocks are added to their parent(s) multibox to support recursive searching. When a user performs a search in a top-level menu, all sub-menu items are
+ * also in the top menu (but hidden) and the matching ones are made visible. When one or more items of a given sub-menu match, a hierarchy tip appears
+ * to let the user know from which sub-menu the item(s) came from.
+ */
+static FText GetSearchHierarchyInfoText(const TArray<FText>& SearchTextHierarchyComponents)
+{
+	FText HierarchyInfoText = FText::GetEmpty();
+	if (!SearchTextHierarchyComponents.IsEmpty())
+	{
+		HierarchyInfoText = SearchTextHierarchyComponents[0];
+		for (int32 Index = 1; Index < SearchTextHierarchyComponents.Num() - 1; ++Index)
+		{
+			HierarchyInfoText = FText::Format(INVTEXT("{0}/{1}"), HierarchyInfoText, SearchTextHierarchyComponents[Index]);
+		}
+	}
+
+	return HierarchyInfoText;
+}
+
+/** Returns the searchable text displayed to the user for this search text hierarchy. This is the text normally displayed by the widget in its own menu. */
+static FText GetBlockWidgetDisplayText(const TArray<FText>& SearchTextHierarchyComponents)
+{
+	return SearchTextHierarchyComponents.IsEmpty() ? FText::GetEmpty() : SearchTextHierarchyComponents.Last();
+}
+
+/** Joins the text components into a single string for debugging purpose. */
+static FString ToString(const TArray<FText>& SearchTextHierarchyComponents)
+{
+	FString Pathname;
+	if (!SearchTextHierarchyComponents.IsEmpty())
+	{
+		Pathname = SearchTextHierarchyComponents[0].ToString();
+		for (int32 Index = 1; Index < SearchTextHierarchyComponents.Num(); ++Index)
+		{
+			Pathname.AppendChar(TEXT('/'));
+			Pathname.Append(SearchTextHierarchyComponents[Index].ToString());
+		}
+	}
+
+	return Pathname;
+}
+
+} // namespace MultiBoxUtils
+} // namespace UE
+
 
 TAttribute<bool> FMultiBoxSettings::UseSmallToolBarIcons;
 TAttribute<bool> FMultiBoxSettings::DisplayMultiboxHooks;
@@ -177,6 +229,7 @@ void FMultiBlock::SetSearchable( bool InSearchable )
 {
 	bSearchable = InSearchable;
 }
+
 bool FMultiBlock::GetSearchable() const
 {
 	return bSearchable;
@@ -630,9 +683,96 @@ void SMultiBoxWidget::AddBlockWidget(const FMultiBlock& Block, TSharedPtr<SHoriz
 
 	TSharedRef<SWidget> BlockWidget = Block.MakeWidget(SharedThis(this), InLocation, bSectionContainsIcons, OptionsWidget)->AsWidget();
 
+	// If the block being added is one of the searchable flatten blocks from one of the sub-menus. (Second pass when building the multibox widget)
+	TSharedRef<const FMultiBlock> BlockRef = Block.AsShared();
+	if (TSharedPtr<FFlattenSearchableBlockInfo>* FlattenBlockInfo = FlattenSearchableBlocks.Find(BlockRef))
+	{
+		// Wrap the block widget to display its ancestor menu as a tip to the user.
+		TSharedRef<SWidget> FlattenWidget = SNew(SVerticalBox)
+			+SVerticalBox::Slot()
+			.AutoHeight()
+			[
+				SNew(SBox)
+				.Padding(FMargin(5, 0))
+				.Visibility_Lambda([this, BlockRef](){ return (*FlattenSearchableBlocks.Find(BlockRef))->HierarchyTipVisibility; })
+				[
+					SNew(STextBlock)
+					.Text(UE::MultiBoxUtils::GetSearchHierarchyInfoText((*FlattenBlockInfo)->SearchableTextHierarchyComponents))
+				]
+			]
+			+SVerticalBox::Slot()
+			.AutoHeight()
+			[
+				BlockWidget
+			];
+
+		(*FlattenBlockInfo)->Widget = FlattenWidget;
+
+		// Remember that this widget was pulled from a sub-menus into this menu and therefore has a special status.
+		FlattenSearchableWidgets.Emplace(BlockWidget, *FlattenBlockInfo);
 
-	TWeakPtr<SWidget> BlockWidgetWeakPtr = BlockWidget;
-	TWeakPtr<const FMultiBlock> BlockWeakPtr = Block.AsShared();
+		// This widgets will only be visible if it matches a search result.
+		FlattenWidget->SetVisibility(EVisibility::Collapsed);
+
+		// The widget corresponding to the block was created above and stored along with its searchable text but the text doesn't contain the ancestors texts.
+		if (TArray<FText>* SearchableTextHierarchy = MultiBoxWidgets.Find(BlockWidget))
+		{
+			// Update the widget searchable texts hierarchy to include the ancestors.
+			*SearchableTextHierarchy = (*FlattenBlockInfo)->SearchableTextHierarchyComponents;
+		}
+
+		BlockWidget = FlattenWidget;
+	}
+	// If the added block is searchable and expands to a sub-menu. (This is the first pass building the multibox widget)
+	else if (GetSearchable() && Block.GetSearchable() && Block.GetType() == EMultiBlockType::MenuEntry && static_cast<const FMenuEntryBlock&>(Block).IsSubMenu())
+	{
+		// If the block widget being added displays a searchable text.
+		const TArray<FText>* SearchableTextHierarchy = MultiBoxWidgets.Find(BlockWidget);
+		const FText& ThisBlockDisplayText = UE::MultiBoxUtils::GetBlockWidgetDisplayText(*SearchableTextHierarchy);
+		if (!ThisBlockDisplayText.IsEmpty())
+		{
+			// Expand the block widget into its sub-menu (as if the user clicked on the block widget to expand it).
+			check(BlockWidget->GetTypeAsString() == TEXT("SMenuEntryBlock"));
+			TSharedRef<SWidget> SubMenuWidget = StaticCastSharedRef<SMenuEntryBlock>(BlockWidget)->MakeNewMenuWidget();
+
+			// If the sub-menu widget is a multibox widget.
+			if (SubMenuWidget->GetTypeAsString() == TEXT("SMultiBoxWidget"))
+			{
+				// Scan this expanded menu widgets.
+				TSharedRef<SMultiBoxWidget> SubMenuMultiboxWidget = StaticCastSharedRef<SMultiBoxWidget>(SubMenuWidget);
+				for (const TPair<TSharedPtr<SWidget>, TArray<FText>>& Pair : SubMenuMultiboxWidget->MultiBoxWidgets)
+				{
+					const TSharedPtr<SWidget>& ChildBlockWidget = Pair.Key;     // The widget.
+					const TArray<FText>& ChildBlockSearchTextHierarchy = Pair.Value;  // The text displayed by the widget to the user along with all ancestor menus. Ex. "["Shapes", "Cube"]
+					const FText& ChildBlockDisplayText =  UE::MultiBoxUtils::GetBlockWidgetDisplayText(ChildBlockSearchTextHierarchy); // The text displayed by the widget to the user. Ex "Cube"
+
+					// Skip if the widget has no searcheable text or if this is the 'search' text entry widget.
+					if(ChildBlockDisplayText.IsEmpty() || ChildBlockWidget == SubMenuMultiboxWidget->GetSearchTextWidget())
+					{
+						continue; // Skip over, not searchable.
+					}
+
+					TSharedPtr<const FMultiBlock> ChildBlock = StaticCastSharedPtr<SMultiBlockBaseWidget>(ChildBlockWidget)->GetBlock();
+					if (ChildBlock->GetSearchable())
+					{
+						// The algorithm to flatten the structure is recursive. It walks down to create/expand the tree and then walk up to collect and flatten the children blocks that have a displayable text.
+						// For example if the blocks were organized as below, A would create B, which would create C1 and C2, then B would collect and merge it children C1 and C2 as B/C1 and B/C2, then A would
+						// collect its children B, B/C1, B/C2 as A/B, A/B/C1 and A/B/C2.
+						// A
+						// |- B
+						//    |- C1
+						//    |- C2
+						TSharedPtr<FFlattenSearchableBlockInfo> FlattenSearchableBlockInfo = MakeShared<FFlattenSearchableBlockInfo>();
+						FlattenSearchableBlockInfo->SearchableTextHierarchyComponents.Add(ThisBlockDisplayText);             // If the multibox was adding the block displayed as "A" in the example, this would be "A".
+						FlattenSearchableBlockInfo->SearchableTextHierarchyComponents.Append(ChildBlockSearchTextHierarchy); // If the multibox was adding the block displayed as "A" and visiting the child ["B", "C1"], that would generate ["A", "B", "C1"] path.
+
+						// The widget corresponding to this flatten block is not added to this multibox yet, preserve the hierarchy information to update the widget searchable hierarchy once the block is added and the widget created.
+						FlattenSearchableBlocks.Emplace(ChildBlock, MoveTemp(FlattenSearchableBlockInfo));
+					}
+				}
+			}
+		}
+	}
 
 	const ISlateStyle* const StyleSet = MultiBox->GetStyleSet();
 
@@ -648,7 +788,7 @@ void SMultiBoxWidget::AddBlockWidget(const FMultiBlock& Block, TSharedPtr<SHoriz
 	{
 		FinalWidget = BlockWidget;
 	}
-	
+
 	TSharedRef<SWidget> FinalWidgetWithHook = SNullWidget::NullWidget;
 
 	if (bDisplayExtensionHooks)
@@ -673,7 +813,6 @@ void SMultiBoxWidget::AddBlockWidget(const FMultiBlock& Block, TSharedPtr<SHoriz
 	{
 		FinalWidgetWithHook = FinalWidget.ToSharedRef();
 	}
-	
 
 	switch (MultiBox->GetType())
 	{
@@ -935,116 +1074,132 @@ void SMultiBoxWidget::BuildMultiBoxWidget()
 	bool bSectionContainsIcons = false;
 	int32 NextMenuSeparator = INDEX_NONE;
 
-	for( int32 Index = 0; Index < Blocks.Num(); Index++ )
+	// The first pass adds this multibox blocks and discover/collect all sub-menus blocks. The second pass adds all sub-menu blocks discovered in the first pass.
+	for (int32 Pass = 0; Pass < 2; ++Pass)
 	{
-		// If we've passed the last menu separator, scan for the next one (the end of the list is also considered a menu separator for the purposes of this index)
-		if (NextMenuSeparator < Index)
+		int32 StartBlockIndex = 0;
+
+		// If this is the second pass, add the blocks that were discovered while expanding all sub-menus in the first pass, creating a flat list (but hidden until a search occurs).
+		if (Pass == 1)
 		{
-			bSectionContainsIcons = false;
-			for (++NextMenuSeparator; NextMenuSeparator < Blocks.Num(); ++NextMenuSeparator)
+			StartBlockIndex = Blocks.Num();
+			for (const TPair<TSharedPtr<const FMultiBlock>, TSharedPtr<FFlattenSearchableBlockInfo>>& Pair : FlattenSearchableBlocks)
 			{
-				const FMultiBlock& TestBlock = *Blocks[NextMenuSeparator];
-				if (!bSectionContainsIcons && TestBlock.HasIcon())
-				{
-					bSectionContainsIcons = true;
-				}
+				MultiBox->AddMultiBlock(Pair.Key.ToSharedRef());
+			}
+		}
 
-				if (TestBlock.GetType() == EMultiBlockType::Separator)
+		for( int32 Index = StartBlockIndex; Index < Blocks.Num(); Index++ )
+		{
+			// If we've passed the last menu separator, scan for the next one (the end of the list is also considered a menu separator for the purposes of this index)
+			if (NextMenuSeparator < Index)
+			{
+				bSectionContainsIcons = false;
+				for (++NextMenuSeparator; NextMenuSeparator < Blocks.Num(); ++NextMenuSeparator)
 				{
-					break;
+					const FMultiBlock& TestBlock = *Blocks[NextMenuSeparator];
+					if (!bSectionContainsIcons && TestBlock.HasIcon())
+					{
+						bSectionContainsIcons = true;
+					}
+
+					if (TestBlock.GetType() == EMultiBlockType::Separator)
+					{
+						break;
+					}
 				}
 			}
-		}
 
-		const FMultiBlock& Block = *Blocks[Index];
-		EMultiBlockLocation::Type Location = EMultiBlockLocation::None;
+			const FMultiBlock& Block = *Blocks[Index];
+			EMultiBlockLocation::Type Location = EMultiBlockLocation::None;
 		
-		TSharedPtr<const FMultiBlock> NextBlock = nullptr;
-
-		if (Blocks.IsValidIndex(Index + 1))
-		{
-			NextBlock = Blocks[Index + 1];
-		}
+			TSharedPtr<const FMultiBlock> NextBlock = nullptr;
 
-		// Determine the location of the current block, used for group styling information
-		{
-			// Check if we are a start or end block
-			if (Block.IsGroupStartBlock())
-			{
-				bInsideGroup = true;
-			}
-			else if (Block.IsGroupEndBlock())
+			if (Blocks.IsValidIndex(Index + 1))
 			{
-				bInsideGroup = false;
+				NextBlock = Blocks[Index + 1];
 			}
 
-			// Check if we are next to a start or end block
-			bool bIsNextToStartBlock = false;
-			bool bIsNextToEndBlock = false;
-			if (NextBlock)
+			// Determine the location of the current block, used for group styling information
 			{
-				if ( NextBlock->IsGroupEndBlock() )
+				// Check if we are a start or end block
+				if (Block.IsGroupStartBlock())
 				{
-					bIsNextToEndBlock = true;
+					bInsideGroup = true;
 				}
-			}
-			if (Index > 0)
-			{
-				const FMultiBlock& PrevBlock = *Blocks[Index - 1];
-				if ( PrevBlock.IsGroupStartBlock() )
+				else if (Block.IsGroupEndBlock())
 				{
-					bIsNextToStartBlock = true;
+					bInsideGroup = false;
 				}
-			}
-
-			// determine location
-			if (bInsideGroup)
-			{
-				// assume we are in the middle of a group
-				Location = EMultiBlockLocation::Middle;
 
-				// We are the start of a group
-				if (bIsNextToStartBlock && !bIsNextToEndBlock)
+				// Check if we are next to a start or end block
+				bool bIsNextToStartBlock = false;
+				bool bIsNextToEndBlock = false;
+				if (NextBlock)
 				{
-					Location = EMultiBlockLocation::Start;
+					if ( NextBlock->IsGroupEndBlock() )
+					{
+						bIsNextToEndBlock = true;
+					}
 				}
-				// we are the end of a group
-				else if (!bIsNextToStartBlock && bIsNextToEndBlock)
+				if (Index > 0)
 				{
-					Location = EMultiBlockLocation::End;
+					const FMultiBlock& PrevBlock = *Blocks[Index - 1];
+					if ( PrevBlock.IsGroupStartBlock() )
+					{
+						bIsNextToStartBlock = true;
+					}
 				}
-				// we are the only block in a group
-				else if (bIsNextToStartBlock && bIsNextToEndBlock)
+
+				// determine location
+				if (bInsideGroup)
 				{
-					Location = EMultiBlockLocation::None;
+					// assume we are in the middle of a group
+					Location = EMultiBlockLocation::Middle;
+
+					// We are the start of a group
+					if (bIsNextToStartBlock && !bIsNextToEndBlock)
+					{
+						Location = EMultiBlockLocation::Start;
+					}
+					// we are the end of a group
+					else if (!bIsNextToStartBlock && bIsNextToEndBlock)
+					{
+						Location = EMultiBlockLocation::End;
+					}
+					// we are the only block in a group
+					else if (bIsNextToStartBlock && bIsNextToEndBlock)
+					{
+						Location = EMultiBlockLocation::None;
+					}
 				}
 			}
-		}
 
-		TSharedPtr<const FToolBarComboButtonBlock> OptionsBlock;
-		if (NextBlock && NextBlock->GetType() == EMultiBlockType::ToolBarComboButton)
-		{
-			TSharedPtr<const FToolBarComboButtonBlock> NextToolBarComboButtonBlock = StaticCastSharedPtr<const FToolBarComboButtonBlock>(NextBlock);
-			if (NextToolBarComboButtonBlock->IsSimpleComboBox())
+			TSharedPtr<const FToolBarComboButtonBlock> OptionsBlock;
+			if (NextBlock && NextBlock->GetType() == EMultiBlockType::ToolBarComboButton)
 			{
-				// Apply a special treatment to simple combo boxes as they represent options for the previous button
-				OptionsBlock = NextToolBarComboButtonBlock;
-				// Skip over options blocks. They are not added directly
-				++Index;
+				TSharedPtr<const FToolBarComboButtonBlock> NextToolBarComboButtonBlock = StaticCastSharedPtr<const FToolBarComboButtonBlock>(NextBlock);
+				if (NextToolBarComboButtonBlock->IsSimpleComboBox())
+				{
+					// Apply a special treatment to simple combo boxes as they represent options for the previous button
+					OptionsBlock = NextToolBarComboButtonBlock;
+					// Skip over options blocks. They are not added directly
+					++Index;
+				}
 			}
-		}
 
-		if( DragPreview.IsValid() && DragPreview.InsertIndex == Index )
-		{
-			// Add the drag preview before if we have it. This block shows where the custom block will be 
-			// added if the user drops it
-			AddBlockWidget(*DragPreview.PreviewBlock, HorizontalBox, VerticalBox, EMultiBlockLocation::None, bSectionContainsIcons, OptionsBlock);
-		}
+			if( DragPreview.IsValid() && DragPreview.InsertIndex == Index )
+			{
+				// Add the drag preview before if we have it. This block shows where the custom block will be 
+				// added if the user drops it
+				AddBlockWidget(*DragPreview.PreviewBlock, HorizontalBox, VerticalBox, EMultiBlockLocation::None, bSectionContainsIcons, OptionsBlock);
+			}
 		
-		// Do not add a block if it is being dragged
-		if( !IsBlockBeingDragged( Blocks[Index] ) )
-		{
-			AddBlockWidget(Block, HorizontalBox, VerticalBox, Location, bSectionContainsIcons, OptionsBlock);
+			// Do not add a block if it is being dragged
+			if( !IsBlockBeingDragged( Blocks[Index] ) )
+			{
+				AddBlockWidget(Block, HorizontalBox, VerticalBox, Location, bSectionContainsIcons, OptionsBlock);
+			}
 		}
 	}
 
@@ -1094,7 +1249,6 @@ void SMultiBoxWidget::BuildMultiBoxWidget()
 	[
 		RootBorder.ToSharedRef()
 	];
-
 }
 
 TSharedRef<SWidget> SMultiBoxWidget::OnWrapButtonClicked()
@@ -1264,7 +1418,6 @@ void SMultiBoxWidget::OnCustomCommandDragEnter( TSharedRef<const FMultiBlock> Mu
 	}
 }
 
-
 void SMultiBoxWidget::OnCustomCommandDragged( TSharedRef<const FMultiBlock> MultiBlock, const FGeometry& MyGeometry, const FDragDropEvent& DragDropEvent )
 {
 	if( MultiBlock != DragPreview.PreviewBlock && MultiBox->IsInEditMode() )
@@ -1362,7 +1515,13 @@ FReply SMultiBoxWidget::FocusNextWidget(EUINavigation NavigationType)
 
 FReply SMultiBoxWidget::OnFocusReceived( const FGeometry& MyGeometry, const FFocusEvent& InFocusEvent )
 {
-	ResetSearch();
+	// Mouse DOWN on an element in the search result may fires OnFocusReceived if the mouse down event is not handled.
+	// If the search is cleared, the widget layout changes and this usually prevents the mouse UP to reach the selected
+	// widget which in turn doesn't trigger the selected item action.
+	if (InFocusEvent.GetCause() != EFocusCause::Mouse)
+	{
+		ResetSearch();
+	}
 
 	if (InFocusEvent.GetCause() == EFocusCause::Navigation)
 	{
@@ -1443,6 +1602,8 @@ void SMultiBoxWidget::ResetSearch()
 
 void SMultiBoxWidget::FilterMultiBoxEntries()
 {
+	VisibleFlattenHierarchyTips.Empty();
+
 	if (SearchText.IsEmpty())
 	{
 		for (auto It = MultiBoxWidgets.CreateConstIterator(); It; ++It)
@@ -1455,6 +1616,13 @@ void SMultiBoxWidget::FilterMultiBoxEntries()
 			SearchBlockWidget->SetVisibility(EVisibility::Collapsed);
 		}
 
+		// Hide the sub-menus widgets that were made visible by searching this multi-box hierarchy.
+		for (TPair<TSharedPtr<const FMultiBlock>, TSharedPtr<FFlattenSearchableBlockInfo>>& Pair: FlattenSearchableBlocks)
+		{
+			Pair.Value->Widget->SetVisibility(EVisibility::Collapsed);
+			Pair.Value->HierarchyTipVisibility = EVisibility::Collapsed;
+		}
+
 		// Return focus to parent widget
 		FSlateApplication::Get().SetUserFocus(0, SharedThis(this));
 
@@ -1463,28 +1631,49 @@ void SMultiBoxWidget::FilterMultiBoxEntries()
 
 	for(auto It = MultiBoxWidgets.CreateConstIterator(); It; ++It)
 	{
-		// Non-searched elements should not be rendered while searching
-		if(It.Value().IsEmpty())
+		FText DisplayText = UE::MultiBoxUtils::GetBlockWidgetDisplayText(It.Value());
+		const TSharedPtr<SWidget>& Widget = It.Key();
+
+		// Non-labeled elements should not be visible when searching
+		if (DisplayText.IsEmpty())
 		{
-			if(SearchText.IsEmpty())
-			{
-				It.Key()->SetVisibility( EVisibility::Visible );
-			}
-			else
-			{
-				It.Key()->SetVisibility( EVisibility::Collapsed );
-			}
+			Widget->SetVisibility(EVisibility::Collapsed);
 		}
 		else
 		{
-			// Compare widget text to the current search text
-			if( It.Value().ToString().Contains( SearchText.ToString() ) )
+			// Does the widget label match the searched text?
+			EVisibility WidgetVisibility = DisplayText.ToString().Contains(SearchText.ToString()) ? EVisibility::Visible : EVisibility::Collapsed;
+
+			// If this widget is one of the sub-menu item that was flatten to support recursive search
+			if (TSharedPtr<FFlattenSearchableBlockInfo>* FlattenBlockInfo = FlattenSearchableWidgets.Find(Widget))
 			{
-				It.Key()->SetVisibility( EVisibility::Visible );
+				(*FlattenBlockInfo)->Widget->SetVisibility(WidgetVisibility);
+				(*FlattenBlockInfo)->HierarchyTipVisibility = WidgetVisibility;
+
+				if (WidgetVisibility == EVisibility::Visible)
+				{
+					// Check if a item in the the same menu already enabled the 'hierarchy tip'. Because they are added in order, every next one with the same parent will be below each other. For example if
+					// the user searches for 'C' and the multibox has expendable menus "Cinematics" containing "Camera X" and "Shapes" containing "Cube", "Cone" and "Cylinder", we want the search result to look like:
+					// 
+					//    - Cinematics > (This is an expandable menu that matches the search)
+					// ##Shapes##        (This is the hierarchy tip enabled once for all matching items in the expandable 'Shapes' menu)
+					//    - Cube
+					//    - Cone
+					//    - Cylinder
+					// ##Cinematics##    (This is the hierarchy tip enabled once for the single matching item in the expandable 'Cinematics' menu)
+					//    - Camera X
+					//
+					bool bHierarchyTipAlreadyVisible = false;
+					VisibleFlattenHierarchyTips.FindOrAdd(UE::MultiBoxUtils::GetSearchHierarchyInfoText((*FlattenBlockInfo)->SearchableTextHierarchyComponents).ToString(), &bHierarchyTipAlreadyVisible);
+					if (bHierarchyTipAlreadyVisible)
+					{
+						(*FlattenBlockInfo)->HierarchyTipVisibility = EVisibility::Collapsed;
+					}
+				}
 			}
-			else
+			else // This widget is a normal item in this menu.
 			{
-				It.Key()->SetVisibility( EVisibility::Collapsed );
+				Widget->SetVisibility(WidgetVisibility);
 			}
 		}
 	}
@@ -1523,7 +1712,7 @@ void SMultiBoxWidget::AddElement(TSharedPtr<SWidget> BlockWidget, FText BlockDis
 		BlockDisplayText = FText::GetEmpty();
 	}
 
-	MultiBoxWidgets.Add(BlockWidget, BlockDisplayText);
+	MultiBoxWidgets.Add(BlockWidget, TArray<FText>{BlockDisplayText});
 }
 
 
diff --git a/Engine/Source/Runtime/Slate/Private/Framework/MultiBox/SMenuEntryBlock.h b/Engine/Source/Runtime/Slate/Private/Framework/MultiBox/SMenuEntryBlock.h
--- a/Engine/Source/Runtime/Slate/Private/Framework/MultiBox/SMenuEntryBlock.h
+++ b/Engine/Source/Runtime/Slate/Private/Framework/MultiBox/SMenuEntryBlock.h
@@ -118,6 +118,7 @@ public:
 	virtual void CreateMenuEntry(class FMenuBuilder& MenuBuilder) const override;
 	virtual bool HasIcon() const override;
 
+	bool IsSubMenu() const { return bIsSubMenu; }
 
 private:
 
@@ -201,6 +202,14 @@ public:
 	 */
 	void Construct( const FArguments& InArgs );
 
+	/**
+	 * Called to create content for a pull-down or sub-menu window when it's summoned by the user
+	 *
+	 * @return	The widget content for the new menu
+	 */
+	TSharedRef< SWidget > MakeNewMenuWidget() const;
+
+
 protected:
 	/** Struct for creating menu entry widgets */
 	struct FMenuEntryBuildParams
@@ -273,12 +282,6 @@ protected:
 	//virtual void Tick( const FGeometry& AllottedGeometry, const double InCurrentTime, const float InDeltaTime ) override;
 	// End of SWidget interface
 
-	/**
-	 * Called to create content for a pull-down or sub-menu window when it's summoned by the user
-	 *
-	 * @return	The widget content for the new menu
-	 */
-	TSharedRef< SWidget > MakeNewMenuWidget() const;
 
 	/**
 	 * Called to get the appropriate border for Buttons on Menu Bars based on whether or not submenu is open
diff --git a/Engine/Source/Runtime/Slate/Public/Framework/MultiBox/MultiBox.h b/Engine/Source/Runtime/Slate/Public/Framework/MultiBox/MultiBox.h
--- a/Engine/Source/Runtime/Slate/Public/Framework/MultiBox/MultiBox.h
+++ b/Engine/Source/Runtime/Slate/Public/Framework/MultiBox/MultiBox.h
@@ -191,6 +191,7 @@ public:
 	* @param	bSearchable		The searchable state to set
 	*/
 	void SetSearchable(bool bSearchable);
+
 	/**
 	* Gets the searchable state of this block
 	*
@@ -523,6 +524,7 @@ class SLATE_API SMultiBlockBaseWidget
 {
 
 public:
+	TSharedPtr<const FMultiBlock> GetBlock() const { return MultiBlock; }
 
 	/** IMultiBlockBaseWidget interface */
 	virtual TSharedRef< SWidget > AsWidget() override;
@@ -806,6 +808,19 @@ private:
 		bool IsValid() const { return BlockName != NAME_None && BlockType != EMultiBlockType::None && PreviewBlock.IsValid() && InsertIndex != INDEX_NONE; }
 	};
 
+	/** Contains information about sub-menu block widgets that were pulled and flatten in the parent menu to enable searching the entire menu tree from the top. */
+	struct FFlattenSearchableBlockInfo
+	{
+		/** The block widget searchable text along with its ancestor menu searchable texts. Ex. ["Menu_X", "Sub_Menu_Y", "BlockText"] */
+		TArray<FText> SearchableTextHierarchyComponents;
+
+		/** The flatten widget wrapping the block widget that was built for the associated block. The flatten widget adds a 'hierarchy tip' widget to indicate the real location of this item in the hierarchy. */
+		TSharedPtr<SWidget> Widget;
+
+		/** Whether the hierarchy tip widget is visible in the flatten search result. This tip let the user know the real location of this item in the hierarchy. */
+		EVisibility HierarchyTipVisibility = EVisibility::Collapsed;
+	};
+
 	/** The MultiBox we're associated with */
 	TSharedPtr< FMultiBox > MultiBox;
 
@@ -821,8 +836,17 @@ private:
 	/** A preview of a block being dragged inside this box */
 	FDraggedMultiBlockPreview DragPreview;
 
-	/* The multibox widgets that are contained, linked to their display text */
-	TMap<TSharedPtr<SWidget>, FText > MultiBoxWidgets;
+	/** The multibox widgets that are contained, linked to their searchable text hierarchies. */
+	TMap<TSharedPtr<SWidget>, TArray<FText>> MultiBoxWidgets;
+
+	/** The set of searchable blocks found in this multibox sub-menus that were collected and flatten in this multibox to support recursively searching this multibox hierarchy. */
+	TMap<TSharedPtr<const FMultiBlock>, TSharedPtr<FFlattenSearchableBlockInfo>> FlattenSearchableBlocks;
+
+	/** The set of searchable blocks widgets that were added from the list of flatten blocks collecteto display search results from this multibox hierarchy. */
+	TMap<TSharedPtr<SWidget>, TSharedPtr<FFlattenSearchableBlockInfo>> FlattenSearchableWidgets;
+
+	/** The list of visible hierarchy tips text (associated with flatten search widgets) in the search result. This is used to group sub-items matching the searched text and avoid showing the same tip several times. */
+	TSet<FString> VisibleFlattenHierarchyTips;
 
 	/* The search widget to be displayed at the top of the multibox */
 	TSharedPtr<SSearchBox> SearchTextWidget;
