Updated tab sidebar/drawer appearance with outline stroke extending to tab

diff --git a/Engine/Source/Runtime/Slate/Private/Framework/Docking/STabDrawer.cpp b/Engine/Source/Runtime/Slate/Private/Framework/Docking/STabDrawer.cpp
--- a/Engine/Source/Runtime/Slate/Private/Framework/Docking/STabDrawer.cpp
+++ b/Engine/Source/Runtime/Slate/Private/Framework/Docking/STabDrawer.cpp
@@ -17,11 +17,12 @@ void STabDrawer::SetCurrentSize(float InSize)
 	CurrentSize = FMath::Clamp(InSize, MinDrawerSize, TargetDrawerSize);
 }
 
-void STabDrawer::Construct(const FArguments& InArgs, TSharedRef<SDockTab> InTab, ETabDrawerOpenDirection InOpenDirection)
+void STabDrawer::Construct(const FArguments& InArgs, TSharedRef<SDockTab> InTab, TWeakPtr<SWidget> InTabButton, ETabDrawerOpenDirection InOpenDirection)
 {
 	OpenDirection = InOpenDirection;
 
 	ForTab = InTab;
+	TabButton = InTabButton;
 	OpenCloseAnimation = FCurveSequence(0.0f, 0.15f, ECurveEaseFunction::QuadOut);
 
 	CurrentSize = 0;
@@ -44,6 +45,18 @@ void STabDrawer::Construct(const FArguments& InArgs, TSharedRef<SDockTab> InTab,
 	BackgroundBrush = FAppStyle::Get().GetBrush("Docking.Sidebar.DrawerBackground");
 	ShadowBrush = FAppStyle::Get().GetBrush("Docking.Sidebar.DrawerShadow");
 	BorderBrush = FAppStyle::Get().GetBrush("Docking.Sidebar.Border");
+	if (OpenDirection == ETabDrawerOpenDirection::Left)
+	{
+		BorderSquareEdgeBrush = FAppStyle::Get().GetBrush("Docking.Sidebar.Border_SquareLeft");
+	}
+	else if (OpenDirection == ETabDrawerOpenDirection::Right)
+	{
+		BorderSquareEdgeBrush = FAppStyle::Get().GetBrush("Docking.Sidebar.Border_SquareRight");
+	}
+	else
+	{
+		BorderSquareEdgeBrush = BorderBrush;
+	}
 
 	FSlateApplication::Get().OnFocusChanging().AddSP(this, &STabDrawer::OnGlobalFocusChanging);
 
@@ -240,20 +253,27 @@ int32 STabDrawer::OnPaint(const FPaintArgs& Args, const FGeometry& AllottedGeome
 	const FGeometry RenderTransformedChildGeometry = GetRenderTransformedGeometry(AllottedGeometry);
 	const FGeometry ResizeHandleGeometry = GetResizeHandleGeometry(AllottedGeometry);
 
-	FPaintGeometry OffsetPaintGeom;
+	const FVector2D LocalSize = AllottedGeometry.GetLocalSize();
+	FVector2D ContentsLocalOrigin;
+	FVector2D ContentsLocalSize;
 	if (OpenDirection == ETabDrawerOpenDirection::Left)
 	{
-		OffsetPaintGeom = RenderTransformedChildGeometry.ToPaintGeometry(FVector2D(0.0f, ShadowOffset.Y), FVector2D(TargetDrawerSize, AllottedGeometry.GetLocalSize().Y - (ShadowOffset.Y * 2)));
+		ContentsLocalOrigin = FVector2D(0.0f, ShadowOffset.Y);
+		ContentsLocalSize = FVector2D(TargetDrawerSize, LocalSize.Y - (ShadowOffset.Y * 2));
 	}
 	else if (OpenDirection == ETabDrawerOpenDirection::Right)
 	{
-		OffsetPaintGeom = RenderTransformedChildGeometry.ToPaintGeometry(ShadowOffset, FVector2D(TargetDrawerSize, AllottedGeometry.GetLocalSize().Y - (ShadowOffset.Y * 2)));
+		ContentsLocalOrigin = ShadowOffset;
+		ContentsLocalSize = FVector2D(TargetDrawerSize, LocalSize.Y - (ShadowOffset.Y * 2));
 	}
 	else
 	{
-		OffsetPaintGeom = RenderTransformedChildGeometry.ToPaintGeometry(ShadowOffset, FVector2D(AllottedGeometry.GetLocalSize().X - (ShadowOffset.X * 2), TargetDrawerSize));
+		ContentsLocalOrigin = ShadowOffset;
+		ContentsLocalSize = FVector2D(LocalSize.X - (ShadowOffset.X * 2), TargetDrawerSize);
 	}
 
+	const FPaintGeometry OffsetPaintGeom = RenderTransformedChildGeometry.ToPaintGeometry(ContentsLocalOrigin, ContentsLocalSize);
+
 	// Draw the resize handle
 	if (bIsResizing || bIsResizeHandleHovered)
 	{
@@ -288,14 +308,106 @@ int32 STabDrawer::OnPaint(const FPaintArgs& Args, const FGeometry& AllottedGeome
 
 	int32 OutLayerId = SCompoundWidget::OnPaint(Args, RenderTransformedChildGeometry, MyCullingRect, OutDrawElements, LayerId, InWidgetStyle, bParentEnabled);
 
+	TSharedPtr<SWidget> TabButtonSP = TabButton.Pin();
+
 	// Top border
-	FSlateDrawElement::MakeBox(
-		OutDrawElements,
-		OutLayerId,
-		OffsetPaintGeom,
-		BorderBrush,
-		ESlateDrawEffect::None,
-		BorderBrush->GetTint(InWidgetStyle));
+	if (OpenDirection == ETabDrawerOpenDirection::Bottom || TabButtonSP == nullptr)
+	{
+		// When opened from the bottom, draw the full border.
+		// Cutting out the "notch" for the corresponding tab is only supported in left/right orientations.
+		FSlateDrawElement::MakeBox(
+			OutDrawElements,
+			OutLayerId,
+			OffsetPaintGeom,
+			BorderBrush,
+			ESlateDrawEffect::None,
+			BorderBrush->GetTint(InWidgetStyle));
+	}
+	else
+	{
+		// Example of how border box is drawn with the tab notch cut out on the right side
+		// (OpenDirection == ETabDrawerOpenDirection::Right)
+		//
+		//                    + - - - - - - +
+		//                    : /---------\ :
+		// ClipAboveTabButton : |         | : 
+		//                    : |         | :
+		//            TabTopY + - - - - - - +
+		//                    : |           :  |
+		// ClipAtTabButton    : |           :  |  (right edge outside clip is clipped off)
+		//                    : |           :  |
+		//         TabBottomY + - - - - - - +
+		//                    : |         | :
+		// ClipBelowTabButton : |         | :
+		//                    : \---------/ :
+		//                    + - - - - - - +
+		//                                  <-->
+		//                              NotchOffset
+		//
+		// Originally, I tried making the middle clip region thinner (to clip out the notch)
+		// while keeping the geometry identical, but this looks worse when the tab notch needs to
+		// be at the top or bottom, since the top/bottom edge of the border wouldn't extend all the
+		// way to the edge.
+		const FGeometry TabButtonGeometry = TabButtonSP->GetPaintSpaceGeometry();
+
+		// Compute the top/bottom of the tab in our local space.
+		const float BorderWidth = BorderBrush->OutlineSettings.Width;
+		const float TabTopY = RenderTransformedChildGeometry.AbsoluteToLocal(TabButtonGeometry.GetAbsolutePositionAtCoordinates(FVector2D::ZeroVector)).Y + 0.5f * BorderWidth;
+		const float TabBottomY = RenderTransformedChildGeometry.AbsoluteToLocal(TabButtonGeometry.GetAbsolutePositionAtCoordinates(FVector2D::UnitVector)).Y - 0.5f * BorderWidth;
+
+		// Create the geometry for the notched portion, where one edge extends past the clipping rect.
+		const FVector2D NotchOffsetSize(TabButtonGeometry.GetLocalSize().X, 0.0f);
+		const FVector2D NotchOffsetTranslate = OpenDirection == ETabDrawerOpenDirection::Left ? -NotchOffsetSize : FVector2D::ZeroVector;
+		const FPaintGeometry NotchOffsetPaintGeom = RenderTransformedChildGeometry.ToPaintGeometry(ContentsLocalOrigin + NotchOffsetTranslate, ContentsLocalSize + NotchOffsetSize);
+
+		// Split the border box into three clipping zones.
+		const FPaintGeometry ClipAboveTabButton = RenderTransformedChildGeometry.ToPaintGeometry(FVector2D(0, 0), FVector2D(LocalSize.X, TabTopY));
+		const FPaintGeometry ClipAtTabButton = RenderTransformedChildGeometry.ToPaintGeometry(FVector2D(0, TabTopY), FVector2D(LocalSize.X, TabBottomY - TabTopY));
+		const FPaintGeometry ClipBelowTabButton = RenderTransformedChildGeometry.ToPaintGeometry(FVector2D(0, TabBottomY), FVector2D(LocalSize.X, LocalSize.Y - TabBottomY));
+
+		// If the tab button touches a corner on the edge of the border, switch the brush to
+		// draw that corner squared-off. When a tab is near the very top or bottom of its sidebar,
+		// this makes the outline look slightly nicer and more connected.
+		const int32 UpperCornerIndex = OpenDirection == ETabDrawerOpenDirection::Left ? 0 : 1;
+		const int32 LowerCornerIndex = OpenDirection == ETabDrawerOpenDirection::Left ? 3 : 2;
+		const bool bTabTouchesUpperCorner = TabTopY < ShadowOffset.Y + BorderBrush->OutlineSettings.CornerRadii[UpperCornerIndex];
+		const bool bTabTouchesLowerCorner = TabBottomY > LocalSize.Y - ShadowOffset.Y - BorderBrush->OutlineSettings.CornerRadii[LowerCornerIndex];
+		const FSlateBrush* AboveTabBrush = bTabTouchesUpperCorner ? BorderSquareEdgeBrush : BorderBrush;
+		const FSlateBrush* BelowTabBrush = bTabTouchesLowerCorner ? BorderSquareEdgeBrush : BorderBrush;
+
+		// Draw portion above the tab
+		OutDrawElements.PushClip(FSlateClippingZone(ClipAboveTabButton));
+		FSlateDrawElement::MakeBox(
+			OutDrawElements,
+			OutLayerId,
+			OffsetPaintGeom,
+			AboveTabBrush,
+			ESlateDrawEffect::None,
+			AboveTabBrush->GetTint(InWidgetStyle));
+		OutDrawElements.PopClip();
+
+		// Draw "notched" portion next to the tab
+		OutDrawElements.PushClip(FSlateClippingZone(ClipAtTabButton));
+		FSlateDrawElement::MakeBox(
+			OutDrawElements,
+			OutLayerId,
+			NotchOffsetPaintGeom,
+			BorderSquareEdgeBrush,
+			ESlateDrawEffect::None,
+			BorderSquareEdgeBrush->GetTint(InWidgetStyle));
+		OutDrawElements.PopClip();
+
+		// Draw portion below the tab
+		OutDrawElements.PushClip(FSlateClippingZone(ClipBelowTabButton));
+		FSlateDrawElement::MakeBox(
+			OutDrawElements,
+			OutLayerId,
+			OffsetPaintGeom,
+			BelowTabBrush,
+			ESlateDrawEffect::None,
+			BelowTabBrush->GetTint(InWidgetStyle));
+		OutDrawElements.PopClip();
+	}
 
 	return OutLayerId+1;
 
diff --git a/Engine/Source/Runtime/Slate/Private/Framework/Docking/STabSidebar.cpp b/Engine/Source/Runtime/Slate/Private/Framework/Docking/STabSidebar.cpp
--- a/Engine/Source/Runtime/Slate/Private/Framework/Docking/STabSidebar.cpp
+++ b/Engine/Source/Runtime/Slate/Private/Framework/Docking/STabSidebar.cpp
@@ -54,7 +54,7 @@ public:
 		TextLayoutCache->SetTextOverflowPolicy(InArgs._OverflowPolicy.IsSet() ? InArgs._OverflowPolicy : TextStyle.OverflowPolicy);
 	}
 
-	int32 OnPaint(const FPaintArgs& Args, const FGeometry& AllottedGeometry, const FSlateRect& MyCullingRect, FSlateWindowElementList& OutDrawElements, int32 LayerId, const FWidgetStyle& InWidgetStyle, bool bParentEnabled) const
+	virtual int32 OnPaint(const FPaintArgs& Args, const FGeometry& AllottedGeometry, const FSlateRect& MyCullingRect, FSlateWindowElementList& OutDrawElements, int32 LayerId, const FWidgetStyle& InWidgetStyle, bool bParentEnabled) const override
 	{
 		// We're going to figure out the bounds of the corresponding horizontal text, and then rotate it into a vertical orientation.
 		const FVector2D LocalSize = AllottedGeometry.GetLocalSize();
@@ -155,25 +155,15 @@ public:
 			.HeightOverride(Size.X)
 			.Clipping(EWidgetClipping::ClipToBounds)
 			[
-				SAssignNew(MainButton, SButton)
-				.ToolTip(ForTab->GetToolTip() ? ForTab->GetToolTip() : TAttribute<TSharedPtr<IToolTip>>())
-				.ToolTipText(ForTab->GetToolTip() ? TAttribute<FText>() : ForTab->GetTabLabel())
-				.ContentPadding(FMargin(0.0f, DockTabStyle->TabPadding.Top, 0.0f, DockTabStyle->TabPadding.Bottom))
-				.OnClicked_Lambda([this](){OnDrawerButtonClicked.ExecuteIfBound(Tab.ToSharedRef()); return FReply::Handled(); })
-				.ForegroundColor(FSlateColor::UseForeground())
+				SNew(SOverlay)
+				+ SOverlay::Slot()
 				[
-					SNew(SOverlay)
-					+ SOverlay::Slot()
-					.HAlign(Location == ESidebarLocation::Left ? HAlign_Left : HAlign_Right)
-					[
-						SAssignNew(OpenIndicator, SComplexGradient)
-						.DesiredSizeOverride(FVector2D(1.0f, 1.0f))
-						.GradientColors(GradientStops)
-						.Orientation(EOrientation::Orient_Horizontal)
-					]
-					+ SOverlay::Slot()
-					.VAlign(VAlign_Fill)
-					.HAlign(HAlign_Center)
+					SAssignNew(MainButton, SButton)
+					.ToolTip(ForTab->GetToolTip() ? ForTab->GetToolTip() : TAttribute<TSharedPtr<IToolTip>>())
+					.ToolTipText(ForTab->GetToolTip() ? TAttribute<FText>() : ForTab->GetTabLabel())
+					.ContentPadding(FMargin(0.0f, DockTabStyle->TabPadding.Top, 0.0f, DockTabStyle->TabPadding.Right))
+					.OnClicked_Lambda([this](){OnDrawerButtonClicked.ExecuteIfBound(Tab.ToSharedRef()); return FReply::Handled(); })
+					.ForegroundColor(FSlateColor::UseForeground())
 					[
 						SNew(SVerticalBox)
 						+ SVerticalBox::Slot()
@@ -192,6 +182,7 @@ public:
 						+ SVerticalBox::Slot()
 						.Padding(0.0f, 5.0f, 0.0f, 0.0f)
 						.FillHeight(1.0f)
+						.HAlign(HAlign_Center)
 						[
 							SAssignNew(Label, STabDrawerTextBlock)
 								.TextStyle(&DockTabStyle->TabTextStyle)
@@ -201,6 +192,19 @@ public:
 						]
 					]
 				]
+				+ SOverlay::Slot()
+				[
+					SAssignNew(OpenBorder, SBorder)
+					.Visibility(EVisibility::HitTestInvisible)
+				]
+				+ SOverlay::Slot()
+				.HAlign(Location == ESidebarLocation::Left ? HAlign_Left : HAlign_Right)
+				[
+					SAssignNew(OpenIndicator, SComplexGradient)
+					.DesiredSizeOverride(FVector2D(1.0f, 1.0f))
+					.GradientColors(GradientStops)
+					.Orientation(EOrientation::Orient_Horizontal)
+				]
 			]
 		];
 
@@ -230,11 +234,24 @@ public:
 		{
 			// this button is the one with the tab that is actually opened so show the opened indicator
 			OpenIndicator->SetVisibility(EVisibility::HitTestInvisible);
+			OpenBorder->SetVisibility(EVisibility::HitTestInvisible);
 			MainButton->SetButtonStyle(&FAppStyle::Get().GetWidgetStyle<FButtonStyle>("Docking.SidebarButton.Opened"));
+
+			switch (Location)
+			{
+			case ESidebarLocation::Left:
+				OpenBorder->SetBorderImage(FAppStyle::Get().GetBrush("Docking.Sidebar.Border_SquareRight"));
+				break;
+			case ESidebarLocation::Right:
+			default:
+				OpenBorder->SetBorderImage(FAppStyle::Get().GetBrush("Docking.Sidebar.Border_SquareLeft"));
+				break;
+			}
 		}
 		else
 		{
 			OpenIndicator->SetVisibility(EVisibility::Collapsed);
+			OpenBorder->SetVisibility(EVisibility::Collapsed);
 			MainButton->SetButtonStyle(&FAppStyle::Get().GetWidgetStyle<FButtonStyle>("Docking.SidebarButton.Closed"));
 		}
 	}
@@ -276,6 +293,10 @@ public:
 		{
 			return DockTabStyle->ActiveForegroundColor;
 		}
+		else if (IsHovered())
+		{
+			return DockTabStyle->HoveredForegroundColor;
+		}
 
 		return FSlateColor::UseStyle();
 	}
@@ -284,6 +305,7 @@ private:
 	TSharedPtr<SDockTab> Tab;
 	TSharedPtr<STabDrawerTextBlock> Label;
 	TSharedPtr<SWidget> OpenIndicator;
+	TSharedPtr<SBorder> OpenBorder;
 	TSharedPtr<SButton> MainButton;
 	FOnGetContent OnGetContextMenuContent;
 	FOnTabDrawerButtonClicked OnDrawerButtonClicked;
@@ -308,7 +330,11 @@ void STabSidebar::Construct(const FArguments& InArgs)
 	FGlobalTabmanager::Get()->OnTabForegrounded_Subscribe(FOnActiveTabChanged::FDelegate::CreateSP(this, &STabSidebar::OnActiveTabChanged));
 
 	ChildSlot
-	.Padding(2.0f,0.0f)
+	.Padding(FMargin(
+		Location == ESidebarLocation::Right ? 2.0f : 0.0f,
+		0.0f,
+		Location == ESidebarLocation::Left ? 2.0f : 0.0f,
+		0.0f))
 	[
 		SNew(SBorder)
 		.Padding(0.0f)
@@ -652,7 +678,7 @@ void STabSidebar::OpenDrawerInternal(TSharedRef<SDockTab> ForTab)
 
 		const FVector2D ShadowOffset(8, 8);
 
-		const float DPIScale = MyWindow->GetDPIScaleFactor();
+		const float WindowScale = FSlateApplication::Get().GetApplicationScale() * MyWindow->GetDPIScaleFactor();
 
 		const float TopOffset = ((ChildSlot.GetPadding().Top + MyGeometry.GetAbsolutePosition().Y) - WindowGeometry.GetAbsolutePosition().Y);
 
@@ -660,9 +686,9 @@ void STabSidebar::OpenDrawerInternal(TSharedRef<SDockTab> ForTab)
 
 		FMargin SlotPadding(
 			Location == ESidebarLocation::Left ? MinDrawerSize : 0.0f,
-			TopOffset / DPIScale - ShadowOffset.Y,
+			TopOffset / WindowScale - ShadowOffset.Y,
 			Location == ESidebarLocation::Right ? MinDrawerSize : 0.0f,
-			BottomOffset / DPIScale - ShadowOffset.Y
+			BottomOffset / WindowScale - ShadowOffset.Y
 		);
 
 		const float MaxPct = .5f;
@@ -674,8 +700,14 @@ void STabSidebar::OpenDrawerInternal(TSharedRef<SDockTab> ForTab)
 
 		const float TargetDrawerSize = (MyWindow->GetSizeInScreen().X * TargetDrawerSizePct) / MyWindow->GetDPIScaleFactor();
 
+		const TPair<TSharedRef<SDockTab>, TSharedRef<STabDrawerButton>>* TabEntry = Tabs.FindByPredicate(
+			[ForTab](auto TabPair) {
+				return TabPair.Key == ForTab;
+			});
+		const TWeakPtr<SWidget> TabButton = TabEntry ? TWeakPtr<SWidget>(TabEntry->Value) : nullptr;
+
 		TSharedRef<STabDrawer> NewDrawer =
-			SNew(STabDrawer, ForTab, Location == ESidebarLocation::Left ? ETabDrawerOpenDirection::Left : ETabDrawerOpenDirection::Right)
+			SNew(STabDrawer, ForTab, TabButton, Location == ESidebarLocation::Left ? ETabDrawerOpenDirection::Left : ETabDrawerOpenDirection::Right)
 			.MinDrawerSize(MinDrawerSize)
 			.TargetDrawerSize(TargetDrawerSize)
 			.MaxDrawerSize(MaxDrawerSize)
diff --git a/Engine/Source/Runtime/Slate/Public/Framework/Docking/STabDrawer.h b/Engine/Source/Runtime/Slate/Public/Framework/Docking/STabDrawer.h
--- a/Engine/Source/Runtime/Slate/Public/Framework/Docking/STabDrawer.h
+++ b/Engine/Source/Runtime/Slate/Public/Framework/Docking/STabDrawer.h
@@ -57,7 +57,7 @@ class STabDrawer : public SCompoundWidget
 	SLATE_END_ARGS()
 
 	~STabDrawer();
-	void Construct(const FArguments& InArgs, TSharedRef<SDockTab> InTab, ETabDrawerOpenDirection InOpenDirection);
+	void Construct(const FArguments& InArgs, TSharedRef<SDockTab> InTab, TWeakPtr<SWidget> InTabButton, ETabDrawerOpenDirection InOpenDirection);
 
 	/** Sets the current size of the drawer, ignoring any open/close animation */
 	void SetCurrentSize(float InSize);
@@ -98,12 +98,14 @@ private:
 	FOnDrawerClosed OnDrawerClosed;
 	TSharedPtr<FActiveTimerHandle> OpenCloseTimer;
 	TSharedPtr<SDockTab> ForTab;
+	TWeakPtr<SWidget> TabButton;
 	FThrottleRequest ResizeThrottleHandle;
 	FThrottleRequest AnimationThrottle;
 	FCurveSequence OpenCloseAnimation;
 	const FSlateBrush* BackgroundBrush;
 	const FSlateBrush* ShadowBrush;
-	const FSlateBrush* BorderBrush;
+	const FSlateBrush* BorderBrush; // border with all corners rounded
+	const FSlateBrush* BorderSquareEdgeBrush; // border with corners squared on one edge depending on open direction
 	const FSplitterStyle* SplitterStyle;
 	FVector2D ShadowOffset;
 	float ExpanderSize;
diff --git a/Engine/Source/Runtime/SlateCore/Private/Styling/StarshipCoreStyle.cpp b/Engine/Source/Runtime/SlateCore/Private/Styling/StarshipCoreStyle.cpp
--- a/Engine/Source/Runtime/SlateCore/Private/Styling/StarshipCoreStyle.cpp
+++ b/Engine/Source/Runtime/SlateCore/Private/Styling/StarshipCoreStyle.cpp
@@ -1630,8 +1630,8 @@ void FStarshipCoreStyle::SetupDockingStyles(TSharedRef<FStyle>& Style)
 
 	FButtonStyle SidebarTabButtonOpened =
 		FButtonStyle(NoBorder)
-		.SetNormal(MinorTabStyle.ForegroundBrush)
-		.SetHovered(MinorTabStyle.ForegroundBrush)
+		.SetNormal(FSlateColorBrush(FStyleColors::Panel))
+		.SetHovered(FSlateColorBrush(FStyleColors::Panel))
 		.SetNormalForeground(MinorTabStyle.NormalForegroundColor)
 		.SetPressedForeground(MinorTabStyle.NormalForegroundColor)
 		.SetHoveredForeground(MinorTabStyle.HoveredForegroundColor);
@@ -1639,8 +1639,8 @@ void FStarshipCoreStyle::SetupDockingStyles(TSharedRef<FStyle>& Style)
 
 	FButtonStyle SidebarTabButtonClosed =
 		FButtonStyle(NoBorder)
-		.SetNormal(MinorTabStyle.NormalBrush)
-		.SetHovered(MinorTabStyle.HoveredBrush)
+		.SetNormal(FSlateNoResource())
+		.SetHovered(FSlateColorBrush(FStyleColors::Panel))
 		.SetNormalForeground(MinorTabStyle.NormalForegroundColor)
 		.SetPressedForeground(MinorTabStyle.NormalForegroundColor)
 		.SetHoveredForeground(MinorTabStyle.HoveredForegroundColor);
@@ -1651,8 +1651,16 @@ void FStarshipCoreStyle::SetupDockingStyles(TSharedRef<FStyle>& Style)
 
 	Style->Set("Docking.Sidebar.DrawerShadow", new BOX_BRUSH("/Starship/Docking/drawer-shadow", FMargin(8.f /64.f), FLinearColor(0.f, 0.f, 0.f, 1.f)));
 	Style->Set("Docking.Sidebar.DrawerBackground", new FSlateColorBrush(FStyleColors::Panel));
-	Style->Set("Docking.Sidebar.Background", new FSlateColorBrush(FStyleColors::Recessed));
-	Style->Set("Docking.Sidebar.Border", new FSlateRoundedBoxBrush(FSlateColor(FLinearColor::Transparent), 5.0f, FStyleColors::Hover, 1.0f) );
+	Style->Set("Docking.Sidebar.Background", new FSlateColorBrush(FStyleColors::Background));
+
+	// Create the regular border brush with all corners rounded and variants with certain corners squared.
+	{
+		const float CornerRadius = 5.0f;
+		const float BorderWidth = 1.0f;
+		Style->Set("Docking.Sidebar.Border", new FSlateRoundedBoxBrush(FStyleColors::Transparent, CornerRadius, FStyleColors::Hover, BorderWidth));
+		Style->Set("Docking.Sidebar.Border_SquareLeft", new FSlateRoundedBoxBrush(FStyleColors::Transparent, FVector4(0.0f, CornerRadius, CornerRadius, 0.0f), FStyleColors::Hover, BorderWidth));
+		Style->Set("Docking.Sidebar.Border_SquareRight", new FSlateRoundedBoxBrush(FStyleColors::Transparent, FVector4(CornerRadius, 0.0f, 0.0f, CornerRadius), FStyleColors::Hover, BorderWidth));
+	}
 
 	// Dock Cross
 	Style->Set("Docking.Cross.DockLeft", new IMAGE_BRUSH("/Docking/OuterDockingIndicator", FVector2D(6.f, 6.f), FLinearColor(1.0f, 0.35f, 0.0f, 0.25f)));
@@ -2123,4 +2131,4 @@ void FStarshipCoreStyle::SetupMultiboxStyles(TSharedRef<FStyle>& Style)
 }
 
 #undef RootToContentDir
-#undef FONT
\ No newline at end of file
+#undef FONT
