Account for window scale/dpi when resizing tab drawer; clip contents to non-shadow area of drawer

diff --git a/Engine/Source/Runtime/Slate/Private/Framework/Docking/STabDrawer.cpp b/Engine/Source/Runtime/Slate/Private/Framework/Docking/STabDrawer.cpp
--- a/Engine/Source/Runtime/Slate/Private/Framework/Docking/STabDrawer.cpp
+++ b/Engine/Source/Runtime/Slate/Private/Framework/Docking/STabDrawer.cpp
@@ -5,6 +5,7 @@
 #include "Styling/SlateTypes.h"
 #include "Layout/ArrangedChildren.h"
 #include "Framework/Application/SlateApplication.h"
+#include "Widgets/Layout/SBox.h"
 
 STabDrawer::~STabDrawer()
 {
@@ -65,7 +66,12 @@ void STabDrawer::Construct(const FArguments& InArgs, TSharedRef<SDockTab> InTab,
 
 	ChildSlot
 	[
-		InArgs._Content.Widget
+		SNew(SBox)
+		.Clipping(EWidgetClipping::ClipToBounds)
+		.Content()
+		[
+			InArgs._Content.Widget
+		]
 	];
 }
 
@@ -209,15 +215,15 @@ FReply STabDrawer::OnMouseMove(const FGeometry& AllottedGeometry, const FPointer
 
 		if (OpenDirection == ETabDrawerOpenDirection::Left)
 		{
-			DeltaSize = (MousePos - InitialResizeGeometry.GetAbsolutePositionAtCoordinates(FVector2D::ZeroVector)).X;
+			DeltaSize = InitialResizeGeometry.AbsoluteToLocal(MousePos).X;
 		}
 		else if (OpenDirection == ETabDrawerOpenDirection::Right)
 		{
-			DeltaSize = (InitialResizeGeometry.GetAbsolutePositionAtCoordinates(FVector2D::ZeroVector) - MousePos).X;
+			DeltaSize = -InitialResizeGeometry.AbsoluteToLocal(MousePos).X;
 		}
 		else
 		{
-			DeltaSize = (InitialResizeGeometry.GetAbsolutePositionAtCoordinates(FVector2D::ZeroVector) - MousePos).Y;
+			DeltaSize = -InitialResizeGeometry.AbsoluteToLocal(MousePos).Y;
 		}
 
 
