Make STabSidebar insert new STabDrawer's into an overlay within SDockingArea instead of as a window overlay

This is in preparation for some work for sidebar tab pinning;
it will help simplify handling of pinned minor/nomad tabs when switching between different major tabs in the editor,
as the pinned tabs will just swap out with the rest of the major tab's contents when switching tabs.

Removed DPI-scale-changed handling for tab sidebar because the drawers are now created relative to
the SDockingArea and should properly resize when DPI changes.

diff --git a/Engine/Source/Runtime/Slate/Private/Framework/Docking/SDockingArea.cpp b/Engine/Source/Runtime/Slate/Private/Framework/Docking/SDockingArea.cpp
--- a/Engine/Source/Runtime/Slate/Private/Framework/Docking/SDockingArea.cpp
+++ b/Engine/Source/Runtime/Slate/Private/Framework/Docking/SDockingArea.cpp
@@ -23,6 +23,8 @@ void SDockingArea::Construct( const FArguments& InArgs, const TSharedRef<FTabMan
 	const TAttribute<EVisibility> TargetCrossVisibility = TAttribute<EVisibility>(SharedThis(this), &SDockingArea::TargetCrossVisibility);
 	const TAttribute<EVisibility> TargetCrossCenterVisibility = TAttribute<EVisibility>(SharedThis(this), &SDockingArea::TargetCrossCenterVisibility);
 
+	const TSharedRef<SOverlay> SidebarDrawersOverlay = SNew(SOverlay);
+
 	// In DockSplitter mode we just act as a thin shell around a Splitter widget
 	this->ChildSlot
 	[
@@ -34,7 +36,7 @@ void SDockingArea::Construct( const FArguments& InArgs, const TSharedRef<FTabMan
 			+ SHorizontalBox::Slot()
 			.AutoWidth()
 			[
-				SAssignNew(LeftSidebar, STabSidebar)
+				SAssignNew(LeftSidebar, STabSidebar, SidebarDrawersOverlay)
 				.Location(ESidebarLocation::Left)
 				.Visibility(EVisibility::Collapsed)
 			]
@@ -46,7 +48,7 @@ void SDockingArea::Construct( const FArguments& InArgs, const TSharedRef<FTabMan
 			+SHorizontalBox::Slot()
 			.AutoWidth()
 			[
-				SAssignNew(RightSidebar, STabSidebar)
+				SAssignNew(RightSidebar, STabSidebar, SidebarDrawersOverlay)
 				.Location(ESidebarLocation::Right)
 				.Visibility(EVisibility::Collapsed)
 			]
@@ -101,6 +103,10 @@ void SDockingArea::Construct( const FArguments& InArgs, const TSharedRef<FTabMan
 			.OwnerNode( SharedThis(this) )
 			.DockDirection( SDockingNode::Center )
 		]
+		+SOverlay::Slot()
+		[
+			SidebarDrawersOverlay
+		]
 	];
 
 
diff --git a/Engine/Source/Runtime/Slate/Private/Framework/Docking/STabSidebar.cpp b/Engine/Source/Runtime/Slate/Private/Framework/Docking/STabSidebar.cpp
--- a/Engine/Source/Runtime/Slate/Private/Framework/Docking/STabSidebar.cpp
+++ b/Engine/Source/Runtime/Slate/Private/Framework/Docking/STabSidebar.cpp
@@ -319,13 +319,10 @@ STabSidebar::~STabSidebar()
 	RemoveAllDrawers();
 }
 
-void STabSidebar::Construct(const FArguments& InArgs)
+void STabSidebar::Construct(const FArguments& InArgs, TSharedRef<SOverlay> InDrawersOverlay)
 {
 	Location = InArgs._Location;
-
-#if WITH_EDITOR
-	FSlateApplication::Get().OnWindowDPIScaleChanged().AddSP(this, &STabSidebar::OnWindowDPIScaleChanged);
-#endif
+	DrawersOverlay = InDrawersOverlay;
 
 	FGlobalTabmanager::Get()->OnTabForegrounded_Subscribe(FOnActiveTabChanged::FDelegate::CreateSP(this, &STabSidebar::OnActiveTabChanged));
 
@@ -491,20 +488,9 @@ void STabSidebar::OnTabDrawerClosed(TSharedRef<STabDrawer> Drawer)
 void STabSidebar::OnTargetDrawerSizeChanged(TSharedRef<STabDrawer> Drawer, float NewSize)
 {
 	TSharedRef<SDockTab> Tab = Drawer->GetTab();
-	TSharedPtr<SWindow> MyWindow = FSlateApplication::Get().FindWidgetWindow(AsShared());
-	if (MyWindow.IsValid())
-	{
-		const float TargetDrawerSizePct = NewSize / MyWindow->GetPaintSpaceGeometry().GetLocalSize().X;
-		Tab->GetParentDockTabStack()->SetTabSidebarSizeCoefficient(Tab, TargetDrawerSizePct);
-	}
-}
 
-void STabSidebar::OnWindowDPIScaleChanged(TSharedRef<SWindow> WindowThatChanged)
-{
-	if (WindowThatChanged == WindowWithOverlayContent)
-	{
-		RemoveAllDrawers();
-	}
+	const float TargetDrawerSizePct = NewSize / DrawersOverlay->GetPaintSpaceGeometry().GetLocalSize().X;
+	Tab->GetParentDockTabStack()->SetTabSidebarSizeCoefficient(Tab, TargetDrawerSizePct);
 }
 
 void STabSidebar::OnActiveTabChanged(TSharedPtr<SDockTab> NewlyActivated, TSharedPtr<SDockTab> PreviouslyActive)
@@ -581,21 +567,12 @@ void STabSidebar::RemoveDrawer(TSharedRef<SDockTab> ForTab)
 	{
 		TSharedRef<STabDrawer> OpenedDrawerRef = *OpenedDrawer;
 
-		if (TSharedPtr<SWindow> MyWindow = WindowWithOverlayContent.Pin())
-		{
-			bool bRemoveSuccessful = MyWindow->RemoveOverlaySlot(OpenedDrawerRef);
-			ensure(bRemoveSuccessful);
-		}
+		bool bRemoveSuccessful = DrawersOverlay->RemoveSlot(OpenedDrawerRef);
+		ensure(bRemoveSuccessful);
 
 		OpenedDrawers.Remove(OpenedDrawerRef);
 	}
 
-	if (OpenedDrawers.Num() == 0)
-	{
-		WindowWithOverlayContent.Reset();
-	}
-
-
 	ForTab->OnTabDrawerClosed();
 
 	UpdateDrawerAppearance();
@@ -668,37 +645,31 @@ void STabSidebar::OpenDrawerInternal(TSharedRef<SDockTab> ForTab)
 	{
 		PendingTabToOpen.Reset();
 
-		TSharedPtr<SWindow> MyWindow = FSlateApplication::Get().FindWidgetWindow(AsShared());
-
-		const FGeometry WindowGeometry = MyWindow->GetTickSpaceGeometry();
+		const FGeometry DrawersOverlayGeometry = DrawersOverlay->GetTickSpaceGeometry();
 		const FGeometry MyGeometry = GetTickSpaceGeometry();
 
 		// Calculate padding for the drawer itself
-		const float MinDrawerSize = MyGeometry.GetLocalSize().X + MyWindow->GetWindowBorderSize().Left;
+		const float MinDrawerSize = MyGeometry.GetLocalSize().X - 4.0f; // overlap with sidebar border slightly
 
 		const FVector2D ShadowOffset(8, 8);
 
-		const float WindowScale = FSlateApplication::Get().GetApplicationScale() * MyWindow->GetDPIScaleFactor();
-
-		const float TopOffset = ((ChildSlot.GetPadding().Top + MyGeometry.GetAbsolutePosition().Y) - WindowGeometry.GetAbsolutePosition().Y);
-
-		const float BottomOffset = ((ChildSlot.GetPadding().Bottom + WindowGeometry.GetAbsolutePositionAtCoordinates(FVector2D::UnitVector).Y) - MyGeometry.GetAbsolutePositionAtCoordinates(FVector2D::UnitVector).Y);
-
 		FMargin SlotPadding(
 			Location == ESidebarLocation::Left ? MinDrawerSize : 0.0f,
-			TopOffset / WindowScale - ShadowOffset.Y,
+			-ShadowOffset.Y,
 			Location == ESidebarLocation::Right ? MinDrawerSize : 0.0f,
-			BottomOffset / WindowScale - ShadowOffset.Y
+			-ShadowOffset.Y
 		);
 
+		const float AvailableWidth = DrawersOverlayGeometry.GetLocalSize().X - SlotPadding.GetTotalSpaceAlong<EOrientation::Orient_Horizontal>();
+
 		const float MaxPct = .5f;
-		const float MaxDrawerSize = MyWindow->GetSizeInScreen().X * 0.50f;
+		const float MaxDrawerSize = AvailableWidth * 0.50f;
 
 		float TargetDrawerSizePct = ForTab->GetParentDockTabStack()->GetTabSidebarSizeCoefficient(ForTab);
 		TargetDrawerSizePct = FMath::Clamp(TargetDrawerSizePct, .0f, .5f);
 
 
-		const float TargetDrawerSize = (MyWindow->GetSizeInScreen().X * TargetDrawerSizePct) / MyWindow->GetDPIScaleFactor();
+		const float TargetDrawerSize = AvailableWidth * TargetDrawerSizePct;
 
 		const TPair<TSharedRef<SDockTab>, TSharedRef<STabDrawerButton>>* TabEntry = Tabs.FindByPredicate(
 			[ForTab](auto TabPair) {
@@ -718,10 +689,7 @@ void STabSidebar::OpenDrawerInternal(TSharedRef<SDockTab> ForTab)
 				ForTab->GetContent()
 			];
 
-		check(!WindowWithOverlayContent.IsValid() || WindowWithOverlayContent == MyWindow);
-		WindowWithOverlayContent = MyWindow;
-
-		MyWindow->AddOverlaySlot()
+		DrawersOverlay->AddSlot()
 			.Padding(SlotPadding)
 			.HAlign(Location == ESidebarLocation::Left ? HAlign_Left : HAlign_Right)
 			[
diff --git a/Engine/Source/Runtime/Slate/Private/Framework/Docking/STabSidebar.h b/Engine/Source/Runtime/Slate/Private/Framework/Docking/STabSidebar.h
--- a/Engine/Source/Runtime/Slate/Private/Framework/Docking/STabSidebar.h
+++ b/Engine/Source/Runtime/Slate/Private/Framework/Docking/STabSidebar.h
@@ -7,6 +7,7 @@
 #include "Widgets/SCompoundWidget.h"
 
 class SDockTab;
+class SOverlay;
 class STabDrawerButton;
 class SVerticalBox;
 class STabDrawer;
@@ -27,7 +28,7 @@ class STabSidebar : public SCompoundWidget
 public:
 	~STabSidebar();
 
-	void Construct(const FArguments& InArgs);
+	void Construct(const FArguments& InArgs, TSharedRef<SOverlay> InDrawersOverlay);
 
 	/**
 	 * Sets an offset for the sidebar from the top of the window
@@ -80,7 +81,6 @@ private:
 	void OnTabDrawerFocusLost(TSharedRef<STabDrawer> Drawer);
 	void OnTabDrawerClosed(TSharedRef<STabDrawer> Drawer);
 	void OnTargetDrawerSizeChanged(TSharedRef<STabDrawer> Drawer, float NewSize);
-	void OnWindowDPIScaleChanged(TSharedRef<SWindow> WindowThatChanged);
 	/** Called when active tab changes which is used to determine if we should close an opened content browser drawer */
 	void OnActiveTabChanged(TSharedPtr<SDockTab> NewlyActivated, TSharedPtr<SDockTab> PreviouslyActive);
 
@@ -115,11 +115,11 @@ private:
 private:
 	TSharedPtr<SVerticalBox> TabBox;
 	TArray<TPair<TSharedRef<SDockTab>, TSharedRef<STabDrawerButton>>> Tabs;
-	TWeakPtr<SWindow> WindowWithOverlayContent;
 	TSharedPtr<FActiveTimerHandle> OpenPendingDrawerTimerHandle;
 	ESidebarLocation Location;
+	TSharedPtr<SOverlay> DrawersOverlay;
 	/** Generally speaking one drawer is only ever open at once but we animate any previous drawer closing so there could be more than one while an animation is playing */
 	TArray<TSharedRef<STabDrawer>> OpenedDrawers;
 	/** Any pending drawer tab to open */
 	TWeakPtr<SDockTab> PendingTabToOpen;
-}; 
\ No newline at end of file
+}; 
