Fixed menu on Mac displaying all items that were pulled in menu to support recursive search.

UE-122188 - UE5 Unreal Editor menus regression on macOS

diff --git a/Engine/Source/Runtime/Slate/Private/Framework/MultiBox/MultiBox.cpp b/Engine/Source/Runtime/Slate/Private/Framework/MultiBox/MultiBox.cpp
--- a/Engine/Source/Runtime/Slate/Private/Framework/MultiBox/MultiBox.cpp
+++ b/Engine/Source/Runtime/Slate/Private/Framework/MultiBox/MultiBox.cpp
@@ -1074,132 +1074,127 @@ void SMultiBoxWidget::BuildMultiBoxWidget()
 	bool bSectionContainsIcons = false;
 	int32 NextMenuSeparator = INDEX_NONE;
 
-	// The first pass adds this multibox blocks and discover/collect all sub-menus blocks. The second pass adds all sub-menu blocks discovered in the first pass.
-	for (int32 Pass = 0; Pass < 2; ++Pass)
+	for( int32 Index = 0; Index < Blocks.Num(); Index++ )
 	{
-		int32 StartBlockIndex = 0;
-
-		// If this is the second pass, add the blocks that were discovered while expanding all sub-menus in the first pass, creating a flat list (but hidden until a search occurs).
-		if (Pass == 1)
+		// If we've passed the last menu separator, scan for the next one (the end of the list is also considered a menu separator for the purposes of this index)
+		if (NextMenuSeparator < Index)
 		{
-			StartBlockIndex = Blocks.Num();
-			for (const TPair<TSharedPtr<const FMultiBlock>, TSharedPtr<FFlattenSearchableBlockInfo>>& Pair : FlattenSearchableBlocks)
+			bSectionContainsIcons = false;
+			for (++NextMenuSeparator; NextMenuSeparator < Blocks.Num(); ++NextMenuSeparator)
 			{
-				MultiBox->AddMultiBlock(Pair.Key.ToSharedRef());
-			}
-		}
-
-		for( int32 Index = StartBlockIndex; Index < Blocks.Num(); Index++ )
-		{
-			// If we've passed the last menu separator, scan for the next one (the end of the list is also considered a menu separator for the purposes of this index)
-			if (NextMenuSeparator < Index)
-			{
-				bSectionContainsIcons = false;
-				for (++NextMenuSeparator; NextMenuSeparator < Blocks.Num(); ++NextMenuSeparator)
+				const FMultiBlock& TestBlock = *Blocks[NextMenuSeparator];
+				if (!bSectionContainsIcons && TestBlock.HasIcon())
 				{
-					const FMultiBlock& TestBlock = *Blocks[NextMenuSeparator];
-					if (!bSectionContainsIcons && TestBlock.HasIcon())
-					{
-						bSectionContainsIcons = true;
-					}
+					bSectionContainsIcons = true;
+				}
 
-					if (TestBlock.GetType() == EMultiBlockType::Separator)
-					{
-						break;
-					}
+				if (TestBlock.GetType() == EMultiBlockType::Separator)
+				{
+					break;
 				}
 			}
+		}
 
-			const FMultiBlock& Block = *Blocks[Index];
-			EMultiBlockLocation::Type Location = EMultiBlockLocation::None;
+		const FMultiBlock& Block = *Blocks[Index];
+		EMultiBlockLocation::Type Location = EMultiBlockLocation::None;
 		
-			TSharedPtr<const FMultiBlock> NextBlock = nullptr;
+		TSharedPtr<const FMultiBlock> NextBlock = nullptr;
+
+		if (Blocks.IsValidIndex(Index + 1))
+		{
+			NextBlock = Blocks[Index + 1];
+		}
 
-			if (Blocks.IsValidIndex(Index + 1))
+		// Determine the location of the current block, used for group styling information
+		{
+			// Check if we are a start or end block
+			if (Block.IsGroupStartBlock())
 			{
-				NextBlock = Blocks[Index + 1];
+				bInsideGroup = true;
+			}
+			else if (Block.IsGroupEndBlock())
+			{
+				bInsideGroup = false;
 			}
 
-			// Determine the location of the current block, used for group styling information
+			// Check if we are next to a start or end block
+			bool bIsNextToStartBlock = false;
+			bool bIsNextToEndBlock = false;
+			if (NextBlock)
 			{
-				// Check if we are a start or end block
-				if (Block.IsGroupStartBlock())
+				if ( NextBlock->IsGroupEndBlock() )
 				{
-					bInsideGroup = true;
+					bIsNextToEndBlock = true;
 				}
-				else if (Block.IsGroupEndBlock())
+			}
+			if (Index > 0)
+			{
+				const FMultiBlock& PrevBlock = *Blocks[Index - 1];
+				if ( PrevBlock.IsGroupStartBlock() )
 				{
-					bInsideGroup = false;
+					bIsNextToStartBlock = true;
 				}
+			}
+
+			// determine location
+			if (bInsideGroup)
+			{
+				// assume we are in the middle of a group
+				Location = EMultiBlockLocation::Middle;
 
-				// Check if we are next to a start or end block
-				bool bIsNextToStartBlock = false;
-				bool bIsNextToEndBlock = false;
-				if (NextBlock)
+				// We are the start of a group
+				if (bIsNextToStartBlock && !bIsNextToEndBlock)
 				{
-					if ( NextBlock->IsGroupEndBlock() )
-					{
-						bIsNextToEndBlock = true;
-					}
+					Location = EMultiBlockLocation::Start;
 				}
-				if (Index > 0)
+				// we are the end of a group
+				else if (!bIsNextToStartBlock && bIsNextToEndBlock)
 				{
-					const FMultiBlock& PrevBlock = *Blocks[Index - 1];
-					if ( PrevBlock.IsGroupStartBlock() )
-					{
-						bIsNextToStartBlock = true;
-					}
+					Location = EMultiBlockLocation::End;
 				}
-
-				// determine location
-				if (bInsideGroup)
+				// we are the only block in a group
+				else if (bIsNextToStartBlock && bIsNextToEndBlock)
 				{
-					// assume we are in the middle of a group
-					Location = EMultiBlockLocation::Middle;
-
-					// We are the start of a group
-					if (bIsNextToStartBlock && !bIsNextToEndBlock)
-					{
-						Location = EMultiBlockLocation::Start;
-					}
-					// we are the end of a group
-					else if (!bIsNextToStartBlock && bIsNextToEndBlock)
-					{
-						Location = EMultiBlockLocation::End;
-					}
-					// we are the only block in a group
-					else if (bIsNextToStartBlock && bIsNextToEndBlock)
-					{
-						Location = EMultiBlockLocation::None;
-					}
+					Location = EMultiBlockLocation::None;
 				}
 			}
+		}
 
-			TSharedPtr<const FToolBarComboButtonBlock> OptionsBlock;
-			if (NextBlock && NextBlock->GetType() == EMultiBlockType::ToolBarComboButton)
+		TSharedPtr<const FToolBarComboButtonBlock> OptionsBlock;
+		if (NextBlock && NextBlock->GetType() == EMultiBlockType::ToolBarComboButton)
+		{
+			TSharedPtr<const FToolBarComboButtonBlock> NextToolBarComboButtonBlock = StaticCastSharedPtr<const FToolBarComboButtonBlock>(NextBlock);
+			if (NextToolBarComboButtonBlock->IsSimpleComboBox())
 			{
-				TSharedPtr<const FToolBarComboButtonBlock> NextToolBarComboButtonBlock = StaticCastSharedPtr<const FToolBarComboButtonBlock>(NextBlock);
-				if (NextToolBarComboButtonBlock->IsSimpleComboBox())
-				{
-					// Apply a special treatment to simple combo boxes as they represent options for the previous button
-					OptionsBlock = NextToolBarComboButtonBlock;
-					// Skip over options blocks. They are not added directly
-					++Index;
-				}
+				// Apply a special treatment to simple combo boxes as they represent options for the previous button
+				OptionsBlock = NextToolBarComboButtonBlock;
+				// Skip over options blocks. They are not added directly
+				++Index;
 			}
+		}
 
-			if( DragPreview.IsValid() && DragPreview.InsertIndex == Index )
-			{
-				// Add the drag preview before if we have it. This block shows where the custom block will be 
-				// added if the user drops it
-				AddBlockWidget(*DragPreview.PreviewBlock, HorizontalBox, VerticalBox, EMultiBlockLocation::None, bSectionContainsIcons, OptionsBlock);
-			}
+		if( DragPreview.IsValid() && DragPreview.InsertIndex == Index )
+		{
+			// Add the drag preview before if we have it. This block shows where the custom block will be 
+			// added if the user drops it
+			AddBlockWidget(*DragPreview.PreviewBlock, HorizontalBox, VerticalBox, EMultiBlockLocation::None, bSectionContainsIcons, OptionsBlock);
+		}
 		
-			// Do not add a block if it is being dragged
-			if( !IsBlockBeingDragged( Blocks[Index] ) )
-			{
-				AddBlockWidget(Block, HorizontalBox, VerticalBox, Location, bSectionContainsIcons, OptionsBlock);
-			}
+		// Do not add a block if it is being dragged
+		if( !IsBlockBeingDragged( Blocks[Index] ) )
+		{
+			AddBlockWidget(Block, HorizontalBox, VerticalBox, Location, bSectionContainsIcons, OptionsBlock);
+		}
+	}
+
+	// The first loop above added the multibox blocks and discovered/collected all sub-menus blocks. This second loops adds all sub-menu blocks discovered in the first pass.
+	for (const TPair<TSharedPtr<const FMultiBlock>, TSharedPtr<FFlattenSearchableBlockInfo>>& Pair : FlattenSearchableBlocks)
+	{
+		// Do not add a block if it is being dragged
+		if (!IsBlockBeingDragged(Pair.Key))
+		{
+			TSharedPtr<const FToolBarComboButtonBlock> OptionsBlock;
+			AddBlockWidget(*Pair.Key, HorizontalBox, VerticalBox, EMultiBlockLocation::None, /*bSectionContainsIcons*/false, OptionsBlock);
 		}
 	}
 
