- Added 'recursive searchable' property to the FMenuEntryBlock to let the menu searching algorith if it should go down the sub-menus.
  - Fixed property binding auto-generated menus that can expand indefinitedly (in a circular fashion) to search only the current level and prevent stack overflow.

diff --git a/Engine/Plugins/Runtime/PropertyAccess/Source/PropertyAccessEditor/Private/SPropertyBinding.cpp b/Engine/Plugins/Runtime/PropertyAccess/Source/PropertyAccessEditor/Private/SPropertyBinding.cpp
--- a/Engine/Plugins/Runtime/PropertyAccess/Source/PropertyAccessEditor/Private/SPropertyBinding.cpp
+++ b/Engine/Plugins/Runtime/PropertyAccess/Source/PropertyAccessEditor/Private/SPropertyBinding.cpp
@@ -296,8 +296,16 @@ bool SPropertyBinding::HasBindablePropertiesRecursive(UStruct* InStruct, TSet<US
 
 TSharedRef<SWidget> SPropertyBinding::OnGenerateDelegateMenu()
 {
+	// The menu itself is be searchable.
+	const bool bSearchableMenu = true;
+
+	// The menu are generated through reflection and sometime the API exposes some recursivity (think about a Widget returning it parent which is also a Widget). Just by reflection
+	// it is not possible to determine when the root object is reached. It needs a kind of simulation which is not implemented. Also, even if the recursivity was correctly handled, the possible
+	// permutations tend to grow exponentially. Until a clever solution is found, the simple approach is to disable recursively searching those menus. User can still search the current one though.
+	const bool bRecursivelySearchableMenu = false;
+
 	const bool bInShouldCloseWindowAfterMenuSelection = true;
-	FMenuBuilder MenuBuilder(bInShouldCloseWindowAfterMenuSelection, nullptr, Args.MenuExtender);
+	FMenuBuilder MenuBuilder(bInShouldCloseWindowAfterMenuSelection, nullptr, Args.MenuExtender, /*bCloseSelfOnly*/false, &FCoreStyle::Get(), bSearchableMenu, NAME_None, bRecursivelySearchableMenu);
 
 	MenuBuilder.BeginSection("BindingActions", LOCTEXT("Bindings", "Bindings"));
 
diff --git a/Engine/Source/Runtime/Slate/Private/Framework/MultiBox/MultiBox.cpp b/Engine/Source/Runtime/Slate/Private/Framework/MultiBox/MultiBox.cpp
--- a/Engine/Source/Runtime/Slate/Private/Framework/MultiBox/MultiBox.cpp
+++ b/Engine/Source/Runtime/Slate/Private/Framework/MultiBox/MultiBox.cpp
@@ -723,8 +723,8 @@ void SMultiBoxWidget::AddBlockWidget(const FMultiBlock& Block, TSharedPtr<SHoriz
 
 		BlockWidget = FlattenWidget;
 	}
-	// If the added block is searchable and expands to a sub-menu. (This is the first pass building the multibox widget)
-	else if (GetSearchable() && Block.GetSearchable() && Block.GetType() == EMultiBlockType::MenuEntry && static_cast<const FMenuEntryBlock&>(Block).IsSubMenu())
+	// If the added block is searchable and the search is allowed to walk down sub-menus recursively. (This is the first pass building the multibox widget)
+	else if (GetSearchable() && Block.GetSearchable() && Block.GetType() == EMultiBlockType::MenuEntry && static_cast<const FMenuEntryBlock&>(Block).IsSubMenu() && static_cast<const FMenuEntryBlock&>(Block).IsRecursivelySearchable())
 	{
 		// If the block widget being added displays a searchable text.
 		const TArray<FText>* SearchableTextHierarchy = MultiBoxWidgets.Find(BlockWidget);
@@ -914,6 +914,7 @@ void SMultiBoxWidget::SetSearchable(bool InSearchable)
 {
 	bSearchable = InSearchable;
 }
+
 bool SMultiBoxWidget::GetSearchable() const
 {
 	return bSearchable;
diff --git a/Engine/Source/Runtime/Slate/Private/Framework/MultiBox/MultiBoxBuilder.cpp b/Engine/Source/Runtime/Slate/Private/Framework/MultiBox/MultiBoxBuilder.cpp
--- a/Engine/Source/Runtime/Slate/Private/Framework/MultiBox/MultiBoxBuilder.cpp
+++ b/Engine/Source/Runtime/Slate/Private/Framework/MultiBox/MultiBoxBuilder.cpp
@@ -241,7 +241,8 @@ void FMenuBuilder::AddSubMenu( const TAttribute<FText>& InMenuLabel, const TAttr
 	ApplySectionBeginning();
 
 	const bool bIsSubMenu = true;
-	TSharedRef< FMenuEntryBlock > NewMenuEntryBlock( new FMenuEntryBlock( InExtensionHook, InMenuLabel, InToolTip, InSubMenu, ExtenderStack.Top(), bIsSubMenu, bInOpenSubMenuOnClick, InIcon, InUIAction, InUserInterfaceActionType, bCloseSelfOnly, bInShouldCloseWindowAfterMenuSelection, CommandListStack.Last() ) );
+	TSharedRef< FMenuEntryBlock > NewMenuEntryBlock = MakeShared<FMenuEntryBlock>( InExtensionHook, InMenuLabel, InToolTip, InSubMenu, ExtenderStack.Top(), bIsSubMenu, bInOpenSubMenuOnClick, InIcon, InUIAction, InUserInterfaceActionType, bCloseSelfOnly, bInShouldCloseWindowAfterMenuSelection, CommandListStack.Last() );
+	NewMenuEntryBlock->SetRecursivelySearchable(bRecursivelySearchable);
 
 	MultiBox->AddMultiBlock( NewMenuEntryBlock );
 }
@@ -251,7 +252,8 @@ void FMenuBuilder::AddSubMenu( const TAttribute<FText>& InMenuLabel, const TAttr
 	ApplySectionBeginning();
 
 	const bool bIsSubMenu = true;
-	TSharedRef< FMenuEntryBlock > NewMenuEntryBlock( new FMenuEntryBlock( InExtensionHook, InMenuLabel, InToolTip, InSubMenu, ExtenderStack.Top(), bIsSubMenu, bInOpenSubMenuOnClick, CommandListStack.Last(), bCloseSelfOnly, InIcon, bInShouldCloseWindowAfterMenuSelection ) );
+	TSharedRef< FMenuEntryBlock > NewMenuEntryBlock = MakeShared<FMenuEntryBlock>( InExtensionHook, InMenuLabel, InToolTip, InSubMenu, ExtenderStack.Top(), bIsSubMenu, bInOpenSubMenuOnClick, CommandListStack.Last(), bCloseSelfOnly, InIcon, bInShouldCloseWindowAfterMenuSelection );
+	NewMenuEntryBlock->SetRecursivelySearchable(bRecursivelySearchable);
 
 	MultiBox->AddMultiBlock( NewMenuEntryBlock );
 }
@@ -261,7 +263,8 @@ void FMenuBuilder::AddSubMenu( const TSharedRef< SWidget > Contents, const FNewM
 	ApplySectionBeginning();
 
 	const bool bIsSubMenu = true;
-	TSharedRef< FMenuEntryBlock > NewMenuEntryBlock( new FMenuEntryBlock( NAME_None, Contents, InSubMenu, ExtenderStack.Top(), bIsSubMenu, bInOpenSubMenuOnClick, CommandListStack.Last(), bCloseSelfOnly, bInShouldCloseWindowAfterMenuSelection ) );
+	TSharedRef< FMenuEntryBlock > NewMenuEntryBlock = MakeShared<FMenuEntryBlock>( NAME_None, Contents, InSubMenu, ExtenderStack.Top(), bIsSubMenu, bInOpenSubMenuOnClick, CommandListStack.Last(), bCloseSelfOnly, bInShouldCloseWindowAfterMenuSelection );
+	NewMenuEntryBlock->SetRecursivelySearchable(bRecursivelySearchable);
 
 	MultiBox->AddMultiBlock( NewMenuEntryBlock );
 }
@@ -271,7 +274,8 @@ void FMenuBuilder::AddSubMenu( const FUIAction& UIAction, const TSharedRef< SWid
 	ApplySectionBeginning();
 
 	const bool bIsSubMenu = true;
-	TSharedRef< FMenuEntryBlock > NewMenuEntryBlock( new FMenuEntryBlock( NAME_None, UIAction, Contents, InSubMenu, ExtenderStack.Top(), bIsSubMenu, CommandListStack.Last(), bCloseSelfOnly, bInShouldCloseWindowAfterMenuSelection ) );
+	TSharedRef< FMenuEntryBlock > NewMenuEntryBlock = MakeShared<FMenuEntryBlock>( NAME_None, UIAction, Contents, InSubMenu, ExtenderStack.Top(), bIsSubMenu, CommandListStack.Last(), bCloseSelfOnly, bInShouldCloseWindowAfterMenuSelection );
+	NewMenuEntryBlock->SetRecursivelySearchable(bRecursivelySearchable);
 
 	MultiBox->AddMultiBlock( NewMenuEntryBlock );
 }
@@ -281,7 +285,8 @@ void FMenuBuilder::AddWrapperSubMenu( const FText& InMenuLabel, const FText& InT
 	ApplySectionBeginning();
 
 	const bool bIsSubMenu = true;
-	TSharedRef< FMenuEntryBlock > NewMenuEntryBlock( new FMenuEntryBlock( NAME_None, InMenuLabel, InToolTip, InSubMenu, ExtenderStack.Top(), bIsSubMenu, false, CommandListStack.Last(), bCloseSelfOnly, InIcon ) );
+	TSharedRef< FMenuEntryBlock > NewMenuEntryBlock = MakeShared<FMenuEntryBlock>( NAME_None, InMenuLabel, InToolTip, InSubMenu, ExtenderStack.Top(), bIsSubMenu, false, CommandListStack.Last(), bCloseSelfOnly, InIcon );
+	NewMenuEntryBlock->SetRecursivelySearchable(bRecursivelySearchable);
 
 	MultiBox->AddMultiBlock( NewMenuEntryBlock );
 }
@@ -291,7 +296,8 @@ void FMenuBuilder::AddWrapperSubMenu( const FText& InMenuLabel, const FText& InT
 	ApplySectionBeginning();
 
 	const bool bIsSubMenu = true;
-	TSharedRef< FMenuEntryBlock > NewMenuEntryBlock( new FMenuEntryBlock( NAME_None, UIAction, InMenuLabel, InToolTip, InSubMenu, ExtenderStack.Top(), bIsSubMenu, false, bCloseSelfOnly, InIcon ) );
+	TSharedRef< FMenuEntryBlock > NewMenuEntryBlock = MakeShared<FMenuEntryBlock>( NAME_None, UIAction, InMenuLabel, InToolTip, InSubMenu, ExtenderStack.Top(), bIsSubMenu, false, bCloseSelfOnly, InIcon );
+	NewMenuEntryBlock->SetRecursivelySearchable(bRecursivelySearchable);
 
 	MultiBox->AddMultiBlock( NewMenuEntryBlock );
 }
@@ -301,7 +307,8 @@ void FMenuBuilder::AddWrapperSubMenu( const FText& InMenuLabel, const FText& InT
 	ApplySectionBeginning();
 
 	const bool bIsSubMenu = true;
-	TSharedRef< FMenuEntryBlock > NewMenuEntryBlock( new FMenuEntryBlock( NAME_None, InMenuLabel, InToolTip, InSubMenu, ExtenderStack.Top(), bIsSubMenu, false, CommandListStack.Last(), bCloseSelfOnly, InIcon ) );
+	TSharedRef< FMenuEntryBlock > NewMenuEntryBlock = MakeShared<FMenuEntryBlock>( NAME_None, InMenuLabel, InToolTip, InSubMenu, ExtenderStack.Top(), bIsSubMenu, false, CommandListStack.Last(), bCloseSelfOnly, InIcon );
+	NewMenuEntryBlock->SetRecursivelySearchable(bRecursivelySearchable);
 
 	MultiBox->AddMultiBlock( NewMenuEntryBlock );
 }
diff --git a/Engine/Source/Runtime/Slate/Private/Framework/MultiBox/SMenuEntryBlock.cpp b/Engine/Source/Runtime/Slate/Private/Framework/MultiBox/SMenuEntryBlock.cpp
--- a/Engine/Source/Runtime/Slate/Private/Framework/MultiBox/SMenuEntryBlock.cpp
+++ b/Engine/Source/Runtime/Slate/Private/Framework/MultiBox/SMenuEntryBlock.cpp
@@ -21,6 +21,7 @@ FMenuEntryBlock::FMenuEntryBlock( const FName& InExtensionHook, const TSharedPtr
 	, ToolTipOverride( InToolTipOverride )
 	, IconOverride( InIconOverride )
 	, bIsSubMenu( false )
+	, bIsRecursivelySearchable( true )
 	, bOpenSubMenuOnClick( false )
 	, UserInterfaceActionType( EUserInterfaceActionType::Button )
 	, bCloseSelfOnly( bInCloseSelfOnly )
@@ -36,6 +37,7 @@ FMenuEntryBlock::FMenuEntryBlock( const FName& InExtensionHook, const TAttribute
 	, IconOverride( InIcon )
 	, EntryBuilder( InEntryBuilder )
 	, bIsSubMenu( bInSubMenu )
+	, bIsRecursivelySearchable( true )
 	, bOpenSubMenuOnClick( bInSubMenuOnClick )
 	, UserInterfaceActionType( InUserInterfaceActionType )
 	, bCloseSelfOnly( bInCloseSelfOnly )
@@ -51,6 +53,7 @@ FMenuEntryBlock::FMenuEntryBlock( const FName& InExtensionHook, const TAttribute
 	, ToolTipOverride( InToolTip )
 	, IconOverride( InIcon )
 	, bIsSubMenu( false )
+	, bIsRecursivelySearchable( true )
 	, bOpenSubMenuOnClick( false )
 	, UserInterfaceActionType( InUserInterfaceActionType )
 	, bCloseSelfOnly( bInCloseSelfOnly )
@@ -66,6 +69,7 @@ FMenuEntryBlock::FMenuEntryBlock( const FName& InExtensionHook, const TAttribute
 	, IconOverride( InIcon )
 	, EntryBuilder( InEntryBuilder )
 	, bIsSubMenu( bInSubMenu )
+	, bIsRecursivelySearchable( true )
 	, bOpenSubMenuOnClick( bInSubMenuOnClick )
 	, UserInterfaceActionType( EUserInterfaceActionType::Button )
 	, bCloseSelfOnly( bInCloseSelfOnly )
@@ -82,6 +86,7 @@ FMenuEntryBlock::FMenuEntryBlock( const FName& InExtensionHook, const TAttribute
 	, IconOverride( InIcon )
 	, MenuBuilder( InMenuBuilder )
 	, bIsSubMenu( bInSubMenu )
+	, bIsRecursivelySearchable( true )
 	, bOpenSubMenuOnClick( bInSubMenuOnClick )
 	, UserInterfaceActionType( EUserInterfaceActionType::Button )
 	, bCloseSelfOnly( bInCloseSelfOnly )
@@ -98,6 +103,7 @@ FMenuEntryBlock::FMenuEntryBlock( const FName& InExtensionHook, const TAttribute
 	, IconOverride( InIcon )
 	, EntryWidget( InEntryWidget )
 	, bIsSubMenu( bInSubMenu )
+	, bIsRecursivelySearchable( true )
 	, bOpenSubMenuOnClick( bInSubMenuOnClick )
 	, UserInterfaceActionType( EUserInterfaceActionType::Button )
 	, bCloseSelfOnly( bInCloseSelfOnly )
@@ -111,6 +117,7 @@ FMenuEntryBlock::FMenuEntryBlock( const FName& InExtensionHook, const FUIAction&
 	, ToolTipOverride( InToolTip )
 	, EntryWidget( Contents )
 	, bIsSubMenu( false )
+	, bIsRecursivelySearchable( true )
 	, bOpenSubMenuOnClick( false )
 	, UserInterfaceActionType( InUserInterfaceActionType )
 	, bCloseSelfOnly( bInCloseSelfOnly )
@@ -124,6 +131,7 @@ FMenuEntryBlock::FMenuEntryBlock( const FName& InExtensionHook, const TSharedRef
 	, EntryBuilder( InEntryBuilder )
 	, EntryWidget( Contents )
 	, bIsSubMenu( bInSubMenu )
+	, bIsRecursivelySearchable( true )
 	, bOpenSubMenuOnClick( bInSubMenuOnClick )
 	, UserInterfaceActionType( EUserInterfaceActionType::Button )
 	, bCloseSelfOnly( bInCloseSelfOnly )
@@ -138,6 +146,7 @@ FMenuEntryBlock::FMenuEntryBlock( const FName& InExtensionHook, const FUIAction&
 	, EntryBuilder( InEntryBuilder )
 	, EntryWidget( Contents )
 	, bIsSubMenu( bInSubMenu )
+	, bIsRecursivelySearchable( true )
 	, bOpenSubMenuOnClick( false )
 	, UserInterfaceActionType( EUserInterfaceActionType::Button )
 	, bCloseSelfOnly( bInCloseSelfOnly )
@@ -153,6 +162,7 @@ FMenuEntryBlock::FMenuEntryBlock( const FName& InExtensionHook, const FUIAction&
 	, IconOverride( InIcon )
 	, MenuBuilder( InMenuBuilder )
 	, bIsSubMenu( bInSubMenu )
+	, bIsRecursivelySearchable( true )
 	, bOpenSubMenuOnClick( bInSubMenuOnClick )
 	, UserInterfaceActionType( EUserInterfaceActionType::Button )
 	, bCloseSelfOnly( bInCloseSelfOnly )
@@ -1493,10 +1503,16 @@ TSharedRef< SWidget > SMenuEntryBlock::MakeNewMenuWidget() const
 			}
 		}
 
-		FMenuBuilder MenuBuilder(MenuEntryBlock->bShouldCloseWindowAfterMenuSelection, MultiBlock->GetActionList(), MenuEntryBlock->Extender, bCloseSelfOnly, StyleSet, /* searchable */ true, SubMenuCustomizationName );
-		{
-			MenuEntryBlock->EntryBuilder.Execute( MenuBuilder );
-		}
+		FMenuBuilder MenuBuilder(MenuEntryBlock->bShouldCloseWindowAfterMenuSelection,
+		                         MultiBlock->GetActionList(),
+		                         MenuEntryBlock->Extender,
+		                         bCloseSelfOnly,
+		                         StyleSet,
+		                         /*Searchable*/true, // Best guess default. This could be changed by the EntryBuilder delegate where more context is available.
+		                         SubMenuCustomizationName,
+		                         MenuEntryBlock->IsRecursivelySearchable()); // Best guess default. This could be changed by the EntryBuilder delegate where more context is available.
+
+		MenuEntryBlock->EntryBuilder.Execute( MenuBuilder );
 
 		return MenuBuilder.MakeWidget();
 	}
diff --git a/Engine/Source/Runtime/Slate/Private/Framework/MultiBox/SMenuEntryBlock.h b/Engine/Source/Runtime/Slate/Private/Framework/MultiBox/SMenuEntryBlock.h
--- a/Engine/Source/Runtime/Slate/Private/Framework/MultiBox/SMenuEntryBlock.h
+++ b/Engine/Source/Runtime/Slate/Private/Framework/MultiBox/SMenuEntryBlock.h
@@ -118,14 +118,20 @@ public:
 	virtual void CreateMenuEntry(class FMenuBuilder& MenuBuilder) const override;
 	virtual bool HasIcon() const override;
 
+	/** Returns whether this menu entry block opens a sub-menu. */
 	bool IsSubMenu() const { return bIsSubMenu; }
 
+	/** Sets whether the menu search algorithm should walk down this menu sub-menus. */
+	void SetRecursivelySearchable(bool bInRecursivelySearchable) { bIsRecursivelySearchable = bInRecursivelySearchable; }
+
+	/** Returns whether the menu search algorithm should walk down this menu sub-menus. */
+	bool IsRecursivelySearchable() const { return bIsRecursivelySearchable; }
+
 private:
 
 	/** FMultiBlock private interface */
 	virtual TSharedRef< class IMultiBlockBaseWidget > ConstructWidget() const override;
 
-
 private:
 
 	// Friend our corresponding widget class
@@ -153,6 +159,9 @@ private:
 	/** True if this menu entry opens a sub-menu */
 	bool bIsSubMenu;
 
+	/** True if the search algorithm should walk down this menu sub menus. Usually true, unless the menu has circular/infinite expansion (happens in some menus generated on the fly by reflection). */
+	bool bIsRecursivelySearchable;
+
 	/** True if this menu entry opens a sub-menu by clicking on it only */
 	bool bOpenSubMenuOnClick;
 
diff --git a/Engine/Source/Runtime/Slate/Public/Framework/MultiBox/MultiBoxBuilder.h b/Engine/Source/Runtime/Slate/Public/Framework/MultiBox/MultiBoxBuilder.h
--- a/Engine/Source/Runtime/Slate/Public/Framework/MultiBox/MultiBoxBuilder.h
+++ b/Engine/Source/Runtime/Slate/Public/Framework/MultiBox/MultiBoxBuilder.h
@@ -221,11 +221,17 @@ public:
 	 * @param	bInShouldCloseWindowAfterMenuSelection	Sets whether or not the window that contains this multibox should be destroyed after the user clicks on a menu item in this box
 	 * @param	InCommandList	The action list that maps command infos to delegates that should be called for each command associated with a multiblock widget
 	 * @param	bInCloseSelfOnly	True if clicking on a menu entry closes itself only and its children but not the entire stack
+	 * @param	bInSearchable	True if the menu is searchable
+	 * @param	bInRecursivelySearchable True if search algorithm should go down the sub-menus when searching. If false, the search will not scan the sub-menus. Recursive search is usually disabled on
+	 *                                   menus that are automatically generated in a way that the menu can expand indefinitely in a circular fashion. For example, the Blueprint API binding reflected on
+	 *                                   function return type can do that. Think about a function like "A* A::GetParent()". If a root menu expands "A" functions and expands on the function return types, then
+	 *                                   selecting "GetParent()" will expand another "A" menu. Without simulation, the reflection API don't know when the recursion finishes, nor does the recursive search algorithm.
 	 */
-	FMenuBuilder( const bool bInShouldCloseWindowAfterMenuSelection, TSharedPtr< const FUICommandList > InCommandList, TSharedPtr<FExtender> InExtender = TSharedPtr<FExtender>(), const bool bCloseSelfOnly = false, const ISlateStyle* InStyleSet = &FCoreStyle::Get(), bool bInSearchable = true, FName InMenuName = NAME_None)
+	FMenuBuilder( const bool bInShouldCloseWindowAfterMenuSelection, TSharedPtr< const FUICommandList > InCommandList, TSharedPtr<FExtender> InExtender = TSharedPtr<FExtender>(), const bool bCloseSelfOnly = false, const ISlateStyle* InStyleSet = &FCoreStyle::Get(), bool bInSearchable = true, FName InMenuName = NAME_None, bool bInRecursivelySearchable = true)
 		: FBaseMenuBuilder( EMultiBoxType::Menu, bInShouldCloseWindowAfterMenuSelection, InCommandList, bCloseSelfOnly, InExtender, InStyleSet, NAME_None, InMenuName )
 		, bSectionNeedsToBeApplied(false)
 		, bSearchable(bInSearchable)
+		, bRecursivelySearchable(bInRecursivelySearchable)
 		, bIsEditing(false)
 	{
 		if(bSearchable)
@@ -342,6 +348,9 @@ private:
 	/** Whether this menu is searchable */
 	bool bSearchable;
 
+	/** Whether the search algorithm should walk down this menu sub-menu(s) (if the menu is searchable in first place). */
+	bool bRecursivelySearchable;
+
 	/** Whether menu is currently being edited */
 	bool bIsEditing;
 };
