Fix Scrollbox Padding doing nothing.

diff --git a/Engine/Source/Runtime/Slate/Private/Widgets/Layout/SScrollBox.cpp b/Engine/Source/Runtime/Slate/Private/Widgets/Layout/SScrollBox.cpp
--- a/Engine/Source/Runtime/Slate/Private/Widgets/Layout/SScrollBox.cpp
+++ b/Engine/Source/Runtime/Slate/Private/Widgets/Layout/SScrollBox.cpp
@@ -172,9 +172,9 @@ void SScrollBox::Construct( const FArguments& InArgs )
 		ScrollBar = ConstructScrollBar();
 		ScrollBar->SetDragFocusCause(InArgs._ScrollBarDragFocusCause);
 		ScrollBar->SetThickness(InArgs._ScrollBarThickness);
-		ScrollBar->SetPadding(InArgs._ScrollBarPadding);
 		ScrollBar->SetUserVisibility(InArgs._ScrollBarVisibility);
 		ScrollBar->SetScrollBarAlwaysVisible(InArgs._ScrollBarAlwaysVisible);
+		ScrollBarPadding = InArgs._ScrollBarPadding;
 
 		bScrollBarIsExternal = false;
 	}
@@ -256,6 +256,7 @@ void SScrollBox::ConstructVerticalLayout()
 	if (!bScrollBarIsExternal)
 	{
 		PanelAndScrollbar->AddSlot()
+		.Padding(ScrollBarPadding)
 		.AutoWidth()
 		[
 			ScrollBar.ToSharedRef()
@@ -532,7 +533,7 @@ void SScrollBox::SetScrollBarThickness(FVector2D InThickness)
 
 void SScrollBox::SetScrollBarPadding(const FMargin& InPadding)
 {
-	ScrollBar->SetPadding(InPadding);
+	ScrollBarPadding = InPadding;
 }
 
 void SScrollBox::SetScrollBarRightClickDragAllowed(bool bIsAllowed)
diff --git a/Engine/Source/Runtime/Slate/Public/Widgets/Layout/SScrollBox.h b/Engine/Source/Runtime/Slate/Public/Widgets/Layout/SScrollBox.h
--- a/Engine/Source/Runtime/Slate/Public/Widgets/Layout/SScrollBox.h
+++ b/Engine/Source/Runtime/Slate/Public/Widgets/Layout/SScrollBox.h
@@ -341,6 +341,9 @@ protected:
 	/** Whether to permit overscroll on this scroll box */
 	EAllowOverscroll AllowOverscroll;
 
+	/** Padding to the scrollbox */
+	FMargin ScrollBarPadding;
+
 	/** Whether to back pad this scroll box, allowing user to scroll backward until child contents are no longer visible */
 	bool BackPadScrolling;
 
