Sidebar tab pinning - pinned tabs' drawers stay open instead of automatically dismissing when they lose focus

STabManager and SDockingTabStack: store pinned state in FTab, serialize pinned state, helpers for accessing/setting pinned state

STabDrawer:
- allow bypassing animation when opening drawer (e.g. for restoring pinned drawers on startup) and provide helper to check whether drawer is in the middle of closing
- automatically activate tab when its drawer gains focus and deactivate tab when its drawer starts to close

STabSidebar:
- change blue stripe on sidebar tab button to indicate tab active state, which is more consistent with regular tabs
- open sidebar tabs on mouse press (not mouse click), which is more consistent with regular tabs
- add UI for toggling pin state to the sidebar tab button
- prevent pinned sidebar tab drawers from closing when they lose focus
- automatically summon back drawers for pinned tabs when all other drawers are closed

diff --git a/Engine/Source/Runtime/Slate/Private/Framework/Docking/SDockingTabStack.cpp b/Engine/Source/Runtime/Slate/Private/Framework/Docking/SDockingTabStack.cpp
--- a/Engine/Source/Runtime/Slate/Private/Framework/Docking/SDockingTabStack.cpp
+++ b/Engine/Source/Runtime/Slate/Private/Framework/Docking/SDockingTabStack.cpp
@@ -351,6 +351,26 @@ void SDockingTabStack::SetTabSidebarSizeCoefficient(const TSharedRef<SDockTab>&
 	}
 }
 
+bool SDockingTabStack::IsTabPinnedInSidebar(const TSharedRef<SDockTab>& InTab)
+{
+	FTabManager::FTab* Tab = Tabs.FindByPredicate(FTabMatcher(InTab->GetLayoutIdentifier()));
+	if (Tab)
+	{
+		return Tab->bPinnedInSidebar;
+	}
+
+	return false;
+}
+
+void SDockingTabStack::SetTabPinnedInSidebar(const TSharedRef<SDockTab>& InTab, bool bPinnedInSidebar)
+{
+	FTabManager::FTab* Tab = Tabs.FindByPredicate(FTabMatcher(InTab->GetLayoutIdentifier()));
+	if (Tab)
+	{
+		Tab->bPinnedInSidebar = bPinnedInSidebar;
+	}
+}
+
 const TSlotlessChildren<SDockTab>& SDockingTabStack::GetTabs() const
 {
 	return TabWell->GetTabs();
diff --git a/Engine/Source/Runtime/Slate/Private/Framework/Docking/SDockingTabStack.h b/Engine/Source/Runtime/Slate/Private/Framework/Docking/SDockingTabStack.h
--- a/Engine/Source/Runtime/Slate/Private/Framework/Docking/SDockingTabStack.h
+++ b/Engine/Source/Runtime/Slate/Private/Framework/Docking/SDockingTabStack.h
@@ -80,6 +80,9 @@ public:
 	float GetTabSidebarSizeCoefficient(const TSharedRef<SDockTab>& InTab);
 	void SetTabSidebarSizeCoefficient(const TSharedRef<SDockTab>& InTab, float InSizeCoefficient);
 
+	bool IsTabPinnedInSidebar(const TSharedRef<SDockTab>& InTab);
+	void SetTabPinnedInSidebar(const TSharedRef<SDockTab>& InTab, bool bPinnedInSidebar);
+
 	/** @return true if the specified tab can be moved to a sidebar */
 	bool CanMoveTabToSideBar(TSharedRef<SDockTab> Tab) const;
 
diff --git a/Engine/Source/Runtime/Slate/Private/Framework/Docking/STabDrawer.cpp b/Engine/Source/Runtime/Slate/Private/Framework/Docking/STabDrawer.cpp
--- a/Engine/Source/Runtime/Slate/Private/Framework/Docking/STabDrawer.cpp
+++ b/Engine/Source/Runtime/Slate/Private/Framework/Docking/STabDrawer.cpp
@@ -1,6 +1,7 @@
 // Copyright Epic Games, Inc. All Rights Reserved.
 
 #include "Framework/Docking/STabDrawer.h"
+#include "Framework/Docking/TabManager.h"
 #include "Styling/AppStyle.h"
 #include "Styling/SlateTypes.h"
 #include "Layout/ArrangedChildren.h"
@@ -75,8 +76,15 @@ void STabDrawer::Construct(const FArguments& InArgs, TSharedRef<SDockTab> InTab,
 	];
 }
 
-void STabDrawer::Open()
+void STabDrawer::Open(bool bAnimateOpen)
 {
+	if (!bAnimateOpen)
+	{
+		SetCurrentSize(TargetDrawerSize);
+		OpenCloseAnimation.JumpToEnd();
+		return;
+	}
+
 	OpenCloseAnimation.Play(AsShared(), false, OpenCloseAnimation.IsPlaying() ? OpenCloseAnimation.GetSequenceTime() : 0.0f, false);
 
 	if (!OpenCloseTimer.IsValid())
@@ -98,6 +106,14 @@ void STabDrawer::Close()
 		AnimationThrottle = FSlateThrottleManager::Get().EnterResponsiveMode();
 		OpenCloseTimer = RegisterActiveTimer(0.0f, FWidgetActiveTimerDelegate::CreateSP(this, &STabDrawer::UpdateAnimation));
 	}
+
+	// Make sure that this tab isn't active if it's still active when closing.
+	// This might happen if the tab drawer lost focus to a non-tab (e.g. Content Drawer), so no other tab has become active.
+	// Don't unconditionally clear the active tab, since this drawer might be closing due to another tab taking focus.
+	if (FGlobalTabmanager::Get()->GetActiveTab() == ForTab)
+	{
+		FGlobalTabmanager::Get()->SetActiveTab(nullptr);
+	}
 }
 
 bool STabDrawer::IsOpen() const
@@ -105,6 +121,11 @@ bool STabDrawer::IsOpen() const
 	return !OpenCloseAnimation.IsAtStart();
 }
 
+bool STabDrawer::IsClosing() const
+{
+	return OpenCloseAnimation.IsPlaying() && OpenCloseAnimation.IsInReverse();
+}
+
 const TSharedRef<SDockTab> STabDrawer::GetTab() const
 {
 	return ForTab.ToSharedRef();
@@ -498,13 +519,26 @@ void STabDrawer::OnGlobalFocusChanging(const FFocusEvent& FocusEvent, const FWea
 		TArray<TSharedRef<SWidget>, TInlineAllocator<4>> LegalFocusWidgets;
 		LegalFocusWidgets.Add(ThisWidget);
 		LegalFocusWidgets.Add(ChildSlot.GetWidget());
+		if (TSharedPtr<SWidget> TabButtonSP = TabButton.Pin())
+		{
+			LegalFocusWidgets.Add(TabButtonSP.ToSharedRef());
+		}
 
 		bool bShouldLoseFocus = false;
 		// Do not close due to slow tasks as those opening send window activation events
-		if (!GIsSlowTask && !FSlateApplication::Get().GetActiveModalWindow().IsValid() && !IsLegalWidgetFocused(NewFocusedWidgetPath, MakeArrayView(LegalFocusWidgets)))
+		if (!GIsSlowTask && !FSlateApplication::Get().GetActiveModalWindow().IsValid())
 		{
-			if(NewFocusedWidgetPath.IsValid())
+			if (IsLegalWidgetFocused(NewFocusedWidgetPath, MakeArrayView(LegalFocusWidgets)))
+			{
+				// New focus is on this tab, so make it active
+				if (!IsClosing())
+				{
+					FGlobalTabmanager::Get()->SetActiveTab(ForTab);
+				}
+			}
+			else if (NewFocusedWidgetPath.IsValid())
 			{
+				// New focus is on something else, try to check if it's a menu
 				TSharedRef<SWindow> NewWindow = NewFocusedWidgetPath.GetWindow();
 				TSharedPtr<SWindow> MyWindow = FSlateApplication::Get().FindWidgetWindow(ThisWidget);
 
diff --git a/Engine/Source/Runtime/Slate/Private/Framework/Docking/STabSidebar.cpp b/Engine/Source/Runtime/Slate/Private/Framework/Docking/STabSidebar.cpp
--- a/Engine/Source/Runtime/Slate/Private/Framework/Docking/STabSidebar.cpp
+++ b/Engine/Source/Runtime/Slate/Private/Framework/Docking/STabSidebar.cpp
@@ -2,6 +2,7 @@
 
 #include "STabSidebar.h"
 #include "Widgets/Input/SButton.h"
+#include "Widgets/Input/SCheckBox.h"
 #include "SDockingTabWell.h"
 #include "Widgets/Layout/SBox.h"
 #include "Widgets/Images/SImage.h"
@@ -15,7 +16,18 @@
 
 #define LOCTEXT_NAMESPACE "TabSidebar"
 
-DECLARE_DELEGATE_OneParam(FOnTabDrawerButtonClicked, TSharedRef<SDockTab>);
+DECLARE_DELEGATE_OneParam(FOnTabDrawerButtonPressed, TSharedRef<SDockTab>);
+DECLARE_DELEGATE_TwoParams(FOnTabDrawerPinButtonToggled, TSharedRef<SDockTab>, bool);
+
+static bool IsTabPinned(TSharedRef<SDockTab> Tab)
+{
+	return Tab->GetParentDockTabStack()->IsTabPinnedInSidebar(Tab);
+}
+
+static void SetTabPinned(TSharedRef<SDockTab> Tab, bool bIsPinned)
+{
+	return Tab->GetParentDockTabStack()->SetTabPinnedInSidebar(Tab, bIsPinned);
+}
 
 /**
  * Vertical text block for use in the tab drawer button.
@@ -124,7 +136,8 @@ class STabDrawerButton : public SCompoundWidget
 
 	SLATE_BEGIN_ARGS(STabDrawerButton)
 	{}
-		SLATE_EVENT(FOnTabDrawerButtonClicked, OnDrawerButtonClicked)
+		SLATE_EVENT(FOnTabDrawerButtonPressed, OnDrawerButtonPressed)
+		SLATE_EVENT(FOnTabDrawerPinButtonToggled, OnDrawerPinButtonToggled)
 		SLATE_EVENT(FOnGetContent, OnGetContextMenuContent)
 	SLATE_END_ARGS()
 
@@ -138,7 +151,8 @@ public:
 		// Sometimes tabs can be renamed so ensure that we pick up the rename
 		ForTab->SetOnTabRenamed(SDockTab::FOnTabRenamed::CreateSP(this, &STabDrawerButton::OnTabRenamed));
 
-		OnDrawerButtonClicked = InArgs._OnDrawerButtonClicked;
+		OnDrawerButtonPressed = InArgs._OnDrawerButtonPressed;
+		OnDrawerPinButtonToggled = InArgs._OnDrawerPinButtonToggled;
 		OnGetContextMenuContent = InArgs._OnGetContextMenuContent;
 		Tab = ForTab;
 		Location = InLocation;
@@ -161,8 +175,9 @@ public:
 					SAssignNew(MainButton, SButton)
 					.ToolTip(ForTab->GetToolTip() ? ForTab->GetToolTip() : TAttribute<TSharedPtr<IToolTip>>())
 					.ToolTipText(ForTab->GetToolTip() ? TAttribute<FText>() : ForTab->GetTabLabel())
-					.ContentPadding(FMargin(0.0f, DockTabStyle->TabPadding.Top, 0.0f, DockTabStyle->TabPadding.Right))
-					.OnClicked_Lambda([this](){OnDrawerButtonClicked.ExecuteIfBound(Tab.ToSharedRef()); return FReply::Handled(); })
+					.ContentPadding(FMargin(0.0f, DockTabStyle->TabPadding.Top, 0.0f, DockTabStyle->TabPadding.Bottom))
+					// activate tab on mouse down (not mouse down-up) for consistency with non-sidebar tabs
+					.OnPressed_Lambda([this](){OnDrawerButtonPressed.ExecuteIfBound(Tab.ToSharedRef()); })
 					.ForegroundColor(FSlateColor::UseForeground())
 					[
 						SNew(SVerticalBox)
@@ -180,7 +195,7 @@ public:
 							//.RenderTransformPivot(FVector2D(.5f, .5f))
 						]
 						+ SVerticalBox::Slot()
-						.Padding(0.0f, 5.0f, 0.0f, 0.0f)
+						.Padding(0.0f, 5.0f, 0.0f, 5.0f)
 						.FillHeight(1.0f)
 						.HAlign(HAlign_Center)
 						[
@@ -190,6 +205,26 @@ public:
 								.OverflowPolicy(ETextOverflowPolicy::Ellipsis)
 								.Clipping(EWidgetClipping::ClipToBounds)
 						]
+						+ SVerticalBox::Slot()
+						.AutoHeight()
+						.HAlign(HAlign_Center)
+						.Padding(0.0f, 5.0f, 0.0f, 4.0f)
+						[
+							SAssignNew(PinButton, SCheckBox)
+							.Style(FAppStyle::Get(), "ToggleButtonCheckbox")
+							.Visibility(this, &STabDrawerButton::GetPinButtonVisibility)
+							.ToolTipText(this, &STabDrawerButton::GetPinButtonToolTipText)
+							.IsChecked(this, &STabDrawerButton::IsPinButtonChecked)
+							.OnCheckStateChanged(this, &STabDrawerButton::OnPinButtonCheckStateChanged)
+							.Padding(2.0f)
+							.HAlign(HAlign_Center)
+							[
+								SNew(SImage)
+								.ColorAndOpacity(FSlateColor::UseForeground())
+								.Image(FAppStyle::Get().GetBrush("Docking.Sidebar.PinTabIcon"))
+								.DesiredSizeOverride(FVector2D(12, 12))
+							]
+						]
 					]
 				]
 				+ SOverlay::Slot()
@@ -200,10 +235,11 @@ public:
 				+ SOverlay::Slot()
 				.HAlign(Location == ESidebarLocation::Left ? HAlign_Left : HAlign_Right)
 				[
-					SAssignNew(OpenIndicator, SComplexGradient)
+					SAssignNew(ActiveIndicator, SComplexGradient)
 					.DesiredSizeOverride(FVector2D(1.0f, 1.0f))
 					.GradientColors(GradientStops)
 					.Orientation(EOrientation::Orient_Horizontal)
+					.Visibility(this, &STabDrawerButton::GetActiveTabIndicatorVisibility)
 				]
 			]
 		];
@@ -232,8 +268,7 @@ public:
 
 		if (OpenedDrawer == Tab)
 		{
-			// this button is the one with the tab that is actually opened so show the opened indicator
-			OpenIndicator->SetVisibility(EVisibility::HitTestInvisible);
+			// this button is the one with the tab that is actually opened so show the tab border
 			OpenBorder->SetVisibility(EVisibility::HitTestInvisible);
 			MainButton->SetButtonStyle(&FAppStyle::Get().GetWidgetStyle<FButtonStyle>("Docking.SidebarButton.Opened"));
 
@@ -250,7 +285,6 @@ public:
 		}
 		else
 		{
-			OpenIndicator->SetVisibility(EVisibility::Collapsed);
 			OpenBorder->SetVisibility(EVisibility::Collapsed);
 			MainButton->SetButtonStyle(&FAppStyle::Get().GetWidgetStyle<FButtonStyle>("Docking.SidebarButton.Closed"));
 		}
@@ -289,7 +323,7 @@ public:
 
 	virtual FSlateColor GetForegroundColor() const
 	{
-		if (OpenIndicator->GetVisibility() != EVisibility::Collapsed)
+		if (ActiveIndicator->GetVisibility() != EVisibility::Collapsed)
 		{
 			return DockTabStyle->ActiveForegroundColor;
 		}
@@ -302,13 +336,54 @@ public:
 	}
 
 private:
+	EVisibility GetActiveTabIndicatorVisibility() const
+	{
+		return Tab->IsActive() ? EVisibility::HitTestInvisible : EVisibility::Collapsed;
+	}
+
+	EVisibility GetPinButtonVisibility() const
+	{
+		return (IsTabPinned(Tab.ToSharedRef()) || IsHovered() || Tab->IsActive()) ? EVisibility::Visible : EVisibility::Hidden;
+	}
+
+	FText GetPinButtonToolTipText() const
+	{
+		if (IsTabPinned(Tab.ToSharedRef()))
+		{
+			return LOCTEXT("UnpinTabToolTip", "Unpin Tab");
+		}
+		else
+		{
+			return LOCTEXT("PinTabToolTip", "Pin Tab");
+		}
+	}
+
+	ECheckBoxState IsPinButtonChecked() const
+	{
+		if (IsTabPinned(Tab.ToSharedRef()))
+		{
+			return ECheckBoxState::Checked;
+		}
+		else
+		{
+			return ECheckBoxState::Unchecked;
+		}
+	}
+
+	void OnPinButtonCheckStateChanged(ECheckBoxState State)
+	{
+		OnDrawerPinButtonToggled.ExecuteIfBound(Tab.ToSharedRef(), State == ECheckBoxState::Checked);
+	}
+
 	TSharedPtr<SDockTab> Tab;
 	TSharedPtr<STabDrawerTextBlock> Label;
-	TSharedPtr<SWidget> OpenIndicator;
+	TSharedPtr<SWidget> ActiveIndicator;
 	TSharedPtr<SBorder> OpenBorder;
 	TSharedPtr<SButton> MainButton;
+	TSharedPtr<SCheckBox> PinButton;
 	FOnGetContent OnGetContextMenuContent;
-	FOnTabDrawerButtonClicked OnDrawerButtonClicked;
+	FOnTabDrawerButtonPressed OnDrawerButtonPressed;
+	FOnTabDrawerPinButtonToggled OnDrawerPinButtonToggled;
 	const FDockTabStyle* DockTabStyle;
 	ESidebarLocation Location;
 };
@@ -324,8 +399,6 @@ void STabSidebar::Construct(const FArguments& InArgs, TSharedRef<SOverlay> InDra
 	Location = InArgs._Location;
 	DrawersOverlay = InDrawersOverlay;
 
-	FGlobalTabmanager::Get()->OnTabForegrounded_Subscribe(FOnActiveTabChanged::FDelegate::CreateSP(this, &STabSidebar::OnActiveTabChanged));
-
 	ChildSlot
 	.Padding(FMargin(
 		Location == ESidebarLocation::Right ? 2.0f : 0.0f,
@@ -356,7 +429,8 @@ void STabSidebar::AddTab(TSharedRef<SDockTab> Tab)
 
 		TSharedRef<STabDrawerButton> TabButton =
 			SNew(STabDrawerButton, Tab, Location)
-			.OnDrawerButtonClicked(this, &STabSidebar::OnTabDrawerButtonClicked)
+			.OnDrawerButtonPressed(this, &STabSidebar::OnTabDrawerButtonPressed)
+			.OnDrawerPinButtonToggled(this, &STabSidebar::OnTabDrawerPinButtonToggled)
 			.OnGetContextMenuContent(this, &STabSidebar::OnGetTabDrawerContextMenuWidget, Tab);
 
 		// Figure out the size this tab should be when opened later. We do it now when the tab still has valid geometry.  Once it is moved to the sidebar it will not.
@@ -371,6 +445,13 @@ void STabSidebar::AddTab(TSharedRef<SDockTab> Tab)
 			}
 		}
 
+		// We don't currently allow more than one pinned tab per sidebar, so enforce that
+		// Note: it's possible to relax this if users actually want multiple pinned tabs
+		if (FindFirstPinnedTab())
+		{
+			SetTabPinned(Tab, false);
+		}
+
 		TabBox->AddSlot()
 			// Make the tabs evenly fill the sidebar until they reach the max size
 			.FillHeight(1.0f)
@@ -381,6 +462,12 @@ void STabSidebar::AddTab(TSharedRef<SDockTab> Tab)
 			];
 
 		Tabs.Emplace(Tab, TabButton);
+
+		// If this tab is a pinned tab, then open the drawer automatically after it's added
+		if (IsTabPinned(Tab))
+		{
+			OpenDrawerNextFrame(Tab, /*bAnimateOpen=*/ false);
+		}
 	}
 }
 
@@ -401,6 +488,12 @@ bool STabSidebar::RemoveTab(TSharedRef<SDockTab> TabToRemove)
 		TabBox->RemoveSlot(TabPair.Value);
 
 		RemoveDrawer(TabToRemove);
+		SummonPinnedTabIfNothingOpened();
+
+		// Clear the pinned flag when the tab is removed from the sidebar.
+		// (Users probably expect that pinning a tab, restoring it/closing it,
+		// then moving it to the sidebar again will leave it unpinned the second time.)
+		SetTabPinned(TabToRemove, false);
 
 		if (Tabs.Num() == 0)
 		{
@@ -463,21 +556,55 @@ bool STabSidebar::TryOpenSidebarDrawer(TSharedRef<SDockTab> ForTab)
 
 	if (FoundIndex != INDEX_NONE)
 	{
-		OpenDrawerNextFrame(ForTab);
+		OpenDrawerNextFrame(ForTab, /*bAnimateOpen=*/ true);
 		return true;
 	}
 
 	return false;
 }
 
-void STabSidebar::OnTabDrawerButtonClicked(TSharedRef<SDockTab> ForTab)
+void STabSidebar::OnTabDrawerButtonPressed(TSharedRef<SDockTab> ForTab)
+{
+	if (ForTab->IsActive())
+	{
+		// When clicking on the button of an active (but unpinned) tab, close that tab drawer
+		if (!IsTabPinned(ForTab))
+		{
+			CloseDrawerInternal(ForTab);
+		}
+	}
+	else
+	{
+		// Otherwise clicking on an inactive tab should open the drawer
+		OpenDrawerInternal(ForTab, /*bAnimateOpen=*/ true);
+	}
+}
+
+void STabSidebar::OnTabDrawerPinButtonToggled(TSharedRef<SDockTab> ForTab, bool bIsPinned)
 {
-	OpenDrawerInternal(ForTab);
+	// Set pin state for given tab; clear the pin state for all other tabs
+	for (const auto& TabAndButton : Tabs)
+	{
+		const TSharedRef<SDockTab>& Tab = TabAndButton.Key;
+		SetTabPinned(Tab, Tab == ForTab ? bIsPinned : false);
+	}
+
+	// Open any newly-pinned tab
+	if (bIsPinned)
+	{
+		OpenDrawerInternal(ForTab, /*bAnimateOpen=*/ true);
+	}
 }
 
 void STabSidebar::OnTabDrawerFocusLost(TSharedRef<STabDrawer> Drawer)
 {
-	Drawer->Close();
+	// Don't automatically close a pinned tab that is in the foreground
+	if (IsTabPinned(Drawer->GetTab()) && GetForegroundTab() == Drawer->GetTab())
+	{
+		return;
+	}
+
+	CloseDrawerInternal(Drawer->GetTab());
 }
 
 void STabSidebar::OnTabDrawerClosed(TSharedRef<STabDrawer> Drawer)
@@ -493,15 +620,6 @@ void STabSidebar::OnTargetDrawerSizeChanged(TSharedRef<STabDrawer> Drawer, float
 	Tab->GetParentDockTabStack()->SetTabSidebarSizeCoefficient(Tab, TargetDrawerSizePct);
 }
 
-void STabSidebar::OnActiveTabChanged(TSharedPtr<SDockTab> NewlyActivated, TSharedPtr<SDockTab> PreviouslyActive)
-{
-	// If a new major tab was activated remove any visible drawer instantly
-	if (NewlyActivated.IsValid() && NewlyActivated->GetVisualTabRole() == ETabRole::MajorTab)
-	{
-		RemoveAllDrawers();
-	}
-}
-
 TSharedRef<SWidget> STabSidebar::OnGetTabDrawerContextMenuWidget(TSharedRef<SDockTab> ForTab)
 {
 	const bool bCloseAfterSelection = true;
@@ -555,17 +673,9 @@ void STabSidebar::OnCloseTab(TSharedRef<SDockTab> TabToClose)
 
 void STabSidebar::RemoveDrawer(TSharedRef<SDockTab> ForTab)
 {
-	TSharedRef<STabDrawer>* OpenedDrawer =
-		OpenedDrawers.FindByPredicate(
-			[&ForTab](TSharedRef<STabDrawer>& Drawer)
-			{
-				return ForTab == Drawer->GetTab();
-			}
-	);
-
-	if(OpenedDrawer)
+	if(TSharedPtr<STabDrawer> OpenedDrawer = FindOpenedDrawer(ForTab))
 	{
-		TSharedRef<STabDrawer> OpenedDrawerRef = *OpenedDrawer;
+		TSharedRef<STabDrawer> OpenedDrawerRef = OpenedDrawer.ToSharedRef();
 
 		bool bRemoveSuccessful = DrawersOverlay->RemoveSlot(OpenedDrawerRef);
 		ensure(bRemoveSuccessful);
@@ -578,22 +688,10 @@ void STabSidebar::RemoveDrawer(TSharedRef<SDockTab> ForTab)
 	UpdateDrawerAppearance();
 }
 
-void STabSidebar::CloseAllDrawers()
-{
-	PendingTabToOpen.Reset();
-
-	// Closing drawers can remove them from the opened drawers list so copy the list first
-	TArray<TSharedRef<STabDrawer>> OpenedDrawersCopy = OpenedDrawers;
-
-	for (TSharedRef<STabDrawer>& Drawer : OpenedDrawersCopy)
-	{
-		Drawer->Close();
-	}
-}
-
 void STabSidebar::RemoveAllDrawers()
 {
 	PendingTabToOpen.Reset();
+	bAnimatePendingTabOpen = false;
 
 	// Closing drawers can remove them from the opened drawers list so copy the list first
 	TArray<TSharedRef<STabDrawer>> OpenedDrawersCopy = OpenedDrawers;
@@ -608,106 +706,138 @@ EActiveTimerReturnType STabSidebar::OnOpenPendingDrawerTimer(double CurrentTime,
 {
 	if (TSharedPtr<SDockTab> Tab = PendingTabToOpen.Pin())
 	{
-		OpenDrawerInternal(Tab.ToSharedRef());
+		// Wait until the drawers overlay has been arranged once to open the drawer
+		// It might not have geometry yet if we're adding back tabs on startup
+		if (DrawersOverlay->GetTickSpaceGeometry().GetLocalSize().IsZero())
+		{
+			return EActiveTimerReturnType::Continue;
+		}
+
+		OpenDrawerInternal(Tab.ToSharedRef(), bAnimatePendingTabOpen);
 	}
 
 	OpenPendingDrawerTimerHandle.Reset();
 	PendingTabToOpen.Reset();
+	bAnimatePendingTabOpen = false;
 
 	return EActiveTimerReturnType::Stop;
 }
 
-void STabSidebar::OpenDrawerNextFrame(TSharedRef<SDockTab> ForTab)
+void STabSidebar::OpenDrawerNextFrame(TSharedRef<SDockTab> ForTab, bool bAnimateOpen)
 {
 	PendingTabToOpen = ForTab;
+	bAnimatePendingTabOpen = bAnimateOpen;
 	if (!OpenPendingDrawerTimerHandle.IsValid())
 	{
 		OpenPendingDrawerTimerHandle = RegisterActiveTimer(0.0f, FWidgetActiveTimerDelegate::CreateSP(this, &STabSidebar::OnOpenPendingDrawerTimer));
 	}
 }
 
-void STabSidebar::OpenDrawerInternal(TSharedRef<SDockTab> ForTab)
+void STabSidebar::OpenDrawerInternal(TSharedRef<SDockTab> ForTab, bool bAnimateOpen)
 {
-	TSharedRef<STabDrawer>* OpenedDrawer =
-		OpenedDrawers.FindByPredicate(
-			[&ForTab](TSharedRef<STabDrawer>& Drawer)
-			{
-				return ForTab == Drawer->GetTab();
-			}
-	);
-
-	if (OpenedDrawer)
+	if (FindOpenedDrawer(ForTab))
 	{
-		// Drawer already opened close it
-		(*OpenedDrawer)->Close();
+		// Drawer already opened so don't do anything
+		return;
 	}
-	else
-	{
-		PendingTabToOpen.Reset();
 
-		const FGeometry DrawersOverlayGeometry = DrawersOverlay->GetTickSpaceGeometry();
-		const FGeometry MyGeometry = GetTickSpaceGeometry();
+	PendingTabToOpen.Reset();
+	bAnimatePendingTabOpen = false;
 
-		// Calculate padding for the drawer itself
-		const float MinDrawerSize = MyGeometry.GetLocalSize().X - 4.0f; // overlap with sidebar border slightly
+	const FGeometry DrawersOverlayGeometry = DrawersOverlay->GetTickSpaceGeometry();
+	const FGeometry MyGeometry = GetTickSpaceGeometry();
 
-		const FVector2D ShadowOffset(8, 8);
+	// Calculate padding for the drawer itself
+	const float MinDrawerSize = MyGeometry.GetLocalSize().X - 4.0f; // overlap with sidebar border slightly
 
-		FMargin SlotPadding(
-			Location == ESidebarLocation::Left ? MinDrawerSize : 0.0f,
-			-ShadowOffset.Y,
-			Location == ESidebarLocation::Right ? MinDrawerSize : 0.0f,
-			-ShadowOffset.Y
-		);
+	const FVector2D ShadowOffset(8, 8);
 
-		const float AvailableWidth = DrawersOverlayGeometry.GetLocalSize().X - SlotPadding.GetTotalSpaceAlong<EOrientation::Orient_Horizontal>();
+	FMargin SlotPadding(
+		Location == ESidebarLocation::Left ? MinDrawerSize : 0.0f,
+		-ShadowOffset.Y,
+		Location == ESidebarLocation::Right ? MinDrawerSize : 0.0f,
+		-ShadowOffset.Y
+	);
 
-		const float MaxPct = .5f;
-		const float MaxDrawerSize = AvailableWidth * 0.50f;
+	const float AvailableWidth = DrawersOverlayGeometry.GetLocalSize().X - SlotPadding.GetTotalSpaceAlong<EOrientation::Orient_Horizontal>();
 
-		float TargetDrawerSizePct = ForTab->GetParentDockTabStack()->GetTabSidebarSizeCoefficient(ForTab);
-		TargetDrawerSizePct = FMath::Clamp(TargetDrawerSizePct, .0f, .5f);
+	const float MaxPct = .5f;
+	const float MaxDrawerSize = AvailableWidth * 0.50f;
 
+	float TargetDrawerSizePct = ForTab->GetParentDockTabStack()->GetTabSidebarSizeCoefficient(ForTab);
+	TargetDrawerSizePct = FMath::Clamp(TargetDrawerSizePct, .0f, .5f);
 
-		const float TargetDrawerSize = AvailableWidth * TargetDrawerSizePct;
 
-		const TPair<TSharedRef<SDockTab>, TSharedRef<STabDrawerButton>>* TabEntry = Tabs.FindByPredicate(
-			[ForTab](auto TabPair) {
-				return TabPair.Key == ForTab;
-			});
-		const TWeakPtr<SWidget> TabButton = TabEntry ? TWeakPtr<SWidget>(TabEntry->Value) : nullptr;
-
-		TSharedRef<STabDrawer> NewDrawer =
-			SNew(STabDrawer, ForTab, TabButton, Location == ESidebarLocation::Left ? ETabDrawerOpenDirection::Left : ETabDrawerOpenDirection::Right)
-			.MinDrawerSize(MinDrawerSize)
-			.TargetDrawerSize(TargetDrawerSize)
-			.MaxDrawerSize(MaxDrawerSize)
-			.OnDrawerFocusLost(this, &STabSidebar::OnTabDrawerFocusLost)
-			.OnDrawerClosed(this, &STabSidebar::OnTabDrawerClosed)
-			.OnTargetDrawerSizeChanged(this, &STabSidebar::OnTargetDrawerSizeChanged)
-			[
-				ForTab->GetContent()
-			];
+	const float TargetDrawerSize = AvailableWidth * TargetDrawerSizePct;
 
-		DrawersOverlay->AddSlot()
-			.Padding(SlotPadding)
-			.HAlign(Location == ESidebarLocation::Left ? HAlign_Left : HAlign_Right)
-			[
-				NewDrawer
-			];
+	const TPair<TSharedRef<SDockTab>, TSharedRef<STabDrawerButton>>* TabEntry = Tabs.FindByPredicate(
+		[ForTab](auto TabPair) {
+			return TabPair.Key == ForTab;
+		});
+	const TWeakPtr<SWidget> TabButton = TabEntry ? TWeakPtr<SWidget>(TabEntry->Value) : nullptr;
+
+	TSharedRef<STabDrawer> NewDrawer =
+		SNew(STabDrawer, ForTab, TabButton, Location == ESidebarLocation::Left ? ETabDrawerOpenDirection::Left : ETabDrawerOpenDirection::Right)
+		.MinDrawerSize(MinDrawerSize)
+		.TargetDrawerSize(TargetDrawerSize)
+		.MaxDrawerSize(MaxDrawerSize)
+		.OnDrawerFocusLost(this, &STabSidebar::OnTabDrawerFocusLost)
+		.OnDrawerClosed(this, &STabSidebar::OnTabDrawerClosed)
+		.OnTargetDrawerSizeChanged(this, &STabSidebar::OnTargetDrawerSizeChanged)
+		[
+			ForTab->GetContent()
+		];
+
+	DrawersOverlay->AddSlot()
+		.Padding(SlotPadding)
+		.HAlign(Location == ESidebarLocation::Left ? HAlign_Left : HAlign_Right)
+		[
+			NewDrawer
+		];
+
+	NewDrawer->Open(bAnimateOpen);
 
-		NewDrawer->Open();
+	OpenedDrawers.Add(NewDrawer);
 
-		FSlateApplication::Get().SetKeyboardFocus(NewDrawer);
+	// This changes the focus and will trigger focus-related events, such as closing other tabs,
+	// so it's important that we only call it after we added the new drawer to OpenedDrawers.
+	FSlateApplication::Get().SetKeyboardFocus(NewDrawer);
 
-		OpenedDrawers.Add(NewDrawer);
+	ForTab->OnTabDrawerOpened();
 
-		ForTab->OnTabDrawerOpened();
+	UpdateDrawerAppearance();
+}
+
+void STabSidebar::CloseDrawerInternal(TSharedRef<SDockTab> ForTab)
+{
+	if (TSharedPtr<STabDrawer> OpenedDrawer = FindOpenedDrawer(ForTab))
+	{
+		OpenedDrawer->Close();
 	}
 
+	SummonPinnedTabIfNothingOpened();
 	UpdateDrawerAppearance();
 }
 
+void STabSidebar::SummonPinnedTabIfNothingOpened()
+{
+	// If there's already a tab in the foreground, don't bring the pinned tab forward
+	if (GetForegroundTab())
+	{
+		return;
+	}
+
+	// But if there's no current foreground tab, then bring forward a pinned tab (there should be at most one)
+	// This should happen when:
+	// - the current foreground tab is not pinned and loses focus
+	// - the current foreground tab's drawer is manually closed by pressing on the tab button
+	// - closing or restoring the current foreground tab
+	if (TSharedPtr<SDockTab> PinnedTab = FindFirstPinnedTab())
+	{
+		OpenDrawerInternal(PinnedTab.ToSharedRef(), /*bAnimateOpen=*/ true);
+	}
+}
+
 void STabSidebar::UpdateDrawerAppearance()
 {
 	TSharedPtr<SDockTab> OpenedTab;
@@ -722,4 +852,37 @@ void STabSidebar::UpdateDrawerAppearance()
 	}
 }
 
+TSharedPtr<SDockTab> STabSidebar::FindFirstPinnedTab() const
+{
+	const TPair<TSharedRef<SDockTab>, TSharedRef<STabDrawerButton>>* PinnedTab =
+		Tabs.FindByPredicate(
+			[](const TPair<TSharedRef<SDockTab>, TSharedRef<STabDrawerButton>>& TabAndButton)
+			{
+				return IsTabPinned(TabAndButton.Key);
+			});
+	return PinnedTab ? TSharedPtr<SDockTab>(PinnedTab->Key) : nullptr;
+}
+
+TSharedPtr<SDockTab> STabSidebar::GetForegroundTab() const
+{
+	const int32 Index =
+		OpenedDrawers.FindLastByPredicate(
+			[](const TSharedRef<STabDrawer>& Drawer)
+			{
+				return Drawer->IsOpen() && !Drawer->IsClosing();
+			});
+	return Index == INDEX_NONE ? nullptr : TSharedPtr<SDockTab>(OpenedDrawers[Index]->GetTab());
+}
+
+TSharedPtr<STabDrawer> STabSidebar::FindOpenedDrawer(TSharedRef<SDockTab> ForTab) const
+{
+	const TSharedRef<STabDrawer>* OpenedDrawer =
+		OpenedDrawers.FindByPredicate(
+			[&ForTab](TSharedRef<STabDrawer>& Drawer)
+			{
+				return ForTab == Drawer->GetTab();
+			});
+	return OpenedDrawer ? TSharedPtr<STabDrawer>(*OpenedDrawer) : nullptr;
+}
+
 #undef LOCTEXT_NAMESPACE
diff --git a/Engine/Source/Runtime/Slate/Private/Framework/Docking/STabSidebar.h b/Engine/Source/Runtime/Slate/Private/Framework/Docking/STabSidebar.h
--- a/Engine/Source/Runtime/Slate/Private/Framework/Docking/STabSidebar.h
+++ b/Engine/Source/Runtime/Slate/Private/Framework/Docking/STabSidebar.h
@@ -17,6 +17,9 @@ enum class ESidebarLocation : uint8;
 /**
  * A Sidebar is a widget that contains STabDrawers which can be opened and closed from the drawer to allow temporary access to the tab
  * A drawer is automatically dismissed when it or any of its children loses focus
+ * 
+ * Tabs can be pinned in the sidebar: this prevents them from being automatically dismissed when they lose focus,
+ * and it also causes them to be summoned back when all other drawers are dismissed
  */
 class STabSidebar : public SCompoundWidget
 {
@@ -77,12 +80,11 @@ public:
 	 */
 	bool TryOpenSidebarDrawer(TSharedRef<SDockTab> ForTab);
 private:
-	void OnTabDrawerButtonClicked(TSharedRef<SDockTab> ForTab);
+	void OnTabDrawerButtonPressed(TSharedRef<SDockTab> ForTab);
+	void OnTabDrawerPinButtonToggled(TSharedRef<SDockTab> ForTab, bool bIsPinned);
 	void OnTabDrawerFocusLost(TSharedRef<STabDrawer> Drawer);
 	void OnTabDrawerClosed(TSharedRef<STabDrawer> Drawer);
 	void OnTargetDrawerSizeChanged(TSharedRef<STabDrawer> Drawer, float NewSize);
-	/** Called when active tab changes which is used to determine if we should close an opened content browser drawer */
-	void OnActiveTabChanged(TSharedPtr<SDockTab> NewlyActivated, TSharedPtr<SDockTab> PreviouslyActive);
 
 	TSharedRef<SWidget> OnGetTabDrawerContextMenuWidget(TSharedRef<SDockTab> ForTab);
 	void OnRestoreTab(TSharedRef<SDockTab> TabToRestore);
@@ -95,23 +97,42 @@ private:
 	void RemoveDrawer(TSharedRef<SDockTab> ForTab);
 
 	/**
-	 * Removes all drawers instantly
+	 * Removes all drawers instantly (including drawers for pinned tabs)
 	 */
 	void RemoveAllDrawers();
-	
-	/**
-	 * Closes all drawers, playing a close animation and waiting to remove until the animation is complete
-	 */
-	void CloseAllDrawers();
 
 	EActiveTimerReturnType OnOpenPendingDrawerTimer(double CurrentTime, float DeltaTime);
-	void OpenDrawerNextFrame(TSharedRef<SDockTab> ForTab);
-	void OpenDrawerInternal(TSharedRef<SDockTab> ForTab);
+	void OpenDrawerNextFrame(TSharedRef<SDockTab> ForTab, bool bAnimateOpen);
+	void OpenDrawerInternal(TSharedRef<SDockTab> ForTab, bool bAnimateOpen);
+	void CloseDrawerInternal(TSharedRef<SDockTab> ForTab);
+
+	/**
+	 * Reopens the pinned tab only if there are no other open drawers;
+	 * this should be used to bring pinned tabs back after other tabs lose focus/are closed
+	 */
+	void SummonPinnedTabIfNothingOpened();
 
 	/**
 	 * Updates the appearance of open drawers
 	 */
 	void UpdateDrawerAppearance();
+
+	/**
+	 * Returns the first tab in this sidebar that is marked pinned
+	 */
+	TSharedPtr<SDockTab> FindFirstPinnedTab() const;
+
+	/**
+	 * Returns the tab for the last-opened drawer that is still open,
+	 * excluding any drawers that are in the process of closing
+	 */
+	TSharedPtr<SDockTab> GetForegroundTab() const;
+
+	/**
+	 * Returns the drawer for the given tab if it's open
+	 */
+	TSharedPtr<STabDrawer> FindOpenedDrawer(TSharedRef<SDockTab> ForTab) const;
+
 private:
 	TSharedPtr<SVerticalBox> TabBox;
 	TArray<TPair<TSharedRef<SDockTab>, TSharedRef<STabDrawerButton>>> Tabs;
@@ -122,4 +143,6 @@ private:
 	TArray<TSharedRef<STabDrawer>> OpenedDrawers;
 	/** Any pending drawer tab to open */
 	TWeakPtr<SDockTab> PendingTabToOpen;
+	/** Whether to animate the pending tab open */
+	bool bAnimatePendingTabOpen = false;
 }; 
diff --git a/Engine/Source/Runtime/Slate/Private/Framework/Docking/TabManager.cpp b/Engine/Source/Runtime/Slate/Private/Framework/Docking/TabManager.cpp
--- a/Engine/Source/Runtime/Slate/Private/Framework/Docking/TabManager.cpp
+++ b/Engine/Source/Runtime/Slate/Private/Framework/Docking/TabManager.cpp
@@ -298,12 +298,14 @@ TSharedRef<FTabManager::FLayoutNode> FTabManager::FLayout::NewFromString_Helper(
 
 			FString SidebarLocation;
 			float SidebarSizeCoefficient = .15f;
+			bool bPinnedInSidebar = false;
 			if (TabAsJson->TryGetStringField(TEXT("SidebarLocation"), SidebarLocation))
 			{
 				TabAsJson->TryGetNumberField(TEXT("SidebarCoeff"), SidebarSizeCoefficient);
+				TabAsJson->TryGetBoolField(TEXT("SidebarPinned"), bPinnedInSidebar);
 			}
 
-			NewStack->AddTab(TabId, TabStateFromString( TabAsJson->GetStringField(TEXT("TabState"))), SidebarLocationFromString(SidebarLocation), SidebarSizeCoefficient);
+			NewStack->AddTab(TabId, TabStateFromString( TabAsJson->GetStringField(TEXT("TabState"))), SidebarLocationFromString(SidebarLocation), SidebarSizeCoefficient, bPinnedInSidebar);
 		}
 		return NewStack;
 	}
@@ -470,6 +472,7 @@ TSharedRef<FJsonObject> FTabManager::FLayout::PersistToString_Helper(const TShar
 				{
 					TabAsJson->SetStringField(TEXT("SidebarLocation"), StringFromSidebarLocation(Tab.SidebarLocation));
 					TabAsJson->SetNumberField(TEXT("SidebarCoeff"), Tab.SidebarSizeCoefficient);
+					TabAsJson->SetBoolField(TEXT("SidebarPinned"), Tab.bPinnedInSidebar);
 				}
 
 				TabsAsJson.Add( MakeShareable( new FJsonValueObject(TabAsJson) ) );
diff --git a/Engine/Source/Runtime/Slate/Public/Framework/Docking/STabDrawer.h b/Engine/Source/Runtime/Slate/Public/Framework/Docking/STabDrawer.h
--- a/Engine/Source/Runtime/Slate/Public/Framework/Docking/STabDrawer.h
+++ b/Engine/Source/Runtime/Slate/Public/Framework/Docking/STabDrawer.h
@@ -62,8 +62,12 @@ class STabDrawer : public SCompoundWidget
 	/** Sets the current size of the drawer, ignoring any open/close animation */
 	void SetCurrentSize(float InSize);
 
-	/** Begins an animation which opens the drawer */
-	void Open();
+	/**
+	 * Opens the drawer
+	 *
+	 * @param bAnimateOpen Whether to play an animation when opening the drawer, defaults to true
+	 */
+	void Open(bool bAnimateOpen=true);
 
 	/** Begins an animation which closes the drawer */
 	void Close();
@@ -71,6 +75,9 @@ class STabDrawer : public SCompoundWidget
 	/** @return true if the drawer is open */
 	bool IsOpen() const;
 
+	/** @return true if the drawer is currently playing the close animation */
+	bool IsClosing() const;
+
 	/** @return the tab whose contents is being shown in the drawer */
 	const TSharedRef<SDockTab> GetTab() const;
 
diff --git a/Engine/Source/Runtime/Slate/Public/Framework/Docking/TabManager.h b/Engine/Source/Runtime/Slate/Public/Framework/Docking/TabManager.h
--- a/Engine/Source/Runtime/Slate/Public/Framework/Docking/TabManager.h
+++ b/Engine/Source/Runtime/Slate/Public/Framework/Docking/TabManager.h
@@ -399,15 +399,17 @@ class SLATE_API FTabManager : public TSharedFromThis<FTabManager>
 				, TabState(InTabState)
 				, SidebarLocation(ESidebarLocation::None)
 				, SidebarSizeCoefficient(0.0f)
+				, bPinnedInSidebar(false)
 			{
 				check(InTabState != ETabState::SidebarTab);
 			}
 
-			FTab(const FTabId& InTabId, ETabState::Type InTabState, ESidebarLocation InSidebarLocation, float InSidebarSizeCoefficient)
+			FTab(const FTabId& InTabId, ETabState::Type InTabState, ESidebarLocation InSidebarLocation, float InSidebarSizeCoefficient, bool bInPinnedInSidebar)
 				: TabId(InTabId)
 				, TabState(InTabState)
 				, SidebarLocation(InSidebarLocation)
 				, SidebarSizeCoefficient(InSidebarSizeCoefficient)
+				, bPinnedInSidebar(bInPinnedInSidebar)
 			{
 				check(InTabState != ETabState::SidebarTab || InSidebarLocation != ESidebarLocation::None);
 			}
@@ -421,6 +423,7 @@ class SLATE_API FTabManager : public TSharedFromThis<FTabManager>
 			ETabState::Type TabState;
 			ESidebarLocation SidebarLocation;
 			float SidebarSizeCoefficient;
+			bool bPinnedInSidebar;
 		};
 
 		class SLATE_API FStack : public FLayoutNode
@@ -446,18 +449,18 @@ class SLATE_API FTabManager : public TSharedFromThis<FTabManager>
 					return SharedThis(this);
 				}
 
-				TSharedRef<FStack> AddTab(const FName TabType, ETabState::Type InTabState, ESidebarLocation InSidebarLocation, float SidebarSizeCoefficient)
+				TSharedRef<FStack> AddTab(const FName TabType, ETabState::Type InTabState, ESidebarLocation InSidebarLocation, float SidebarSizeCoefficient, bool bPinnedInSidebar=false)
 				{
 					check(InTabState != ETabState::SidebarTab || InSidebarLocation != ESidebarLocation::None);
-					Tabs.Add(FTab(FTabId(TabType), InTabState, InSidebarLocation, SidebarSizeCoefficient));
+					Tabs.Add(FTab(FTabId(TabType), InTabState, InSidebarLocation, SidebarSizeCoefficient, bPinnedInSidebar));
 					return SharedThis(this);
 				}
 
-				TSharedRef<FStack> AddTab(const FTabId TabId, ETabState::Type InTabState, ESidebarLocation InSidebarLocation, float SidebarSizeCoefficient)
+				TSharedRef<FStack> AddTab(const FTabId TabId, ETabState::Type InTabState, ESidebarLocation InSidebarLocation, float SidebarSizeCoefficient, bool bPinnedInSidebar=false)
 				{
 					check(InTabState != ETabState::SidebarTab || InSidebarLocation != ESidebarLocation::None);
 
-					Tabs.Add(FTab(TabId, InTabState, InSidebarLocation, SidebarSizeCoefficient));
+					Tabs.Add(FTab(TabId, InTabState, InSidebarLocation, SidebarSizeCoefficient, bPinnedInSidebar));
 
 					return SharedThis(this);
 				}
diff --git a/Engine/Source/Runtime/SlateCore/Private/Styling/StarshipCoreStyle.cpp b/Engine/Source/Runtime/SlateCore/Private/Styling/StarshipCoreStyle.cpp
--- a/Engine/Source/Runtime/SlateCore/Private/Styling/StarshipCoreStyle.cpp
+++ b/Engine/Source/Runtime/SlateCore/Private/Styling/StarshipCoreStyle.cpp
@@ -1663,6 +1663,8 @@ void FStarshipCoreStyle::SetupDockingStyles(TSharedRef<FStyle>& Style)
 		Style->Set("Docking.Sidebar.Border_SquareRight", new FSlateRoundedBoxBrush(FStyleColors::Transparent, FVector4(CornerRadius, 0.0f, 0.0f, CornerRadius), FStyleColors::Hover, BorderWidth));
 	}
 
+	Style->Set("Docking.Sidebar.PinTabIcon", new IMAGE_BRUSH_SVG("Starship/Docking/pin", Icon12x12));
+
 	// Dock Cross
 	Style->Set("Docking.Cross.DockLeft", new IMAGE_BRUSH("/Docking/OuterDockingIndicator", FVector2D(6.f, 6.f), FLinearColor(1.0f, 0.35f, 0.0f, 0.25f)));
 	Style->Set("Docking.Cross.DockLeft_Hovered", new IMAGE_BRUSH("/Docking/OuterDockingIndicator", FVector2D(6.f, 6.f), FLinearColor(1.0f, 0.35f, 0.0f)));
