Update context menu item wording for tab wells/sidebars

diff --git a/Engine/Source/Runtime/Slate/Private/Framework/Docking/SDockingTabStack.cpp b/Engine/Source/Runtime/Slate/Private/Framework/Docking/SDockingTabStack.cpp
--- a/Engine/Source/Runtime/Slate/Private/Framework/Docking/SDockingTabStack.cpp
+++ b/Engine/Source/Runtime/Slate/Private/Framework/Docking/SDockingTabStack.cpp
@@ -198,6 +198,7 @@ void SDockingTabStack::Construct( const FArguments& InArgs, const TSharedRef<FTa
 				.Visibility(this, &SDockingTabStack::GetUnhideButtonVisibility)
 				.DesiredSizeScale(this, &SDockingTabStack::GetUnhideTabWellButtonScale)
 				.ButtonColorAndOpacity(this, &SDockingTabStack::GetUnhideTabWellButtonOpacity)
+				.ToolTipText(LOCTEXT("UnhideTabWellToolTip", "Show Tabs"))
 				[
 					// button should be big enough to show its own image
 					SNew(SSpacer)
@@ -647,7 +648,7 @@ TSharedRef<SWidget> SDockingTabStack::MakeContextMenu()
 		MenuBuilder.BeginSection("DockingTabStackOptions", LOCTEXT("TabOptionsHeading", "Options") );
 		{
 			MenuBuilder.AddMenuEntry(
-				LOCTEXT("CollapseTabWell", "Collapse Tab Well"),
+				LOCTEXT("CollapseTabWell", "Hide Tabs"),
 				LOCTEXT("CollapseTabWellTooltip", "Collapses the tabs headers to save room."),
 				FSlateIcon(),
 				FUIAction(
@@ -659,7 +660,7 @@ TSharedRef<SWidget> SDockingTabStack::MakeContextMenu()
 			if(IsTabAllowedInSidebar(TabWell->GetForegroundTab()))
 			{
 				MenuBuilder.AddMenuEntry(
-					LOCTEXT("MoveToSidebar", "Move To Sidebar"),
+					LOCTEXT("MoveToSidebar", "Dock to Sidebar"),
 					LOCTEXT("MoveToSidebarTooltip", "Moves this tab to a sidebar drawer on the side of the window closest to the tab.\nThe tab can be opened from the drawer and will automatically close again when clicking off it."),
 					FSlateIcon(),
 					FUIAction(
diff --git a/Engine/Source/Runtime/Slate/Private/Framework/Docking/STabSidebar.cpp b/Engine/Source/Runtime/Slate/Private/Framework/Docking/STabSidebar.cpp
--- a/Engine/Source/Runtime/Slate/Private/Framework/Docking/STabSidebar.cpp
+++ b/Engine/Source/Runtime/Slate/Private/Framework/Docking/STabSidebar.cpp
@@ -629,7 +629,7 @@ TSharedRef<SWidget> STabSidebar::OnGetTabDrawerContextMenuWidget(TSharedRef<SDoc
 		MenuBuilder.BeginSection("RestoreOptions", LOCTEXT("RestoreOptions", "Options"));
 		{
 			MenuBuilder.AddMenuEntry(
-				LOCTEXT("AutoHideTab", "Restore Tab"),
+				LOCTEXT("AutoHideTab", "Undock from Sidebar"),
 				LOCTEXT("HideTabWellTooltip", "Moves this tab out of the sidebar and back to a full tab where it previously was before it was added to the sidebar."),
 				FSlateIcon(),
 				FUIAction(
