Split ITableRow into its own file from STableRow.h.

diff --git a/Engine/Source/Runtime/Slate/Public/Widgets/Views/ITableRow.h b/Engine/Source/Runtime/Slate/Public/Widgets/Views/ITableRow.h
new file mode 100644
--- /dev/null
+++ b/Engine/Source/Runtime/Slate/Public/Widgets/Views/ITableRow.h
@@ -0,0 +1,71 @@
+// Copyright Epic Games, Inc. All Rights Reserved.
+
+#pragma once
+
+#include "Containers/BitArray.h"
+#include "Framework/Views/ITypedTableView.h"
+#include "Math/Vector2D.h"
+#include "Templates/SharedPointer.h"
+#include "Widgets/SWidget.h"
+
+/**
+ * Interface for table views to talk to their rows.
+ */
+class SLATE_API ITableRow
+{
+public:
+
+	/**
+	* Called when the row has been generated and associated with an item in the owning table.
+	* Any attempts to access the item associated with the row prior to this (i.e. in Construct()) will fail, as the association is not yet established in the owning table.
+	*/
+	virtual void InitializeRow() = 0;
+
+	/**
+	* Called when the row has been released from the owning table and is no longer associated with any items therein.
+	* Only relevant if the row widgets are pooled or otherwise referenced/kept alive outside the owning table. Otherwise, the row is destroyed.
+	*/
+	virtual void ResetRow() = 0;
+
+	/**
+	* @param InIndexInList  The index of the item for which this widget was generated
+	*/
+	virtual void SetIndexInList(int32 InIndexInList) = 0;
+
+	/** @return true if the corresponding item is expanded; false otherwise*/
+	virtual bool IsItemExpanded() const = 0;
+
+	/** Toggle the expansion of the item associated with this row */
+	virtual void ToggleExpansion() = 0;
+
+	/** @return True if the corresponding item is selected; false otherwise */
+	virtual bool IsItemSelected() const = 0;
+
+	/** @return how nested the item associated with this row when it is in a TreeView */
+	virtual int32 GetIndentLevel() const = 0;
+
+	/** @return Does this item have children? */
+	virtual int32 DoesItemHaveChildren() const = 0;
+
+	/** @return BitArray where each entry corresponds to whether this item needs a vertical wire draw for that depth. */
+	virtual TBitArray<> GetWiresNeededByDepth() const = 0;
+
+	/** @return true if this item is the last direct descendant of its parent. */
+	virtual bool IsLastChild() const = 0;
+		
+	/** @return this table row as a widget */
+	virtual TSharedRef<SWidget> AsWidget() = 0;
+
+	/** @return the content of this table row */
+	virtual TSharedPtr<SWidget> GetContent() = 0;
+
+	/** Called when the expander arrow for this row is shift+clicked */
+	virtual void Private_OnExpanderArrowShiftClicked() = 0;
+
+	/** @return the size for the specified column name */
+	virtual FVector2D GetRowSizeForColumn(const FName& InColumnName) const = 0;
+
+protected:
+	/** Called to query the selection mode for the row */
+	virtual ESelectionMode::Type GetSelectionMode() const = 0;
+};
\ No newline at end of file
diff --git a/Engine/Source/Runtime/Slate/Public/Widgets/Views/STableRow.h b/Engine/Source/Runtime/Slate/Public/Widgets/Views/STableRow.h
--- a/Engine/Source/Runtime/Slate/Public/Widgets/Views/STableRow.h
+++ b/Engine/Source/Runtime/Slate/Public/Widgets/Views/STableRow.h
@@ -3,29 +3,30 @@
 #pragma once
 
 #include "CoreMinimal.h"
-#include "Misc/Attribute.h"
 #include "InputCoreTypes.h"
-#include "Styling/SlateColor.h"
-#include "Layout/Geometry.h"
-#include "Input/Events.h"
+#include "ITableRow.h"
+#include "Framework/Views/ITypedTableView.h"
+#include "Framework/Views/TableViewTypeTraits.h"
 #include "Input/DragAndDrop.h"
+#include "Input/Events.h"
 #include "Input/Reply.h"
-#include "Widgets/SNullWidget.h"
-#include "Widgets/DeclarativeSyntaxSupport.h"
-#include "Widgets/SWidget.h"
+#include "Layout/Geometry.h"
 #include "Layout/Margin.h"
-#include "Styling/SlateTypes.h"
-#include "Styling/CoreStyle.h"
-#include "Widgets/Layout/SBorder.h"
-#include "Framework/Views/ITypedTableView.h"
-#include "Widgets/Views/STableViewBase.h"
+#include "Misc/Attribute.h"
 #include "Rendering/DrawElements.h"
+#include "Styling/CoreStyle.h"
+#include "Styling/SlateColor.h"
+#include "Styling/SlateTypes.h"
 #include "Types/SlateStructs.h"
-#include "Widgets/SBoxPanel.h"
+#include "Widgets/DeclarativeSyntaxSupport.h"
+#include "Widgets/Layout/SBorder.h"
 #include "Widgets/Layout/SBox.h"
+#include "Widgets/SBoxPanel.h"
+#include "Widgets/SNullWidget.h"
+#include "Widgets/SWidget.h"
 #include "Widgets/Views/SExpanderArrow.h"
 #include "Widgets/Views/SHeaderRow.h"
-#include "Framework/Views/TableViewTypeTraits.h"
+#include "Widgets/Views/STableViewBase.h"
 #if WITH_ACCESSIBILITY
 #include "GenericPlatform/Accessibility/GenericAccessibleInterfaces.h"
 #include "Widgets/Accessibility/SlateCoreAccessibleWidgets.h"
@@ -35,68 +36,6 @@
 
 template <typename ItemType> class SListView;
 
-/**
- * Interface for table views to talk to their rows.
- */
-class SLATE_API ITableRow
-{
-	public:
-
-		/**
-		 * Called when the row has been generated and associated with an item in the owning table.
-		 * Any attempts to access the item associated with the row prior to this (i.e. in Construct()) will fail, as the association is not yet established in the owning table.
-		 */
-		virtual void InitializeRow() = 0;
-
-		/**
-		 * Called when the row has been released from the owning table and is no longer associated with any items therein.
-		 * Only relevant if the row widgets are pooled or otherwise referenced/kept alive outside the owning table. Otherwise, the row is destroyed.
-		 */
-		virtual void ResetRow() = 0;
-
-		/**
-		 * @param InIndexInList  The index of the item for which this widget was generated
-		 */
-		virtual void SetIndexInList( int32 InIndexInList ) = 0;
-
-		/** @return true if the corresponding item is expanded; false otherwise*/
-		virtual bool IsItemExpanded() const = 0;
-
-		/** Toggle the expansion of the item associated with this row */
-		virtual void ToggleExpansion() = 0;
-
-		/** @return True if the corresponding item is selected; false otherwise */
-		virtual bool IsItemSelected() const = 0;
-
-		/** @return how nested the item associated with this row when it is in a TreeView */
-		virtual int32 GetIndentLevel() const = 0;
-
-		/** @return Does this item have children? */
-		virtual int32 DoesItemHaveChildren() const = 0;
-
-		/** @return BitArray where each entry corresponds to whether this item needs a vertical wire draw for that depth. */
-		virtual TBitArray<> GetWiresNeededByDepth() const = 0;
-
-		/** @return true if this item is the last direct descendant of its parent. */
-		virtual bool IsLastChild() const = 0;
-		
-		/** @return this table row as a widget */
-		virtual TSharedRef<SWidget> AsWidget() = 0;
-
-		/** @return the content of this table row */
-		virtual TSharedPtr<SWidget> GetContent() = 0;
-
-		/** Called when the expander arrow for this row is shift+clicked */
-		virtual void Private_OnExpanderArrowShiftClicked() = 0;
-
-		/** @return the size for the specified column name */
-		virtual FVector2D GetRowSizeForColumn(const FName& InColumnName) const = 0;
-
-	protected:
-		/** Called to query the selection mode for the row */
-		virtual ESelectionMode::Type GetSelectionMode() const = 0;
-};
-
 /**
  * When the table row should signal the owner widget that the selection changed.
  * This only affect the selection with the left mouse button!
