Fix SSuggestionTextBox not closing the suggestion list when the user clicks outside it.

diff --git a/Engine/Source/Runtime/Slate/Private/Widgets/Input/SSuggestionTextBox.cpp b/Engine/Source/Runtime/Slate/Private/Widgets/Input/SSuggestionTextBox.cpp
--- a/Engine/Source/Runtime/Slate/Private/Widgets/Input/SSuggestionTextBox.cpp
+++ b/Engine/Source/Runtime/Slate/Private/Widgets/Input/SSuggestionTextBox.cpp
@@ -29,6 +29,8 @@ void SSuggestionTextBox::Construct( const FArguments& InArgs )
 	OnTextChanged = InArgs._OnTextChanged;
 	OnTextCommitted = InArgs._OnTextCommitted;
 
+	FSlateApplication::Get().OnFocusChanging().AddSP(this, &SSuggestionTextBox::OnGlobalFocusChanging);
+
 	ChildSlot
 	[
 		SAssignNew(MenuAnchor, SMenuAnchor)
@@ -58,7 +60,7 @@ void SSuggestionTextBox::Construct( const FArguments& InArgs )
 					.BorderImage(InArgs._BackgroundImage)
 					.Padding(FMargin(2))
 					[
-						SNew(SVerticalBox)
+						SAssignNew(VerticalBox, SVerticalBox)
 						
 						+ SVerticalBox::Slot()
 							.AutoHeight()
@@ -389,3 +391,17 @@ void SSuggestionTextBox::HandleTextBoxTextCommitted( const FText& InText, ETextC
 		ClearSuggestions();
 	}
 }
+
+
+void SSuggestionTextBox::OnGlobalFocusChanging(const FFocusEvent& FocusEvent, const FWeakWidgetPath& OldFocusedWidgetPath, const TSharedPtr<SWidget>& OldFocusedWidget, const FWidgetPath& NewFocusedWidgetPath, const TSharedPtr<SWidget>& NewFocusedWidget)
+{
+	if (TextBox.IsValid() && OldFocusedWidgetPath.ContainsWidget(TextBox.Get()) && FocusEvent.GetCause() == EFocusCause::Mouse)
+	{
+		// If the textbox has lost focus and the SuggestionList has not gained it, then we assume the user clicked somewhere else in the app and clear the SuggestionList.
+		if (VerticalBox.IsValid() && !NewFocusedWidgetPath.ContainsWidget(VerticalBox.Get()))
+		{
+			ClearSuggestions();
+			OnTextCommitted.ExecuteIfBound(TextBox->GetText(), ETextCommit::Type::OnUserMovedFocus);
+		}
+	}
+}
diff --git a/Engine/Source/Runtime/Slate/Public/Widgets/Input/SSuggestionTextBox.h b/Engine/Source/Runtime/Slate/Public/Widgets/Input/SSuggestionTextBox.h
--- a/Engine/Source/Runtime/Slate/Public/Widgets/Input/SSuggestionTextBox.h
+++ b/Engine/Source/Runtime/Slate/Public/Widgets/Input/SSuggestionTextBox.h
@@ -16,6 +16,7 @@
 #include "Widgets/Views/STableRow.h"
 
 class SMenuAnchor;
+class SVerticalBox;
 
 template< typename ObjectType > class TAttribute;
 
@@ -220,6 +221,9 @@ private:
 	// Callback committing the text box input.
 	void HandleTextBoxTextCommitted( const FText& InText, ETextCommit::Type CommitInfo );
 
+	// Called when global focus changes so we can detect when to close the SuggestionList.
+	void OnGlobalFocusChanging(const FFocusEvent& FocusEvent, const FWeakWidgetPath& OldFocusedWidgetPath, const TSharedPtr<SWidget>& OldFocusedWidget, const FWidgetPath& NewFocusedWidgetPath, const TSharedPtr<SWidget>& NewFocusedWidget);
+
 private:
 
 	// Holds a flag to prevent recursive calls in UI callbacks.
@@ -243,6 +247,9 @@ private:
 	// The styling of the suggestion text.
 	const FTextBlockStyle* SuggestionTextStyle;
 
+	// The VerticalBox containing the SuggestionList.
+	TSharedPtr<SVerticalBox> VerticalBox;
+
 private:
 
 	// Holds a delegate that is executed before the input history is shown.
