Close unpinned tab drawers when any other minor tab becomes active

Previously, sidebar tab drawers would remain open if a minor tab was activated in a child window.
This change has the tab drawer listen to active tab changes and execute its focus-lost delegate if any other tab becomes active;
this has the effect of the sidebar requesting the drawer to close if the tab is unpinned.
Note that other child windows, like the color picker, still shouldn't cause the drawer to close (e.g. when picking a color in details).

diff --git a/Engine/Source/Runtime/Slate/Private/Framework/Docking/STabDrawer.cpp b/Engine/Source/Runtime/Slate/Private/Framework/Docking/STabDrawer.cpp
--- a/Engine/Source/Runtime/Slate/Private/Framework/Docking/STabDrawer.cpp
+++ b/Engine/Source/Runtime/Slate/Private/Framework/Docking/STabDrawer.cpp
@@ -61,6 +61,7 @@ void STabDrawer::Construct(const FArguments& InArgs, TSharedRef<SDockTab> InTab,
 	}
 
 	FSlateApplication::Get().OnFocusChanging().AddSP(this, &STabDrawer::OnGlobalFocusChanging);
+	FGlobalTabmanager::Get()->OnActiveTabChanged_Subscribe(FOnActiveTabChanged::FDelegate::CreateSP(this, &STabDrawer::OnActiveTabChanged));
 
 	bIsResizeHandleHovered = false;
 	bIsResizing = false;
@@ -538,10 +539,12 @@ void STabDrawer::OnGlobalFocusChanging(const FFocusEvent& FocusEvent, const FWea
 			}
 			else if (NewFocusedWidgetPath.IsValid())
 			{
-				// New focus is on something else, try to check if it's a menu
+				// New focus is on something else, try to check if it's a menu or child window
 				TSharedRef<SWindow> NewWindow = NewFocusedWidgetPath.GetWindow();
 				TSharedPtr<SWindow> MyWindow = FSlateApplication::Get().FindWidgetWindow(ThisWidget);
 
+				// See if this is a child window (like a color picker being opened from details), and if so, don't dismiss
+				// Rely on OnActiveTabChanged below to lose focus if the child window actually contains tabs
 				if (!NewWindow->IsDescendantOf(MyWindow))
 				{
 					if (TSharedPtr<SWidget> MenuHost = FSlateApplication::Get().GetMenuHostWidget())
@@ -573,3 +576,12 @@ void STabDrawer::OnGlobalFocusChanging(const FFocusEvent& FocusEvent, const FWea
 		}
 	}
 }
+
+void STabDrawer::OnActiveTabChanged(TSharedPtr<SDockTab> PreviouslyActive, TSharedPtr<SDockTab> NewlyActivated)
+{
+	// This tab lost the active status to some other tab; treat this like focus was lost
+	if (PreviouslyActive == ForTab && NewlyActivated.IsValid())
+	{
+		OnDrawerFocusLost.ExecuteIfBound(SharedThis(this));
+	}
+}
diff --git a/Engine/Source/Runtime/Slate/Public/Framework/Docking/STabDrawer.h b/Engine/Source/Runtime/Slate/Public/Framework/Docking/STabDrawer.h
--- a/Engine/Source/Runtime/Slate/Public/Framework/Docking/STabDrawer.h
+++ b/Engine/Source/Runtime/Slate/Public/Framework/Docking/STabDrawer.h
@@ -97,6 +97,7 @@ private:
 	FGeometry GetResizeHandleGeometry(const FGeometry& AllottedGeometry) const;
 	EActiveTimerReturnType UpdateAnimation(double CurrentTime, float DeltaTime);
 	void OnGlobalFocusChanging(const FFocusEvent& FocusEvent, const FWeakWidgetPath& OldFocusedWidgetPath, const TSharedPtr<SWidget>& OldFocusedWidget, const FWidgetPath& NewFocusedWidgetPath, const TSharedPtr<SWidget>& NewFocusedWidget);
+	void OnActiveTabChanged(TSharedPtr<SDockTab> NewlyActivated, TSharedPtr<SDockTab> PreviouslyActive);
 
 private:
 	FGeometry InitialResizeGeometry;
