Fix a number of details panel style issues

diff --git a/Engine/Source/Editor/DetailCustomizations/Private/BodyInstanceCustomization.cpp b/Engine/Source/Editor/DetailCustomizations/Private/BodyInstanceCustomization.cpp
--- a/Engine/Source/Editor/DetailCustomizations/Private/BodyInstanceCustomization.cpp
+++ b/Engine/Source/Editor/DetailCustomizations/Private/BodyInstanceCustomization.cpp
@@ -1530,6 +1530,7 @@ void FBodyInstanceCustomizationHelper::AddMassInKg(IDetailCategoryBuilder& Physi
 			MassInKgOverrideHandle->CreatePropertyNameWidget()
 		]
 		.ValueContent()
+		.VAlign(VAlign_Center)
 		[
 			SNew(SNumericEntryBox<float>)
 			.IsEnabled(this, &FBodyInstanceCustomizationHelper::IsBodyMassEnabled)
diff --git a/Engine/Source/Editor/DetailCustomizations/Private/ColorStructCustomization.cpp b/Engine/Source/Editor/DetailCustomizations/Private/ColorStructCustomization.cpp
--- a/Engine/Source/Editor/DetailCustomizations/Private/ColorStructCustomization.cpp
+++ b/Engine/Source/Editor/DetailCustomizations/Private/ColorStructCustomization.cpp
@@ -87,14 +87,17 @@ TSharedRef<SWidget> FColorStructCustomization::CreateColorWidget(TWeakPtr<IPrope
 	return
 		SNew(SBox)
 		.Padding(FMargin(0,0,4.0f,0.0f))
+		.VAlign(VAlign_Center)
 		[
 			SAssignNew(ColorWidgetBackgroundBorder, SBorder)
 			.Padding(1)
 			.BorderImage(FAppStyle::Get().GetBrush("ColorPicker.RoundedSolidBackground"))
 			.BorderBackgroundColor(this, &FColorStructCustomization::GetColorWidgetBorderColor)
+			.VAlign(VAlign_Center)
 			[
 				SNew(SOverlay)
 				+ SOverlay::Slot()
+				.VAlign(VAlign_Center)
 				[
 					SAssignNew(ColorPickerParentWidget, SColorBlock)
 					.AlphaBackgroundBrush(FAppStyle::Get().GetBrush("ColorPicker.RoundedAlphaBackground"))
@@ -107,6 +110,7 @@ TSharedRef<SWidget> FColorStructCustomization::CreateColorWidget(TWeakPtr<IPrope
 					.IsEnabled(this, &FColorStructCustomization::IsValueEnabled, StructWeakHandlePtr)
 				]
 				+ SOverlay::Slot()
+				.VAlign(VAlign_Center)
 				[
 					SNew(SBorder)
 					.Visibility(this, &FColorStructCustomization::GetMultipleValuesTextVisibility)
diff --git a/Engine/Source/Editor/DetailCustomizations/Private/MobilityCustomization.cpp b/Engine/Source/Editor/DetailCustomizations/Private/MobilityCustomization.cpp
--- a/Engine/Source/Editor/DetailCustomizations/Private/MobilityCustomization.cpp
+++ b/Engine/Source/Editor/DetailCustomizations/Private/MobilityCustomization.cpp
@@ -136,7 +136,7 @@ void FMobilityCustomization::GenerateHeaderRowContent(FDetailWidgetRow& WidgetRo
 		.Content()
 		[
 			SNew(STextBlock)
-			.Font(FAppStyle::Get().GetFontStyle("PropertyWindow.MobilityFont"))
+			.Font(FAppStyle::Get().GetFontStyle("PropertyWindow.NormalFont"))
 			.Text(this, &FMobilityCustomization::GetActiveMobilityText)
 			.ToolTipText(this, &FMobilityCustomization::GetActiveMobilityToolTip)
 		];
@@ -151,6 +151,7 @@ void FMobilityCustomization::GenerateHeaderRowContent(FDetailWidgetRow& WidgetRo
 	]
 	.ValueContent()
 	.MaxDesiredWidth(0)
+	.VAlign(VAlign_Center)
 	[
 		ComboBox
 	];
@@ -163,7 +164,7 @@ TSharedRef<SWidget> FMobilityCustomization::OnGenerateWidget(const FName InMobil
 	const FText ToolTip = GetOptionToolTip(Mobility, bForLight);
 
 	return SNew(STextBlock)
-		.Font(FAppStyle::Get().GetFontStyle("PropertyWindow.MobilityFont"))
+		.Font(FAppStyle::Get().GetFontStyle("PropertyWindow.NormalFont"))
 		.Text(Text)
 		.ToolTipText(ToolTip);
 }
diff --git a/Engine/Source/Editor/DetailCustomizations/Private/StaticMeshComponentDetails.cpp b/Engine/Source/Editor/DetailCustomizations/Private/StaticMeshComponentDetails.cpp
--- a/Engine/Source/Editor/DetailCustomizations/Private/StaticMeshComponentDetails.cpp
+++ b/Engine/Source/Editor/DetailCustomizations/Private/StaticMeshComponentDetails.cpp
@@ -56,6 +56,7 @@ void FStaticMeshComponentDetails::CustomizeDetails(IDetailLayoutBuilder& DetailB
 			OverrideLightResProperty->CreatePropertyNameWidget()
 		]
 		.ValueContent()
+		.VAlign(VAlign_Center)
 		[
 			SNew(SWidgetSwitcher)
 			.WidgetIndex(this, &FStaticMeshComponentDetails::HandleNoticeSwitcherWidgetIndex)
diff --git a/Engine/Source/Editor/EditorWidgets/Private/STextPropertyEditableTextBox.cpp b/Engine/Source/Editor/EditorWidgets/Private/STextPropertyEditableTextBox.cpp
--- a/Engine/Source/Editor/EditorWidgets/Private/STextPropertyEditableTextBox.cpp
+++ b/Engine/Source/Editor/EditorWidgets/Private/STextPropertyEditableTextBox.cpp
@@ -24,6 +24,8 @@
 #include "Internationalization/StringTableRegistry.h"
 #include "Serialization/TextReferenceCollector.h"
 #include "Styling/StyleColors.h"
+#include "Widgets/Layout/SLinkedBox.h"
+#include "SSimpleComboButton.h"
 
 #define LOCTEXT_NAMESPACE "STextPropertyEditableTextBox"
 
@@ -717,6 +719,7 @@ void STextPropertyEditableTextBox::Construct(const FArguments& InArgs, const TSh
 		[
 			SAssignNew(HorizontalBox, SHorizontalBox)
 			+SHorizontalBox::Slot()
+			.VAlign(VAlign_Center)
 			.FillWidth(1.0f)
 			[
 				SNew(SBox)
@@ -751,6 +754,7 @@ void STextPropertyEditableTextBox::Construct(const FArguments& InArgs, const TSh
 		[
 			SAssignNew(HorizontalBox, SHorizontalBox)
 			+SHorizontalBox::Slot()
+			.VAlign(VAlign_Center)
 			[
 				SNew(SBox)
 				.MinDesiredWidth(InArgs._MinDesiredWidth)
@@ -781,20 +785,8 @@ void STextPropertyEditableTextBox::Construct(const FArguments& InArgs, const TSh
 	HorizontalBox->AddSlot()
 		.AutoWidth()
 		[
-			SNew(SComboButton)
-			.HAlign(HAlign_Center)
-			.VAlign(VAlign_Center)
-			.ContentPadding(FMargin(4, 0))
-			.ButtonStyle(FEditorStyle::Get(), "HoverHintOnly")
-			.ForegroundColor(FSlateColor::UseForeground())
-			.ToolTipText(this, &STextPropertyEditableTextBox::GetAdvancedTextSettingsComboToolTip)
-			.HasDownArrow(false)
-			.ButtonContent()
-			[
-				SNew(SImage)
-				.Image(this, &STextPropertyEditableTextBox::GetAdvancedTextSettingsComboImage)
-				.ColorAndOpacity(FSlateColor::UseForeground())
-			]
+			SNew(SSimpleComboButton)
+			.Icon(this, &STextPropertyEditableTextBox::GetAdvancedTextSettingsComboImage)
 			.MenuContent()
 			[
 				SNew(SBox)
diff --git a/Engine/Source/Editor/LandscapeEditor/Private/LandscapeEditorDetailCustomization_NewLandscape.cpp b/Engine/Source/Editor/LandscapeEditor/Private/LandscapeEditorDetailCustomization_NewLandscape.cpp
--- a/Engine/Source/Editor/LandscapeEditor/Private/LandscapeEditorDetailCustomization_NewLandscape.cpp
+++ b/Engine/Source/Editor/LandscapeEditor/Private/LandscapeEditorDetailCustomization_NewLandscape.cpp
@@ -361,7 +361,6 @@ void FLandscapeEditorDetailCustomization_NewLandscape::CustomizeDetails(IDetailL
 		.FillWidth(1)
 		[
 			SNew(SNumericEntryBox<int32>)
-			.LabelVAlign(VAlign_Center)
 			.Font(DetailBuilder.GetDetailFont())
 			.MinValue(1)
 			.MaxValue(32)
@@ -386,7 +385,6 @@ void FLandscapeEditorDetailCustomization_NewLandscape::CustomizeDetails(IDetailL
 		.FillWidth(1)
 		[
 			SNew(SNumericEntryBox<int32>)
-			.LabelVAlign(VAlign_Center)
 			.Font(DetailBuilder.GetDetailFont())
 			.MinValue(1)
 			.MaxValue(32)
@@ -508,15 +506,12 @@ void FLandscapeEditorDetailCustomization_NewLandscape::CustomizeDetails(IDetailL
 		]
 	]
 	.ValueContent()
+	.VAlign(VAlign_Center)
 	[
-		SNew(SBox)
-		.Padding(FMargin(0,0,12,0)) // Line up with the other properties due to having no reset to default button
-		[
-			SNew(SEditableTextBox)
-			.IsReadOnly(true)
-			.Font(DetailBuilder.GetDetailFont())
-			.Text(this, &FLandscapeEditorDetailCustomization_NewLandscape::GetTotalComponentCount)
-		]
+		SNew(SEditableTextBox)
+		.IsReadOnly(true)
+		.Font(DetailBuilder.GetDetailFont())
+		.Text(this, &FLandscapeEditorDetailCustomization_NewLandscape::GetTotalComponentCount)
 	];
 
 	NewLandscapeCategory.AddCustomRow(FText::GetEmpty())
diff --git a/Engine/Source/Editor/PropertyEditor/Private/UserInterface/PropertyEditor/SPropertyEditorText.cpp b/Engine/Source/Editor/PropertyEditor/Private/UserInterface/PropertyEditor/SPropertyEditorText.cpp
--- a/Engine/Source/Editor/PropertyEditor/Private/UserInterface/PropertyEditor/SPropertyEditorText.cpp
+++ b/Engine/Source/Editor/PropertyEditor/Private/UserInterface/PropertyEditor/SPropertyEditorText.cpp
@@ -92,6 +92,13 @@ void SPropertyEditorText::GetDesiredWidth( float& OutMinDesiredWidth, float& Out
 	OutMaxDesiredWidth = 600.0f;
 }
 
+FVector2D SPropertyEditorText::ComputeDesiredSize(float Scale) const
+{
+	FVector2D LocalDesiredSize = SCompoundWidget::ComputeDesiredSize(Scale);
+
+	return FVector2D(LocalDesiredSize.X, 22.0f);
+}
+
 bool SPropertyEditorText::Supports( const TSharedRef< FPropertyEditor >& InPropertyEditor )
 {
 	const TSharedRef< FPropertyNode > PropertyNode = InPropertyEditor->GetPropertyNode();
diff --git a/Engine/Source/Editor/PropertyEditor/Private/UserInterface/PropertyEditor/SPropertyEditorText.h b/Engine/Source/Editor/PropertyEditor/Private/UserInterface/PropertyEditor/SPropertyEditorText.h
--- a/Engine/Source/Editor/PropertyEditor/Private/UserInterface/PropertyEditor/SPropertyEditorText.h
+++ b/Engine/Source/Editor/PropertyEditor/Private/UserInterface/PropertyEditor/SPropertyEditorText.h
@@ -30,6 +30,8 @@ public:
 
 	void GetDesiredWidth( float& OutMinDesiredWidth, float& OutMaxDesiredWidth );
 
+	FVector2D ComputeDesiredSize(float Scale) const;
+
 	bool SupportsKeyboardFocus() const override;
 
 	FReply OnFocusReceived( const FGeometry& MyGeometry, const FFocusEvent& InFocusEvent ) override;
diff --git a/Engine/Source/Runtime/Slate/Public/Widgets/Input/SNumericEntryBox.h b/Engine/Source/Runtime/Slate/Public/Widgets/Input/SNumericEntryBox.h
--- a/Engine/Source/Runtime/Slate/Public/Widgets/Input/SNumericEntryBox.h
+++ b/Engine/Source/Runtime/Slate/Public/Widgets/Input/SNumericEntryBox.h
@@ -261,22 +261,32 @@ public:
 				EditableText.ToSharedRef()
 			];
 
+		TSharedPtr<SWidget> MainContents;
 
-		if (InArgs._Label.Widget == SNullWidget::NullWidget || InArgs._LabelLocation == ELabelLocation::Inside)
+		const bool bHasLabel = InArgs._Label.Widget != SNullWidget::NullWidget;
+
+		bool bHasInsideLabel = false;
+		if (bHasLabel && InArgs._LabelLocation == ELabelLocation::Inside)
 		{
-			if(InArgs._Label.Widget != SNullWidget::NullWidget)
-			{
-				Overlay->AddSlot()
-					.HAlign(HAlign_Left)
-					.VAlign(InArgs._LabelVAlign)
-					.Padding(InArgs._LabelPadding)
-					[
-						InArgs._Label.Widget
-					];
-			}
+			bHasInsideLabel = true;
 
-			ChildSlot
-			[
+			Overlay->AddSlot()
+				.HAlign(HAlign_Left)
+				.VAlign(InArgs._LabelVAlign)
+				.Padding(InArgs._LabelPadding)
+				[
+					InArgs._Label.Widget
+				];
+		}
+
+
+		if (bAllowSpin && !bHasInsideLabel)
+		{
+			MainContents = Overlay;
+		}
+		else 
+		{
+			MainContents =
 				SNew(SBorder)
 				.BorderImage(this, &SNumericEntryBox<NumericType>::GetBorderImage)
 				.BorderBackgroundColor(InArgs._BorderBackgroundColor)
@@ -284,7 +294,14 @@ public:
 				.Padding(0.f)
 				[
 					Overlay
-				]
+				];
+		}
+
+		if(!bHasLabel || bHasInsideLabel)
+		{
+			ChildSlot
+			[
+				MainContents.ToSharedRef()
 			];
 		}
 		else
@@ -302,14 +319,7 @@ public:
 				]
 				+ SHorizontalBox::Slot()
 				[
-					SNew(SBorder)
-					.BorderImage(this, &SNumericEntryBox<NumericType>::GetBorderImage)
-					.BorderBackgroundColor(InArgs._BorderBackgroundColor)
-					.ForegroundColor(InArgs._BorderForegroundColor)
-					.Padding(0.f)
-					[
-						Overlay
-					]
+					MainContents.ToSharedRef()
 				]
 			];
 		}
diff --git a/Engine/Source/Runtime/Slate/Public/Widgets/Input/SSpinBox.h b/Engine/Source/Runtime/Slate/Public/Widgets/Input/SSpinBox.h
--- a/Engine/Source/Runtime/Slate/Public/Widgets/Input/SSpinBox.h
+++ b/Engine/Source/Runtime/Slate/Public/Widgets/Input/SSpinBox.h
@@ -294,7 +294,7 @@ public:
 	virtual int32 OnPaint(const FPaintArgs& Args, const FGeometry& AllottedGeometry, const FSlateRect& MyCullingRect, FSlateWindowElementList& OutDrawElements, int32 LayerId, const FWidgetStyle& InWidgetStyle, bool bParentEnabled) const override
 	{
 
-		const bool bActiveFeedback = bDragging;
+		const bool bActiveFeedback = bDragging || IsInTextMode();
 
 		const FSlateBrush* BackgroundImage = bActiveFeedback ? BackgroundActiveBrush : IsHovered() ?
 			BackgroundHoveredBrush :
