Fixed handling of first and last boxes in multi-line SSegmentedControls. Previously the end cap was never displayed, because the column index was being compared to the slot index.

diff --git a/Engine/Source/Runtime/Slate/Public/Widgets/Input/SSegmentedControl.h b/Engine/Source/Runtime/Slate/Public/Widgets/Input/SSegmentedControl.h
--- a/Engine/Source/Runtime/Slate/Public/Widgets/Input/SSegmentedControl.h
+++ b/Engine/Source/Runtime/Slate/Public/Widgets/Input/SSegmentedControl.h
@@ -17,8 +17,8 @@
  * Slots require a templated value to return when the segment is selected by the user.
  * Users can specify text, icon or provide custom content to each Segment.
  *
- * Note: It is currently not possible to add segments after initialization 
- *  (i.e. there is no AddSlot). 
+ * Note: It is currently not possible to add segments after initialization
+ *  (i.e. there is no AddSlot).
  */
 
 template< typename OptionType >
@@ -223,16 +223,28 @@ public:
 				];
 			}
 
-			int32 ColumnIndex = MaxSegmentsPerLine ? SlotIndex % MaxSegmentsPerLine : SlotIndex;
-			UniformBox->AddSlot(ColumnIndex, MaxSegmentsPerLine > 0 ? SlotIndex / MaxSegmentsPerLine : 0)
+			const FCheckBoxStyle* CheckBoxStyle = &Style->ControlStyle;
+			if (SlotIndex == 0)
+			{
+				CheckBoxStyle = &Style->FirstControlStyle;
+			}
+			else if (SlotIndex == NumSlots - 1)
+			{
+				CheckBoxStyle = &Style->LastControlStyle;
+			}
+
+			const int32 ColumnIndex = MaxSegmentsPerLine ? SlotIndex % MaxSegmentsPerLine : SlotIndex;
+			const int32 RowIndex = MaxSegmentsPerLine ? SlotIndex / MaxSegmentsPerLine : 0;
+
 			// Note HAlignment is applied at the check box level because if it were applied here it would make the slots look physically disconnected from each other 
+			UniformBox->AddSlot(ColumnIndex, RowIndex)
 			.VAlign(ChildSlotPtr->GetVerticalAlignment())
 			[
 				SAssignNew(ChildSlotPtr->_CheckBox, SCheckBox)
 				.Clipping(EWidgetClipping::ClipToBounds)
 				.HAlign(ChildSlotPtr->GetHorizontalAlignment())
 				.ToolTipText(ChildSlotPtr->_Tooltip)
-				.Style(ColumnIndex == 0 ? &Style->FirstControlStyle : ColumnIndex == (NumSlots - 1) ? &Style->LastControlStyle : &Style->ControlStyle)
+				.Style(CheckBoxStyle)
 				.IsChecked(GetCheckBoxStateAttribute(ChildValue))
 				.OnCheckStateChanged(this, &SSegmentedControl::CommitValue, ChildValue)
 				.Padding(UniformPadding)
