Fix Slider and text selections offset at non-standard viewport resolutions

diff --git a/Engine/Source/Runtime/Engine/Private/Slate/SceneViewport.cpp b/Engine/Source/Runtime/Engine/Private/Slate/SceneViewport.cpp
--- a/Engine/Source/Runtime/Engine/Private/Slate/SceneViewport.cpp
+++ b/Engine/Source/Runtime/Engine/Private/Slate/SceneViewport.cpp
@@ -534,6 +534,20 @@ FReply FSceneViewport::AcquireFocusAndCapture(FIntPoint MousePosition, EFocusCau
 		{
 			bCursorHiddenDueToCapture = true;
 			MousePosBeforeHiddenDueToCapture = MousePosition;
+			
+			// The slate app will correct mouse positions for non-standard screen / viewport resolution combos
+			// We want to save the mouse position pre-correction so it isn't applied twice when restoring mouse position
+			TSharedPtr<SWindow> Window = FSlateApplication::Get().FindWidgetWindow(ViewportWidgetRef);
+			if (FSlateApplication::Get().GetTransformFullscreenMouseInput() && !GIsEditor && Window.IsValid() && Window->GetWindowMode() == EWindowMode::Fullscreen)
+			{
+				FDisplayMetrics CachedDisplayMetrics;
+				FSlateApplication::Get().GetCachedDisplayMetrics(CachedDisplayMetrics);
+				FVector2D WindowSize = Window->GetSizeInScreen();
+				FVector2D DisplaySize = { (float)CachedDisplayMetrics.PrimaryDisplayWidth, (float)CachedDisplayMetrics.PrimaryDisplayHeight };
+				FVector2D CorrectionScale = DisplaySize / WindowSize;
+				MousePosBeforeHiddenDueToCapture = { (int32)(MousePosition.X * CorrectionScale.X), (int32)(MousePosition.Y * CorrectionScale.Y) };
+			}
+
 		}
 
 		if ( bCursorHiddenDueToCapture || !bShouldShowMouseCursor )
diff --git a/Engine/Source/Runtime/Slate/Private/Framework/Application/SlateApplication.cpp b/Engine/Source/Runtime/Slate/Private/Framework/Application/SlateApplication.cpp
--- a/Engine/Source/Runtime/Slate/Private/Framework/Application/SlateApplication.cpp
+++ b/Engine/Source/Runtime/Slate/Private/Framework/Application/SlateApplication.cpp
@@ -3107,6 +3107,8 @@ void FSlateApplication::ProcessReply( const FWidgetPath& CurrentEventPath, const
 		check(WidgetsUnderMouse);
 		check(InMouseEvent);
 
+		FPointerEvent TransformedPointerEvent = TransformPointerEvent(*InMouseEvent, WidgetsUnderMouse->GetWindow());
+
 		SlateUser->SetDragDropContent(ReplyDragDropContent.ToSharedRef());
 		
 		const FWeakWidgetPath LastWidgetsUnderCursor = SlateUser->GetLastWidgetsUnderPointer(PointerIndex);
@@ -3132,7 +3134,7 @@ void FSlateApplication::ProcessReply( const FWidgetPath& CurrentEventPath, const
 		// To process the beginning of the drag operation, widgets previously
 		// under the mouse cursor receive the OnMouseLeave notification,
 		// regardless of whether the cursor is still over them or not.
-		FEventRouter::Route<FNoReply>(this, FEventRouter::FBubblePolicy(LastWidgetsUnderCursor.ToWidgetPath()), *InMouseEvent, [](const FArrangedWidget& SomeWidget, const FPointerEvent& PointerEvent)
+		FEventRouter::Route<FNoReply>(this, FEventRouter::FBubblePolicy(LastWidgetsUnderCursor.ToWidgetPath()), TransformedPointerEvent, [](const FArrangedWidget& SomeWidget, const FPointerEvent& PointerEvent)
 		{
 			SomeWidget.Widget->OnMouseLeave( PointerEvent );
 #if WITH_SLATE_DEBUGGING
@@ -3142,7 +3144,7 @@ void FSlateApplication::ProcessReply( const FWidgetPath& CurrentEventPath, const
 		}, ESlateDebuggingInputEvent::MouseLeave);
 
 		// Then, the original widget started the drag receives OnDragEnter.
-		FEventRouter::Route<FNoReply>(this, FEventRouter::FBubblePolicy(CurrentEventPath), FDragDropEvent( *InMouseEvent, ReplyDragDropContent ), [](const FArrangedWidget& SomeWidget, const FDragDropEvent& DragDropEvent )
+		FEventRouter::Route<FNoReply>(this, FEventRouter::FBubblePolicy(CurrentEventPath), FDragDropEvent( TransformedPointerEvent, ReplyDragDropContent ), [](const FArrangedWidget& SomeWidget, const FDragDropEvent& DragDropEvent )
 		{
 			SomeWidget.Widget->OnDragEnter( SomeWidget.Geometry, DragDropEvent );
 #if WITH_SLATE_DEBUGGING
@@ -3196,8 +3198,10 @@ void FSlateApplication::ProcessReply( const FWidgetPath& CurrentEventPath, const
 									// It's possible for mouse event to be null if we end up here from a keyboard event. If so, we should synthesize an event
 									if (InMouseEvent)
 									{
+										FPointerEvent TransformedPointerEvent = TransformPointerEvent(*InMouseEvent, WidgetsUnderMouse->GetWindow());
+
 										// Note that the event's pointer position is not translated.
-										WidgetPreviouslyUnderCursor->OnMouseLeave(*InMouseEvent);
+										WidgetPreviouslyUnderCursor->OnMouseLeave(TransformedPointerEvent);
 									}
 									else
 									{
@@ -3338,11 +3342,13 @@ void FSlateApplication::ProcessReply( const FWidgetPath& CurrentEventPath, const
 	{
 		checkSlow(InMouseEvent);
 
+		FPointerEvent TransformedPointerEvent = TransformPointerEvent(*InMouseEvent, WidgetsUnderMouse->GetWindow());
+
 		SlateUser->StartDragDetection(
 			WidgetsUnderMouse->GetPathDownTo(TheReply.GetDetectDragRequest().ToSharedRef()),
-			InMouseEvent->GetPointerIndex(),
+			TransformedPointerEvent.GetPointerIndex(),
 			TheReply.GetDetectDragRequestButton(),
-			InMouseEvent->GetScreenSpacePosition());
+			TransformedPointerEvent.GetScreenSpacePosition());
 	}
 
 	// Set focus if requested.
@@ -4077,6 +4083,25 @@ bool FSlateApplication::HasUserMouseCapture(int32 UserIndex) const
 	return FoundUser && FoundUser->HasAnyCapture();
 }
 
+FPointerEvent FSlateApplication::TransformPointerEvent(const FPointerEvent& PointerEvent, const TSharedPtr<SWindow>& Window) const
+{
+	FPointerEvent TransformedPointerEvent = PointerEvent;
+	if (Window)
+	{
+		if (TransformFullscreenMouseInput && !GIsEditor && Window->GetWindowMode() == EWindowMode::Fullscreen)
+		{
+			// Screen space mapping scales everything. When window resolution doesn't match platform resolution, 
+			// this causes offset cursor hit-tests in fullscreen. Correct in slate since we are first window-aware slate processor.
+			FVector2D WindowSize = Window->GetSizeInScreen();
+			FVector2D DisplaySize = { (float)CachedDisplayMetrics.PrimaryDisplayWidth, (float)CachedDisplayMetrics.PrimaryDisplayHeight };
+
+			TransformedPointerEvent = FPointerEvent(PointerEvent, PointerEvent.GetScreenSpacePosition() * WindowSize / DisplaySize, PointerEvent.GetLastScreenSpacePosition() * WindowSize / DisplaySize);
+		}
+	}
+
+	return TransformedPointerEvent;
+}
+
 bool FSlateApplication::DoesWidgetHaveMouseCaptureByUser(const TSharedPtr<const SWidget> Widget, int32 UserIndex, TOptional<int32> PointerIndex) const
 {
 	if (TSharedPtr<const FSlateUser> FoundUser = GetUser(UserIndex))
@@ -4668,7 +4693,9 @@ bool FSlateApplication::ProcessMouseButtonDownEvent( const TSharedPtr< FGenericW
 			FScopedSwitchWorldHack SwitchWorld(MouseCaptorPath);
 			bInGame = FApp::IsGame();
 
-			Reply = FEventRouter::Route<FReply>(this, FEventRouter::FToLeafmostPolicy(MouseCaptorPath), MouseEvent, [] (const FArrangedWidget& InMouseCaptorWidget, const FPointerEvent& Event)
+			FPointerEvent TransformedPointerEvent = TransformPointerEvent(MouseEvent, MouseCaptorPath.GetWindow());
+
+			Reply = FEventRouter::Route<FReply>(this, FEventRouter::FToLeafmostPolicy(MouseCaptorPath), TransformedPointerEvent, [] (const FArrangedWidget& InMouseCaptorWidget, const FPointerEvent& Event)
 			{
 				const FReply TempReply = InMouseCaptorWidget.Widget->OnPreviewMouseButtonDown(InMouseCaptorWidget.Geometry, Event);
 #if WITH_SLATE_DEBUGGING
@@ -4679,7 +4706,7 @@ bool FSlateApplication::ProcessMouseButtonDownEvent( const TSharedPtr< FGenericW
 
 			if ( !Reply.IsEventHandled() )
 			{
-				Reply = FEventRouter::Route<FReply>(this, FEventRouter::FToLeafmostPolicy(MouseCaptorPath), MouseEvent,
+				Reply = FEventRouter::Route<FReply>(this, FEventRouter::FToLeafmostPolicy(MouseCaptorPath), TransformedPointerEvent,
 					[this] (const FArrangedWidget& InMouseCaptorWidget, const FPointerEvent& Event)
 				{
 					FReply TempReply = FReply::Unhandled();
@@ -4738,6 +4765,8 @@ bool FSlateApplication::ProcessMouseButtonDownEvent( const TSharedPtr< FGenericW
 FReply FSlateApplication::RoutePointerDownEvent(const FWidgetPath& WidgetsUnderPointer, const FPointerEvent& PointerEvent)
 {
 	TScopeCounter<int32> BeginInput(ProcessingInput);
+
+	FPointerEvent TransformedPointerEvent = WidgetsUnderPointer.IsValid() ? TransformPointerEvent(PointerEvent, WidgetsUnderPointer.GetWindow()) : PointerEvent;
 	
 	TSharedRef<FSlateUser> SlateUser = GetOrCreateUser(PointerEvent);
 	SlateUser->UpdatePointerPosition(PointerEvent);
@@ -4751,7 +4780,7 @@ FReply FSlateApplication::RoutePointerDownEvent(const FWidgetPath& WidgetsUnderP
 
 	const TSharedPtr<SWidget> PreviouslyFocusedWidget = GetKeyboardFocusedWidget();
 
-	FReply Reply = FEventRouter::Route<FReply>( this, FEventRouter::FTunnelPolicy( WidgetsUnderPointer ), PointerEvent, []( const FArrangedWidget TargetWidget, const FPointerEvent& Event )
+	FReply Reply = FEventRouter::Route<FReply>( this, FEventRouter::FTunnelPolicy( WidgetsUnderPointer ), TransformedPointerEvent, []( const FArrangedWidget TargetWidget, const FPointerEvent& Event )
 	{
 		const FReply TempReply = TargetWidget.Widget->OnPreviewMouseButtonDown(TargetWidget.Geometry, Event);
 #if WITH_SLATE_DEBUGGING
@@ -4762,7 +4791,7 @@ FReply FSlateApplication::RoutePointerDownEvent(const FWidgetPath& WidgetsUnderP
 
 	if( !Reply.IsEventHandled() )
 	{
-		Reply = FEventRouter::Route<FReply>( this, FEventRouter::FBubblePolicy( WidgetsUnderPointer ), PointerEvent, [this]( const FArrangedWidget TargetWidget, const FPointerEvent& Event )
+		Reply = FEventRouter::Route<FReply>( this, FEventRouter::FBubblePolicy( WidgetsUnderPointer ), TransformedPointerEvent, [this]( const FArrangedWidget TargetWidget, const FPointerEvent& Event )
 		{
 			FReply TempReply = FReply::Unhandled();
 			if( !TempReply.IsEventHandled() )
@@ -4792,7 +4821,7 @@ FReply FSlateApplication::RoutePointerDownEvent(const FWidgetPath& WidgetsUnderP
 			{
 				const FArrangedWidget& TargetWidget = WidgetsUnderPointer.Widgets[WidgetIndex];
 
-				TargetWidget.Widget->OnMouseEnter(TargetWidget.Geometry, PointerEvent);
+				TargetWidget.Widget->OnMouseEnter(TargetWidget.Geometry, TransformedPointerEvent);
 #if WITH_SLATE_DEBUGGING
 				FSlateDebugging::BroadcastInputEvent(ESlateDebuggingInputEvent::MouseEnter, TargetWidget.Widget);
 #endif
@@ -4860,6 +4889,8 @@ FReply FSlateApplication::RoutePointerUpEvent(const FWidgetPath& WidgetsUnderPoi
 {
 	TScopeCounter<int32> BeginInput(ProcessingInput);
 
+	FPointerEvent TransformedPointerEvent = WidgetsUnderPointer.IsValid() ? TransformPointerEvent(PointerEvent, WidgetsUnderPointer.GetWindow()) : PointerEvent;
+
 	FReply Reply = FReply::Unhandled();
 	TSharedRef<FSlateUser> SlateUser = GetOrCreateUser(PointerEvent);
 	const bool bIsDragDropping = SlateUser->IsDragDroppingAffected(PointerEvent);
@@ -4877,7 +4908,7 @@ FReply FSlateApplication::RoutePointerUpEvent(const FWidgetPath& WidgetsUnderPoi
 			FScopedSwitchWorldHack SwitchWorld( MouseCaptorPath );
 
 			Reply =
-				FEventRouter::Route<FReply>( this, FEventRouter::FToLeafmostPolicy(MouseCaptorPath), PointerEvent, [this]( const FArrangedWidget& TargetWidget, const FPointerEvent& Event )
+				FEventRouter::Route<FReply>( this, FEventRouter::FToLeafmostPolicy(MouseCaptorPath), TransformedPointerEvent, [this]( const FArrangedWidget& TargetWidget, const FPointerEvent& Event )
 				{
 					FReply TempReply = FReply::Unhandled();
 					if (Event.IsTouchEvent())
@@ -4927,7 +4958,7 @@ FReply FSlateApplication::RoutePointerUpEvent(const FWidgetPath& WidgetsUnderPoi
 			SlateUser->ResetDragDropContent();
 		}
 
-		Reply = FEventRouter::Route<FReply>(this, FEventRouter::FBubblePolicy(LocalWidgetsUnderPointer), PointerEvent, [&](const FArrangedWidget& CurWidget, const FPointerEvent& Event)
+		Reply = FEventRouter::Route<FReply>(this, FEventRouter::FBubblePolicy(LocalWidgetsUnderPointer), TransformedPointerEvent, [&](const FArrangedWidget& CurWidget, const FPointerEvent& Event)
 		{
 			if (bIsDragDropping)
 			{
@@ -4993,6 +5024,8 @@ bool FSlateApplication::RoutePointerMoveEvent(const FWidgetPath& WidgetsUnderPoi
 	bool bHandled = false;
 	FWeakWidgetPath LastWidgetsUnderPointer;
 
+	FPointerEvent TransformedPointerEvent = WidgetsUnderPointer.IsValid() ? TransformPointerEvent(PointerEvent, WidgetsUnderPointer.GetWindow()) : PointerEvent;
+
 	TSharedRef<FSlateUser> SlateUser = GetOrCreateUser(PointerEvent);
 	SlateUser->NotifyPointerMoveBegin(PointerEvent);
 
@@ -5107,7 +5140,7 @@ bool FSlateApplication::RoutePointerMoveEvent(const FWidgetPath& WidgetsUnderPoi
 		// Switch worlds widgets in the current path
 		FScopedSwitchWorldHack SwitchWorld(MouseCaptorPath);
 
-		FEventRouter::Route<FNoReply>(this, FEventRouter::FBubblePolicy(WidgetsUnderPointer), PointerEvent, [&MouseCaptorPath, &LastWidgetsUnderPointer](const FArrangedWidget& WidgetUnderCursor, const FPointerEvent& Event)
+		FEventRouter::Route<FNoReply>(this, FEventRouter::FBubblePolicy(WidgetsUnderPointer), TransformedPointerEvent, [&MouseCaptorPath, &LastWidgetsUnderPointer](const FArrangedWidget& WidgetUnderCursor, const FPointerEvent& Event)
 			{
 				if (!LastWidgetsUnderPointer.ContainsWidget(WidgetUnderCursor.GetWidgetPtr()))
 				{
@@ -5122,7 +5155,7 @@ bool FSlateApplication::RoutePointerMoveEvent(const FWidgetPath& WidgetsUnderPoi
 				return FNoReply();
 			}, ESlateDebuggingInputEvent::MouseEnter);
 
-		FReply Reply = FEventRouter::Route<FReply>(this, FEventRouter::FToLeafmostPolicy(MouseCaptorPath), PointerEvent, [this, bIsSynthetic](const FArrangedWidget& MouseCaptorWidget, const FPointerEvent& Event)
+		FReply Reply = FEventRouter::Route<FReply>(this, FEventRouter::FToLeafmostPolicy(MouseCaptorPath), TransformedPointerEvent, [this, bIsSynthetic](const FArrangedWidget& MouseCaptorWidget, const FPointerEvent& Event)
 			{
 				FReply TempReply = FReply::Unhandled();
 
@@ -5193,7 +5226,7 @@ bool FSlateApplication::RoutePointerMoveEvent(const FWidgetPath& WidgetsUnderPoi
 		}
 		else
 		{
-			FEventRouter::Route<FNoReply>(this, FEventRouter::FBubblePolicy(WidgetsUnderPointer), PointerEvent, [&LastWidgetsUnderPointer](const FArrangedWidget& WidgetUnderCursor, const FPointerEvent& Event)
+			FEventRouter::Route<FNoReply>(this, FEventRouter::FBubblePolicy(WidgetsUnderPointer), TransformedPointerEvent, [&LastWidgetsUnderPointer](const FArrangedWidget& WidgetUnderCursor, const FPointerEvent& Event)
 			{
 				if (!LastWidgetsUnderPointer.ContainsWidget(WidgetUnderCursor.GetWidgetPtr()))
 				{
@@ -5207,7 +5240,7 @@ bool FSlateApplication::RoutePointerMoveEvent(const FWidgetPath& WidgetsUnderPoi
 		}
 
 		// Bubble the MouseMove event.
-		FReply Reply = FEventRouter::Route<FReply>(this, FEventRouter::FBubblePolicy(WidgetsUnderPointer), PointerEvent, [&](const FArrangedWidget& CurWidget, const FPointerEvent& Event)
+		FReply Reply = FEventRouter::Route<FReply>(this, FEventRouter::FBubblePolicy(WidgetsUnderPointer), TransformedPointerEvent, [&](const FArrangedWidget& CurWidget, const FPointerEvent& Event)
 		{
 			FReply TempReply = FReply::Unhandled();
 
@@ -5343,7 +5376,9 @@ FReply FSlateApplication::RoutePointerDoubleClickEvent(const FWidgetPath& Widget
 	// Switch worlds widgets in the current path
 	FScopedSwitchWorldHack SwitchWorld( WidgetsUnderPointer );
 
-	Reply = FEventRouter::Route<FReply>( this, FEventRouter::FBubblePolicy( WidgetsUnderPointer ), PointerEvent, []( const FArrangedWidget& TargetWidget, const FPointerEvent& Event )
+	FPointerEvent TransformedPointerEvent = WidgetsUnderPointer.IsValid() ? TransformPointerEvent(PointerEvent, WidgetsUnderPointer.GetWindow()) : PointerEvent;
+
+	Reply = FEventRouter::Route<FReply>( this, FEventRouter::FBubblePolicy( WidgetsUnderPointer ), TransformedPointerEvent, []( const FArrangedWidget& TargetWidget, const FPointerEvent& Event )
 	{
 		const FReply TempReply = TargetWidget.Widget->OnMouseButtonDoubleClick(TargetWidget.Geometry, Event);
 #if WITH_SLATE_DEBUGGING
diff --git a/Engine/Source/Runtime/Slate/Private/Widgets/SViewport.cpp b/Engine/Source/Runtime/Slate/Private/Widgets/SViewport.cpp
--- a/Engine/Source/Runtime/Slate/Private/Widgets/SViewport.cpp
+++ b/Engine/Source/Runtime/Slate/Private/Widgets/SViewport.cpp
@@ -102,21 +102,6 @@ void SViewport::SetActive(bool bActive)
 	}
 }
 
-FVector2D SViewport::TransformToViewport(const FVector2D& CursorPosition) const
-{
-	// Screen space mapping scales everything. When viewport resolution doesn't match platform resolution, 
-	// this causes offset cursor hit-tests in fullscreen. Correct in slate since we are first viewport-aware processor.
-	if (FSlateApplication::Get().GetTransformFullscreenMouseInput() && !GIsEditor && CachedParentWindow.IsValid() && CachedParentWindow.Pin()->GetWindowMode() == EWindowMode::Fullscreen)
-	{
-		FDisplayMetrics CachedDisplayMetrics;
-		FSlateApplication::Get().GetCachedDisplayMetrics(CachedDisplayMetrics);
-		FVector2D DisplaySize = { (float)CachedDisplayMetrics.PrimaryDisplayWidth, (float)CachedDisplayMetrics.PrimaryDisplayHeight };
-		return CursorPosition * GetTickSpaceGeometry().GetAbsoluteSize() / DisplaySize;
-	}
-
-	return CursorPosition;
-}
-
 EActiveTimerReturnType SViewport::EnsureTick(double InCurrentTime, float InDeltaTime)
 {
 	return EActiveTimerReturnType::Continue;
@@ -243,11 +228,7 @@ void SViewport::Tick( const FGeometry& AllottedGeometry, const double InCurrentT
 
 FCursorReply SViewport::OnCursorQuery( const FGeometry& MyGeometry, const FPointerEvent& CursorEvent ) const
 {
-	FPointerEvent ViewportOffsetMouseEvent(
-		CursorEvent,
-		TransformToViewport(CursorEvent.GetScreenSpacePosition()),
-		TransformToViewport(CursorEvent.GetLastScreenSpacePosition()));
-	return ViewportInterface.IsValid() ? ViewportInterface.Pin()->OnCursorQuery(MyGeometry, ViewportOffsetMouseEvent) : FCursorReply::Unhandled();
+	return ViewportInterface.IsValid() ? ViewportInterface.Pin()->OnCursorQuery(MyGeometry, CursorEvent) : FCursorReply::Unhandled();
 }
 
 TOptional<TSharedRef<SWidget>> SViewport::OnMapCursor(const FCursorReply& CursorReply) const
@@ -257,20 +238,12 @@ TOptional<TSharedRef<SWidget>> SViewport::OnMapCursor(const FCursorReply& Cursor
 
 FReply SViewport::OnMouseButtonDown( const FGeometry& MyGeometry, const FPointerEvent& MouseEvent )
 {
-	FPointerEvent ViewportOffsetMouseEvent(
-		MouseEvent, 
-		TransformToViewport(MouseEvent.GetScreenSpacePosition()), 
-		TransformToViewport(MouseEvent.GetLastScreenSpacePosition()));
-	return ViewportInterface.IsValid() ? ViewportInterface.Pin()->OnMouseButtonDown(MyGeometry, ViewportOffsetMouseEvent) : FReply::Unhandled();
+	return ViewportInterface.IsValid() ? ViewportInterface.Pin()->OnMouseButtonDown(MyGeometry, MouseEvent) : FReply::Unhandled();
 }
 
 FReply SViewport::OnMouseButtonUp( const FGeometry& MyGeometry, const FPointerEvent& MouseEvent )
 {
-	FPointerEvent ViewportOffsetMouseEvent(
-		MouseEvent,
-		TransformToViewport(MouseEvent.GetScreenSpacePosition()),
-		TransformToViewport(MouseEvent.GetLastScreenSpacePosition()));
-	return ViewportInterface.IsValid() ? ViewportInterface.Pin()->OnMouseButtonUp(MyGeometry, ViewportOffsetMouseEvent) : FReply::Unhandled();
+	return ViewportInterface.IsValid() ? ViewportInterface.Pin()->OnMouseButtonUp(MyGeometry, MouseEvent) : FReply::Unhandled();
 }
 
 void SViewport::OnMouseEnter( const FGeometry& MyGeometry, const FPointerEvent& MouseEvent )
@@ -279,11 +252,7 @@ void SViewport::OnMouseEnter( const FGeometry& MyGeometry, const FPointerEvent&
 
 	if (ViewportInterface.IsValid())
 	{
-		FPointerEvent ViewportOffsetMouseEvent(
-			MouseEvent,
-			TransformToViewport(MouseEvent.GetScreenSpacePosition()),
-			TransformToViewport(MouseEvent.GetLastScreenSpacePosition()));
-		ViewportInterface.Pin()->OnMouseEnter(MyGeometry, ViewportOffsetMouseEvent);
+		ViewportInterface.Pin()->OnMouseEnter(MyGeometry, MouseEvent);
 	}
 }
 
@@ -293,39 +262,23 @@ void SViewport::OnMouseLeave( const FPointerEvent& MouseEvent )
 
 	if (ViewportInterface.IsValid())
 	{
-		FPointerEvent ViewportOffsetMouseEvent(
-			MouseEvent,
-			TransformToViewport(MouseEvent.GetScreenSpacePosition()),
-			TransformToViewport(MouseEvent.GetLastScreenSpacePosition()));
-		ViewportInterface.Pin()->OnMouseLeave(ViewportOffsetMouseEvent);
+		ViewportInterface.Pin()->OnMouseLeave(MouseEvent);
 	}
 }
 
 FReply SViewport::OnMouseMove( const FGeometry& MyGeometry, const FPointerEvent& MouseEvent )
 {
-	FPointerEvent ViewportOffsetMouseEvent(
-		MouseEvent,
-		TransformToViewport(MouseEvent.GetScreenSpacePosition()),
-		TransformToViewport(MouseEvent.GetLastScreenSpacePosition()));
-	return ViewportInterface.IsValid() ? ViewportInterface.Pin()->OnMouseMove(MyGeometry, ViewportOffsetMouseEvent) : FReply::Unhandled();
+	return ViewportInterface.IsValid() ? ViewportInterface.Pin()->OnMouseMove(MyGeometry, MouseEvent) : FReply::Unhandled();
 }
 
 FReply SViewport::OnMouseWheel( const FGeometry& MyGeometry, const FPointerEvent& MouseEvent )
 {
-	FPointerEvent ViewportOffsetMouseEvent(
-		MouseEvent,
-		TransformToViewport(MouseEvent.GetScreenSpacePosition()),
-		TransformToViewport(MouseEvent.GetLastScreenSpacePosition()));
-	return ViewportInterface.IsValid() ? ViewportInterface.Pin()->OnMouseWheel(MyGeometry, ViewportOffsetMouseEvent) : FReply::Unhandled();
+	return ViewportInterface.IsValid() ? ViewportInterface.Pin()->OnMouseWheel(MyGeometry, MouseEvent) : FReply::Unhandled();
 }
 
 FReply SViewport::OnMouseButtonDoubleClick( const FGeometry& MyGeometry, const FPointerEvent& MouseEvent )
 {
-	FPointerEvent ViewportOffsetMouseEvent(
-		MouseEvent,
-		TransformToViewport(MouseEvent.GetScreenSpacePosition()),
-		TransformToViewport(MouseEvent.GetLastScreenSpacePosition()));
-	return ViewportInterface.IsValid() ? ViewportInterface.Pin()->OnMouseButtonDoubleClick(MyGeometry, ViewportOffsetMouseEvent) : FReply::Unhandled();
+	return ViewportInterface.IsValid() ? ViewportInterface.Pin()->OnMouseButtonDoubleClick(MyGeometry, MouseEvent) : FReply::Unhandled();
 }
 
 FReply SViewport::OnKeyDown( const FGeometry& MyGeometry, const FKeyEvent& KeyEvent )
diff --git a/Engine/Source/Runtime/Slate/Public/Framework/Application/SlateApplication.h b/Engine/Source/Runtime/Slate/Public/Framework/Application/SlateApplication.h
--- a/Engine/Source/Runtime/Slate/Public/Framework/Application/SlateApplication.h
+++ b/Engine/Source/Runtime/Slate/Public/Framework/Application/SlateApplication.h
@@ -987,6 +987,9 @@ protected:
 
 	friend class FEventRouter;
 
+	/** Transforms a pointer event to account for non-standard viewport resolutions */
+	FPointerEvent TransformPointerEvent(const FPointerEvent& PointerEvent, const TSharedPtr<SWindow>& Window) const;
+
 	virtual bool DoesWidgetHaveMouseCaptureByUser(const TSharedPtr<const SWidget> Widget, int32 UserIndex, TOptional<int32> PointerIndex) const override;
 	virtual bool DoesWidgetHaveMouseCapture(const TSharedPtr<const SWidget> Widget) const override;
 
diff --git a/Engine/Source/Runtime/Slate/Public/Widgets/SViewport.h b/Engine/Source/Runtime/Slate/Public/Widgets/SViewport.h
--- a/Engine/Source/Runtime/Slate/Public/Widgets/SViewport.h
+++ b/Engine/Source/Runtime/Slate/Public/Widgets/SViewport.h
@@ -201,9 +201,6 @@ public:
 	UE_DEPRECATED(4.11, "SetWidgetToFocusOnActivate is no longer needed, remove this call.")
 	void SetWidgetToFocusOnActivate(TSharedPtr<SWidget> WidgetToFocus) { }
 
-	/** Transform a cursor position to the viewport */
-	FVector2D TransformToViewport(const FVector2D& CursorPosition) const;
-
 public:
 
 	// SWidget interface
