Consecutive separators are now ignored in toolbars

diff --git a/Engine/Source/Editor/EditorStyle/Private/StarshipStyle.cpp b/Engine/Source/Editor/EditorStyle/Private/StarshipStyle.cpp
--- a/Engine/Source/Editor/EditorStyle/Private/StarshipStyle.cpp
+++ b/Engine/Source/Editor/EditorStyle/Private/StarshipStyle.cpp
@@ -5483,7 +5483,7 @@ void FStarshipEditorStyle::FStyle::SetupPersonaStyle()
 
 		MainToolbarRightButton.SetButtonStyle(RightToolbarButton);
 		MainToolbarRightButton.SetButtonPadding(FMargin(0.0f, 0.0f, 4.0f, 0.0f));
-		MainToolbarRightButton.SetSeparatorPadding(FMargin(6.f, 0.f, 8.f, 0.f));
+		MainToolbarRightButton.SetSeparatorPadding(FMargin(4.f, -5.f, 8.f, -5.f));
 
 		Set("Toolbar.BackplateRight", MainToolbarRightButton);
 
diff --git a/Engine/Source/Runtime/Slate/Private/Framework/MultiBox/MultiBox.cpp b/Engine/Source/Runtime/Slate/Private/Framework/MultiBox/MultiBox.cpp
--- a/Engine/Source/Runtime/Slate/Private/Framework/MultiBox/MultiBox.cpp
+++ b/Engine/Source/Runtime/Slate/Private/Framework/MultiBox/MultiBox.cpp
@@ -1075,6 +1075,8 @@ void SMultiBoxWidget::BuildMultiBoxWidget()
 	bool bSectionContainsIcons = false;
 	int32 NextMenuSeparator = INDEX_NONE;
 
+	int32 ConsecutiveSeparatorCount = 0;
+
 	for( int32 Index = 0; Index < Blocks.Num(); Index++ )
 	{
 		// If we've passed the last menu separator, scan for the next one (the end of the list is also considered a menu separator for the purposes of this index)
@@ -1089,7 +1091,7 @@ void SMultiBoxWidget::BuildMultiBoxWidget()
 					bSectionContainsIcons = true;
 				}
 
-				if (TestBlock.GetType() == EMultiBlockType::Separator)
+				if (TestBlock.IsSeparator())
 				{
 					break;
 				}
@@ -1098,7 +1100,21 @@ void SMultiBoxWidget::BuildMultiBoxWidget()
 
 		const FMultiBlock& Block = *Blocks[Index];
 		EMultiBlockLocation::Type Location = EMultiBlockLocation::None;
-		
+	
+		if (Block.IsSeparator())
+		{
+			++ConsecutiveSeparatorCount;
+			// Skip consecutive separators. Only draw one separator no matter how many are submitted.
+			if (ConsecutiveSeparatorCount > 1)
+			{
+				continue;
+			}
+		}
+		else
+		{
+			ConsecutiveSeparatorCount = 0;
+		}
+
 		TSharedPtr<const FMultiBlock> NextBlock = nullptr;
 
 		if (Blocks.IsValidIndex(Index + 1))
diff --git a/Engine/Source/Runtime/SlateCore/Private/Styling/StarshipCoreStyle.cpp b/Engine/Source/Runtime/SlateCore/Private/Styling/StarshipCoreStyle.cpp
--- a/Engine/Source/Runtime/SlateCore/Private/Styling/StarshipCoreStyle.cpp
+++ b/Engine/Source/Runtime/SlateCore/Private/Styling/StarshipCoreStyle.cpp
@@ -1877,7 +1877,7 @@ void FStarshipCoreStyle::SetupMultiboxStyles(TSharedRef<FStyle>& Style)
 			.SetButtonPadding(FMargin(4.0f, 0.0f))
 			.SetCheckBoxPadding(FMargin(8.0f, 0.0f))
 			.SetSeparatorBrush(FSlateColorBrush(FStyleColors::Recessed))
-			.SetSeparatorPadding(FMargin(4.f, -4.0f))
+			.SetSeparatorPadding(FMargin(4.f, -5.0f))
 			.SetLabelStyle(FTextBlockStyle(NormalText))
 			.SetComboButtonStyle(SlimToolBarComboButton)
 			.SetLabelPadding(FMargin(4.f, 0.f, 0.f, 0.f))
