Add an SToolTip delegate to allow users to optionally adjust the desired window location for interactive tooltips.

diff --git a/Engine/Source/Runtime/Slate/Private/Framework/Application/SlateUser.cpp b/Engine/Source/Runtime/Slate/Private/Framework/Application/SlateUser.cpp
--- a/Engine/Source/Runtime/Slate/Private/Framework/Application/SlateUser.cpp
+++ b/Engine/Source/Runtime/Slate/Private/Framework/Application/SlateUser.cpp
@@ -1258,6 +1258,12 @@ void FSlateUser::UpdateTooltip(const FMenuStack& MenuStack, bool bCanSpawnNewToo
 	{
 		// Keep track of where we want tooltips to be positioned
 		DesiredLocation = GetPreviousCursorPosition() + SlateDefs::TooltipOffsetFromMouse;
+
+		// Allow interactive tooltips to adjust the window location
+		if (NewTooltip && NewTooltip->IsInteractive() && !NewTooltipVisualizer.IsValid())
+		{
+			NewTooltip->OnSetInteractiveWindowLocation(DesiredLocation);
+		}
 	}
 
 	if (TooltipWindowPtr.IsValid())
diff --git a/Engine/Source/Runtime/Slate/Private/Widgets/SToolTip.cpp b/Engine/Source/Runtime/Slate/Private/Widgets/SToolTip.cpp
--- a/Engine/Source/Runtime/Slate/Private/Widgets/SToolTip.cpp
+++ b/Engine/Source/Runtime/Slate/Private/Widgets/SToolTip.cpp
@@ -26,6 +26,7 @@ void SToolTip::Construct( const FArguments& InArgs )
 	ColorAndOpacity = InArgs._ColorAndOpacity;
 	TextMargin = InArgs._TextMargin;
 	BorderImage = InArgs._BorderImage;
+	OnSetInteractiveWindowLocationDelegate = InArgs._OnSetInteractiveWindowLocation;
 	
 	SetContentWidget(InArgs._Content.Widget);
 }
@@ -97,3 +98,11 @@ bool SToolTip::IsInteractive() const
 {
 	return bIsInteractive.Get();
 }
+
+void SToolTip::OnSetInteractiveWindowLocation(FVector2D& InOutDesiredLocation) const
+{
+	if (OnSetInteractiveWindowLocationDelegate.IsBound())
+	{
+		OnSetInteractiveWindowLocationDelegate.Execute(InOutDesiredLocation);
+	}
+}
\ No newline at end of file
diff --git a/Engine/Source/Runtime/Slate/Public/Widgets/SToolTip.h b/Engine/Source/Runtime/Slate/Public/Widgets/SToolTip.h
--- a/Engine/Source/Runtime/Slate/Public/Widgets/SToolTip.h
+++ b/Engine/Source/Runtime/Slate/Public/Widgets/SToolTip.h
@@ -21,6 +21,7 @@ class SLATE_API SToolTip
 	, public IToolTip
 {
 public:
+	DECLARE_DELEGATE_OneParam(FOnSetInteractiveWindowLocation, FVector2D&);
 
 	SLATE_BEGIN_ARGS( SToolTip )
 		: _Text()
@@ -30,6 +31,7 @@ public:
 		, _TextMargin(FMargin(8.0f))
 		, _BorderImage(FCoreStyle::Get().GetBrush("ToolTip.Background"))
 		, _IsInteractive(false)
+		, _OnSetInteractiveWindowLocation()
 	{ }
 
 		/** The text displayed in this tool tip */
@@ -53,6 +55,9 @@ public:
 		/** Whether the tooltip should be considered interactive */
 		SLATE_ATTRIBUTE(bool, IsInteractive)
 
+		/** Hook to modify or override the desired location (in screen space) for interactive tooltip windows. By default, the previous frame's cursor position will be used. */
+		SLATE_EVENT(FOnSetInteractiveWindowLocation, OnSetInteractiveWindowLocation)
+
 	SLATE_END_ARGS()
 
 	/**
@@ -83,7 +88,7 @@ public:
 	virtual bool IsInteractive( ) const override;
 	virtual void OnOpening() override { }
 	virtual void OnClosed() override { }
-
+	virtual void OnSetInteractiveWindowLocation(FVector2D& InOutDesiredLocation) const override;
 
 	virtual const FText& GetTextTooltip() const
 	{
@@ -119,4 +124,7 @@ private:
 	
 	// Whether the tooltip should be considered interactive.
 	TAttribute<bool> bIsInteractive;
+
+	// Optional delegate to modify or override the desired location for an interactive tooltip.
+	FOnSetInteractiveWindowLocation OnSetInteractiveWindowLocationDelegate;
 };
diff --git a/Engine/Source/Runtime/SlateCore/Public/Widgets/IToolTip.h b/Engine/Source/Runtime/SlateCore/Public/Widgets/IToolTip.h
--- a/Engine/Source/Runtime/SlateCore/Public/Widgets/IToolTip.h
+++ b/Engine/Source/Runtime/SlateCore/Public/Widgets/IToolTip.h
@@ -62,6 +62,16 @@ public:
 	 */
 	virtual void OnClosed( ) = 0;
 
+	/**
+	 * Allows an interactive tooltip to adjust and/or change its window location on open,
+	 * as interactive tooltips don't update the tooltip window to follow the cursor position.
+	 * 
+	 * @param InOutDesiredLocation	On input, contains the previous cursor position (in screen coordinates) plus the standard
+	 *								tooltip offset (SlateDefs::TooltipOffsetFromMouse). This position can be modified and/or
+	 *								reset to either adjust the offset or otherwise set to a new position as the output value.
+	 */
+	virtual void OnSetInteractiveWindowLocation(FVector2D& InOutDesiredLocation) const {}
+
 public:
 
 	/** Virtual destructor. */
