Changed MaxSegmentsPerLine comparison to > 0 instead of truthy, to avoid situations where negative values cause problems.

diff --git a/Engine/Source/Runtime/Slate/Public/Widgets/Input/SSegmentedControl.h b/Engine/Source/Runtime/Slate/Public/Widgets/Input/SSegmentedControl.h
--- a/Engine/Source/Runtime/Slate/Public/Widgets/Input/SSegmentedControl.h
+++ b/Engine/Source/Runtime/Slate/Public/Widgets/Input/SSegmentedControl.h
@@ -232,8 +232,8 @@ public:
 				CheckBoxStyle = &Style->LastControlStyle;
 			}
 
-			const int32 ColumnIndex = MaxSegmentsPerLine ? SlotIndex % MaxSegmentsPerLine : SlotIndex;
-			const int32 RowIndex = MaxSegmentsPerLine ? SlotIndex / MaxSegmentsPerLine : 0;
+			const int32 ColumnIndex = MaxSegmentsPerLine > 0 ? SlotIndex % MaxSegmentsPerLine : SlotIndex;
+			const int32 RowIndex = MaxSegmentsPerLine > 0 ? SlotIndex / MaxSegmentsPerLine : 0;
 
 			// Note HAlignment is applied at the check box level because if it were applied here it would make the slots look physically disconnected from each other 
 			UniformBox->AddSlot(ColumnIndex, RowIndex)
