PR #8434: Fix mouse not being properly locked to game window (Contributed by [REDACTED_PERSON])

diff --git a/Engine/Source/Runtime/Slate/Private/Framework/Application/SlateUser.cpp b/Engine/Source/Runtime/Slate/Private/Framework/Application/SlateUser.cpp
--- a/Engine/Source/Runtime/Slate/Private/Framework/Application/SlateUser.cpp
+++ b/Engine/Source/Runtime/Slate/Private/Framework/Application/SlateUser.cpp
@@ -905,7 +905,7 @@ void FSlateUser::LockCursorInternal(const FWidgetPath& WidgetPath)
 		// Generate a screen space clip rect based on the widget's geometry
 #if PLATFORM_DESKTOP
 		const bool bIsBorderlessGameWindow = NativeWindow->IsDefinitionValid() && NativeWindow->GetDefinition().Type == EWindowType::GameWindow && !NativeWindow->GetDefinition().HasOSWindowBorder;
-		const int32 ClipRectAdjustment = bIsBorderlessGameWindow ? 1 : 0;
+		const int32 ClipRectAdjustment = bIsBorderlessGameWindow ? 0 : 1;
 #else
 		const int32 ClipRectAdjustment = 0;
 #endif
