PR #8463: Prevent division by zero in SScaleBox::OnArrangeChildren (Contributed by [REDACTED_PERSON])

diff --git a/Engine/Source/Runtime/Slate/Private/Widgets/Layout/SScaleBox.cpp b/Engine/Source/Runtime/Slate/Private/Widgets/Layout/SScaleBox.cpp
--- a/Engine/Source/Runtime/Slate/Private/Widgets/Layout/SScaleBox.cpp
+++ b/Engine/Source/Runtime/Slate/Private/Widgets/Layout/SScaleBox.cpp
@@ -245,6 +245,11 @@ void SScaleBox::OnArrangeChildren(const FGeometry& AllottedGeometry, FArrangedCh
 		LastFinalOffset = FVector2D(0, 0);
 		float FinalScale = TempComputedContentScale;
 
+		if (FMath::IsNearlyZero(FinalScale))
+		{
+			return;
+		}
+
 		// If we're just filling, there's no scale applied, we're just filling the area.
 		if (CurrentStretch != EStretch::Fill)
 		{
