PR #7042: you can get SMultiLineEditableText text line count (Contributed by [REDACTED_PERSON])

diff --git a/Engine/Source/Runtime/Slate/Private/Widgets/Text/SMultiLineEditableText.cpp b/Engine/Source/Runtime/Slate/Private/Widgets/Text/SMultiLineEditableText.cpp
--- a/Engine/Source/Runtime/Slate/Private/Widgets/Text/SMultiLineEditableText.cpp
+++ b/Engine/Source/Runtime/Slate/Private/Widgets/Text/SMultiLineEditableText.cpp
@@ -104,6 +104,11 @@ void SMultiLineEditableText::SetText(const TAttribute< FText >& InText)
 	EditableTextLayout->SetText(InText);
 }
 
+int32 SMultiLineEditableText::GetTextLineCount()
+{
+	return EditableTextLayout->GetTextLineCount();
+}
+
 FText SMultiLineEditableText::GetText() const
 {
 	return EditableTextLayout->GetText();
diff --git a/Engine/Source/Runtime/Slate/Private/Widgets/Text/SlateEditableTextLayout.cpp b/Engine/Source/Runtime/Slate/Private/Widgets/Text/SlateEditableTextLayout.cpp
--- a/Engine/Source/Runtime/Slate/Private/Widgets/Text/SlateEditableTextLayout.cpp
+++ b/Engine/Source/Runtime/Slate/Private/Widgets/Text/SlateEditableTextLayout.cpp
@@ -255,6 +255,11 @@ void FSlateEditableTextLayout::SetText(const TAttribute<FText>& InText)
 	}
 }
 
+int32 FSlateEditableTextLayout::GetTextLineCount()
+{
+	return TextLayout->GetLineCount();
+}
+
 FText FSlateEditableTextLayout::GetText() const
 {
 	SLATE_CROSS_THREAD_CHECK();
diff --git a/Engine/Source/Runtime/Slate/Public/Widgets/Text/SMultiLineEditableText.h b/Engine/Source/Runtime/Slate/Public/Widgets/Text/SMultiLineEditableText.h
--- a/Engine/Source/Runtime/Slate/Public/Widgets/Text/SMultiLineEditableText.h
+++ b/Engine/Source/Runtime/Slate/Public/Widgets/Text/SMultiLineEditableText.h
@@ -302,6 +302,9 @@ public:
 	/** Sets the ReadOnly attribute */
 	void SetIsReadOnly(const TAttribute< bool >& InIsReadOnly);
 
+	/** Get the number of Text Lines */
+	int32 GetTextLineCount();
+
 	/**
 	 * Sets whether to select all text when the user clicks to give focus on the widget
 	 *
diff --git a/Engine/Source/Runtime/Slate/Public/Widgets/Text/SlateEditableTextLayout.h b/Engine/Source/Runtime/Slate/Public/Widgets/Text/SlateEditableTextLayout.h
--- a/Engine/Source/Runtime/Slate/Public/Widgets/Text/SlateEditableTextLayout.h
+++ b/Engine/Source/Runtime/Slate/Public/Widgets/Text/SlateEditableTextLayout.h
@@ -49,6 +49,9 @@ public:
 	void SetTextStyle(const FTextBlockStyle& InTextStyle);
 	const FTextBlockStyle& GetTextStyle() const;
 
+	/** Get the number of Text Lines */
+	int32 GetTextLineCount();
+
 	/** Set the brush to use when drawing the cursor */
 	void SetCursorBrush(const TAttribute<const FSlateBrush*>& InCursorBrush);
 
