Fix reference to UE4 in TabManager

diff --git a/Engine/Source/Runtime/Slate/Private/Framework/Docking/TabManager.cpp b/Engine/Source/Runtime/Slate/Private/Framework/Docking/TabManager.cpp
--- a/Engine/Source/Runtime/Slate/Private/Framework/Docking/TabManager.cpp
+++ b/Engine/Source/Runtime/Slate/Private/Framework/Docking/TabManager.cpp
@@ -32,23 +32,23 @@ DEFINE_LOG_CATEGORY_STATIC(LogTabManager, Display, All);
 
 #define LOCTEXT_NAMESPACE "TabManager"
 
-static const FString UE4_TABMANAGER_OPENED_TAB_STRING = TEXT("OpenedTab");
-static const FString UE4_TABMANAGER_CLOSED_TAB_STRING = TEXT("ClosedTab");
-static const FString UE4_TABMANAGER_SIDEBAR_TAB_STRING = TEXT("SidebarTab");
-static const FString UE4_TABMANAGER_INVALID_TAB_STRING = TEXT("InvalidTab");
+static const FString UE_TABMANAGER_OPENED_TAB_STRING = TEXT("OpenedTab");
+static const FString UE_TABMANAGER_CLOSED_TAB_STRING = TEXT("ClosedTab");
+static const FString UE_TABMANAGER_SIDEBAR_TAB_STRING = TEXT("SidebarTab");
+static const FString UE_TABMANAGER_INVALID_TAB_STRING = TEXT("InvalidTab");
 
 static FString StringFromTabState(ETabState::Type TabState)
 {
 	switch (TabState)
 	{
 	case ETabState::OpenedTab:
-		return UE4_TABMANAGER_OPENED_TAB_STRING;
+		return UE_TABMANAGER_OPENED_TAB_STRING;
 	case ETabState::ClosedTab:
-		return UE4_TABMANAGER_CLOSED_TAB_STRING;
+		return UE_TABMANAGER_CLOSED_TAB_STRING;
 	case ETabState::SidebarTab:
-		return UE4_TABMANAGER_SIDEBAR_TAB_STRING;
+		return UE_TABMANAGER_SIDEBAR_TAB_STRING;
 	default:
-		return UE4_TABMANAGER_INVALID_TAB_STRING;
+		return UE_TABMANAGER_INVALID_TAB_STRING;
 	}
 }
 
@@ -83,19 +83,19 @@ static ESidebarLocation SidebarLocationFromString(const FString& AsString)
 
 static ETabState::Type TabStateFromString(const FString& AsString)
 {
-	if (AsString == UE4_TABMANAGER_OPENED_TAB_STRING)
+	if (AsString == UE_TABMANAGER_OPENED_TAB_STRING)
 	{
 		return ETabState::OpenedTab;
 	}
-	else if (AsString == UE4_TABMANAGER_CLOSED_TAB_STRING)
+	else if (AsString == UE_TABMANAGER_CLOSED_TAB_STRING)
 	{
 		return ETabState::ClosedTab;
 	}
-	else if (AsString == UE4_TABMANAGER_INVALID_TAB_STRING)
+	else if (AsString == UE_TABMANAGER_INVALID_TAB_STRING)
 	{
 		return ETabState::InvalidTab;
 	}
-	else if (AsString == UE4_TABMANAGER_SIDEBAR_TAB_STRING)
+	else if (AsString == UE_TABMANAGER_SIDEBAR_TAB_STRING)
 	{
 		return ETabState::SidebarTab;
 	}
