Fix to UE4 reference in comments

diff --git a/Engine/Source/Runtime/Slate/Private/Framework/Application/SlateApplication.cpp b/Engine/Source/Runtime/Slate/Private/Framework/Application/SlateApplication.cpp
--- a/Engine/Source/Runtime/Slate/Private/Framework/Application/SlateApplication.cpp
+++ b/Engine/Source/Runtime/Slate/Private/Framework/Application/SlateApplication.cpp
@@ -928,8 +928,8 @@ void FSlateApplication::DestroyRenderer()
 }
 
 /**
- * Called when the user closes the outermost frame (ie quitting the app). Uses standard UE4 global variable
- * so normal UE4 applications work as expected
+ * Called when the user closes the outermost frame (i.e. quitting the app). Uses standard UE global variable
+ * so normal UE applications work as expected
  */
 static void OnRequestExit()
 {
@@ -6791,7 +6791,7 @@ TSharedRef<FSlateApplication> FSlateApplication::InitializeAsStandaloneApplicati
 	// initialize renderer
 	FSlateApplication::Get().InitializeRenderer(PlatformRenderer);
 
-	// set the normal UE4 IsEngineExitRequested() when outer frame is closed
+	// set the normal UE IsEngineExitRequested() when outer frame is closed
 	FSlateApplication::Get().SetExitRequestedHandler(FSimpleDelegate::CreateStatic(&OnRequestExit));
 
 	return Slate;
