Prevent drag enter when no widget can accept the operation. The drag will be disabled until a widget can accept it and the asset will be release in OnDragLeave since it doesn't check if the drag operation is valid.

diff --git a/Engine/Source/Runtime/Slate/Private/Framework/Application/SlateApplication.cpp b/Engine/Source/Runtime/Slate/Private/Framework/Application/SlateApplication.cpp
--- a/Engine/Source/Runtime/Slate/Private/Framework/Application/SlateApplication.cpp
+++ b/Engine/Source/Runtime/Slate/Private/Framework/Application/SlateApplication.cpp
@@ -6641,6 +6641,11 @@ bool FSlateApplication::ProcessDragEnterEvent( TSharedRef<SWindow> WindowEntered
 	SetLastUserInteractionTime(this->GetCurrentTime());
 	
 	FWidgetPath WidgetsUnderCursor = LocateWindowUnderMouse( DragDropEvent.GetScreenSpacePosition(), GetInteractiveTopLevelWindows(), false, DragDropEvent.GetUserIndex());
+	// There are no "interactable" widget under the cursor.
+	if (!WidgetsUnderCursor.IsValid())
+	{
+		return false;
+	}
 
 	// Switch worlds for widgets in the current path
 	FScopedSwitchWorldHack SwitchWorld( WidgetsUnderCursor );
