Add support for shortname usage in FUICommandInfo & InputCoords.
Add support for struct param construction of FMultiBlock & FMenuEntryBlock
Add support for overriding of input action text.

diff --git a/Engine/Source/Runtime/InputCore/Private/InputCoreTypes.cpp b/Engine/Source/Runtime/InputCore/Private/InputCoreTypes.cpp
--- a/Engine/Source/Runtime/InputCore/Private/InputCoreTypes.cpp
+++ b/Engine/Source/Runtime/InputCore/Private/InputCoreTypes.cpp
@@ -534,8 +534,8 @@ void EKeys::Initialize()
 	AddKey(FKeyDetails(EKeys::MouseScrollUp, LOCTEXT("MouseScrollUp", "Mouse Wheel Up"), FKeyDetails::MouseButton | FKeyDetails::ButtonAxis));
 	AddKey(FKeyDetails(EKeys::MouseScrollDown, LOCTEXT("MouseScrollDown", "Mouse Wheel Down"), FKeyDetails::MouseButton | FKeyDetails::ButtonAxis));
 
-	AddKey(FKeyDetails(EKeys::LeftMouseButton, LOCTEXT("LeftMouseButton", "Left Mouse Button"), FKeyDetails::MouseButton));
-	AddKey(FKeyDetails(EKeys::RightMouseButton, LOCTEXT("RightMouseButton", "Right Mouse Button"), FKeyDetails::MouseButton));
+	AddKey(FKeyDetails(EKeys::LeftMouseButton, LOCTEXT("LeftMouseButton", "Left Mouse Button"), FKeyDetails::MouseButton, NAME_None, LOCTEXT("LeftMouseButtonShort", "LMB")));
+	AddKey(FKeyDetails(EKeys::RightMouseButton, LOCTEXT("RightMouseButton", "Right Mouse Button"), FKeyDetails::MouseButton, NAME_None, LOCTEXT("RightMouseButtonShort", "RMB")));
 	AddKey(FKeyDetails(EKeys::MiddleMouseButton, LOCTEXT("MiddleMouseButton", "Middle Mouse Button"), FKeyDetails::MouseButton));
 	AddKey(FKeyDetails(EKeys::ThumbMouseButton, LOCTEXT("ThumbMouseButton", "Thumb Mouse Button"), FKeyDetails::MouseButton));
 	AddKey(FKeyDetails(EKeys::ThumbMouseButton2, LOCTEXT("ThumbMouseButton2", "Thumb Mouse Button 2"), FKeyDetails::MouseButton));
diff --git a/Engine/Source/Runtime/Slate/Private/Framework/Commands/InputChord.cpp b/Engine/Source/Runtime/Slate/Private/Framework/Commands/InputChord.cpp
--- a/Engine/Source/Runtime/Slate/Private/Framework/Commands/InputChord.cpp
+++ b/Engine/Source/Runtime/Slate/Private/Framework/Commands/InputChord.cpp
@@ -10,23 +10,23 @@
 /**
  * Returns the friendly, localized string name of this key binding
  */
-FText FInputChord::GetInputText( ) const
+FText FInputChord::GetInputText(const bool bLongDisplayName) const
 {
 	FFormatNamedArguments Args;
 	Args.Add(TEXT("Modifiers"), GetModifierText());
-	Args.Add(TEXT("Key"), GetKeyText());
+	Args.Add(TEXT("Key"), GetKeyText(bLongDisplayName));
 
 	return FText::Format(LOCTEXT("InputText", "{Modifiers}{Key}"), Args);
 }
 
 
-FText FInputChord::GetKeyText() const
+FText FInputChord::GetKeyText(const bool bLongDisplayName) const
 {
 	FText OutString;
 
 	if (Key.IsValid() && !Key.IsModifierKey())
 	{
-		OutString = Key.GetDisplayName();
+		OutString = Key.GetDisplayName(bLongDisplayName);
 	}
 
 	return OutString;
diff --git a/Engine/Source/Runtime/Slate/Private/Framework/Commands/UICommandInfo.cpp b/Engine/Source/Runtime/Slate/Private/Framework/Commands/UICommandInfo.cpp
--- a/Engine/Source/Runtime/Slate/Private/Framework/Commands/UICommandInfo.cpp
+++ b/Engine/Source/Runtime/Slate/Private/Framework/Commands/UICommandInfo.cpp
@@ -100,7 +100,7 @@ FUICommandInfoDecl::operator TSharedRef<FUICommandInfo>() const
 const FText FUICommandInfo::GetInputText() const
 {	
 	// Just get the text from the first valid chord, there isn't enough room for all of them
-	return GetFirstValidChord()->GetInputText();
+	return GetFirstValidChord()->GetInputText(bUseLongDisplayName);
 }
 
 
diff --git a/Engine/Source/Runtime/Slate/Private/Framework/MultiBox/MultiBoxBuilder.cpp b/Engine/Source/Runtime/Slate/Private/Framework/MultiBox/MultiBoxBuilder.cpp
--- a/Engine/Source/Runtime/Slate/Private/Framework/MultiBox/MultiBoxBuilder.cpp
+++ b/Engine/Source/Runtime/Slate/Private/Framework/MultiBox/MultiBoxBuilder.cpp
@@ -140,7 +140,7 @@ void FBaseMenuBuilder::AddMenuEntry( const TSharedPtr< const FUICommandInfo > In
 	
 	// The command must be valid
 	check( InCommand.IsValid() );
-	TSharedRef< FMenuEntryBlock > NewMenuEntryBlock( new FMenuEntryBlock( InExtensionHook, InCommand, CommandListStack.Last(), InLabelOverride, InToolTipOverride, InIconOverride, bCloseSelfOnly ) );
+	TSharedRef< FMenuEntryBlock > NewMenuEntryBlock = MakeShared<FMenuEntryBlock>( InExtensionHook, InCommand, CommandListStack.Last(), InLabelOverride, InToolTipOverride, InIconOverride, bCloseSelfOnly );
 	NewMenuEntryBlock->SetTutorialHighlightName(GenerateTutorialIdentfierName(TutorialHighlightName, InTutorialHighlightName, InCommand, MultiBox->GetBlocks().Num()));
 	MultiBox->AddMultiBlock( NewMenuEntryBlock );
 
@@ -153,7 +153,7 @@ void FBaseMenuBuilder::AddMenuEntry( const TAttribute<FText>& InLabel, const TAt
 
 	ApplyHook(InExtensionHook, EExtensionHook::Before);
 	
-	TSharedRef< FMenuEntryBlock > NewMenuEntryBlock( new FMenuEntryBlock( InExtensionHook, InLabel, InToolTip, InIcon, InAction, UserInterfaceActionType, bCloseSelfOnly ) );
+	TSharedRef< FMenuEntryBlock > NewMenuEntryBlock = MakeShared<FMenuEntryBlock>( InExtensionHook, InLabel, InToolTip, InIcon, InAction, UserInterfaceActionType, bCloseSelfOnly );
 	NewMenuEntryBlock->SetTutorialHighlightName(GenerateTutorialIdentfierName(TutorialHighlightName, InTutorialHighlightName, nullptr, MultiBox->GetBlocks().Num()));
 	MultiBox->AddMultiBlock( NewMenuEntryBlock );
 	
@@ -166,13 +166,26 @@ void FBaseMenuBuilder::AddMenuEntry( const FUIAction& UIAction, const TSharedRef
 
 	ApplyHook(InExtensionHook, EExtensionHook::Before);
 
-	TSharedRef< FMenuEntryBlock > NewMenuEntryBlock( new FMenuEntryBlock( InExtensionHook, UIAction, Contents, InToolTip, UserInterfaceActionType, bCloseSelfOnly ) );
+	TSharedRef< FMenuEntryBlock > NewMenuEntryBlock = MakeShared<FMenuEntryBlock>( InExtensionHook, UIAction, Contents, InToolTip, UserInterfaceActionType, bCloseSelfOnly );
 	NewMenuEntryBlock->SetTutorialHighlightName(GenerateTutorialIdentfierName(TutorialHighlightName, InTutorialHighlightName, nullptr, MultiBox->GetBlocks().Num()));
 	MultiBox->AddMultiBlock( NewMenuEntryBlock );
 
 	ApplyHook(InExtensionHook, EExtensionHook::After);
 }
 
+void FBaseMenuBuilder::AddMenuEntry(const FMenuEntryParams& InMenuEntryParams)
+{
+	ApplySectionBeginning();
+
+	ApplyHook(InMenuEntryParams.ExtensionHook, EExtensionHook::Before);
+	
+	TSharedPtr< FMenuEntryBlock > NewMenuEntryBlock = MakeShared<FMenuEntryBlock>(InMenuEntryParams);
+	NewMenuEntryBlock->SetTutorialHighlightName(GenerateTutorialIdentfierName(TutorialHighlightName, InMenuEntryParams.TutorialHighlightName, nullptr, MultiBox->GetBlocks().Num()));
+	MultiBox->AddMultiBlock(NewMenuEntryBlock.ToSharedRef());
+
+	ApplyHook(InMenuEntryParams.ExtensionHook, EExtensionHook::After);
+}
+
 TSharedRef< class SWidget > FMenuBuilder::MakeWidget( FMultiBox::FOnMakeMultiBoxBuilderOverride* InMakeMultiBoxBuilderOverride /* = nullptr */, uint32 MaxHeight)
 {
 	// Make menu builders searchable (by default)
diff --git a/Engine/Source/Runtime/Slate/Private/Framework/MultiBox/SMenuEntryBlock.cpp b/Engine/Source/Runtime/Slate/Private/Framework/MultiBox/SMenuEntryBlock.cpp
--- a/Engine/Source/Runtime/Slate/Private/Framework/MultiBox/SMenuEntryBlock.cpp
+++ b/Engine/Source/Runtime/Slate/Private/Framework/MultiBox/SMenuEntryBlock.cpp
@@ -172,6 +172,25 @@ FMenuEntryBlock::FMenuEntryBlock( const FName& InExtensionHook, const FUIAction&
 {
 }
 
+FMenuEntryBlock::FMenuEntryBlock(const FMenuEntryParams& InMenuEntryParams)
+	: FMultiBlock(InMenuEntryParams)
+	, LabelOverride(InMenuEntryParams.LabelOverride)
+	, ToolTipOverride(InMenuEntryParams.ToolTipOverride)
+	, InputBindingOverride(InMenuEntryParams.InputBindingOverride)
+	, IconOverride(InMenuEntryParams.IconOverride)
+	, EntryBuilder(InMenuEntryParams.EntryBuilder)
+	, MenuBuilder(InMenuEntryParams.MenuBuilder)
+	, EntryWidget(InMenuEntryParams.EntryWidget)
+	, bIsSubMenu(InMenuEntryParams.bIsSubMenu)
+	, bIsRecursivelySearchable(InMenuEntryParams.bIsRecursivelySearchable)
+	, bOpenSubMenuOnClick(InMenuEntryParams.bOpenSubMenuOnClick)
+	, UserInterfaceActionType(InMenuEntryParams.UserInterfaceActionType)
+	, bCloseSelfOnly(InMenuEntryParams.bCloseSelfOnly)
+	, Extender(InMenuEntryParams.Extender)
+	, bShouldCloseWindowAfterMenuSelection(InMenuEntryParams.bShouldCloseWindowAfterMenuSelection)
+{
+}
+
 void FMenuEntryBlock::CreateMenuEntry(FMenuBuilder& InMenuBuilder) const
 {
 	InMenuBuilder.AddSubMenu(LabelOverride.Get(), ToolTipOverride.Get(), EntryBuilder, false, IconOverride);	
@@ -434,6 +453,7 @@ TSharedRef< SWidget > SMenuEntryBlock::BuildMenuEntryWidget( const FMenuEntryBui
 {
 	const TAttribute<FText>& Label = InBuildParams.Label;
 	const TAttribute<FText>& EntryToolTip = InBuildParams.ToolTip;
+	const TAttribute<FText>& InputBinding = InBuildParams.InputBinding;
 	const TSharedPtr< const FMenuEntryBlock > MenuEntryBlock = InBuildParams.MenuEntryBlock;
 	const TSharedPtr< const FMultiBox > MultiBox = InBuildParams.MultiBox;
 	const TSharedPtr< const FUICommandInfo >& UICommand = InBuildParams.UICommand;
@@ -449,9 +469,6 @@ TSharedRef< SWidget > SMenuEntryBlock::BuildMenuEntryWidget( const FMenuEntryBui
 	const ISlateStyle* const StyleSet = InBuildParams.StyleSet;
 	const FName& StyleName = InBuildParams.StyleName;
 
-	// Grab the friendly text name for this action's input binding
-	FText InputBindingText = UICommand.IsValid() ? UICommand->GetInputText().ToUpper() : FText::GetEmpty();
-
 	// Allow menu item buttons to be triggered on mouse-up events if the menu is configured to be
 	// dismissed automatically after clicking.  This preserves the behavior people expect for context
 	// menus and pull-down menus
@@ -575,13 +592,13 @@ TSharedRef< SWidget > SMenuEntryBlock::BuildMenuEntryWidget( const FMenuEntryBui
 		.HAlign( HAlign_Right )
 		[
 			SNew( SBox )
-			.Visibility(InputBindingText.IsEmpty() ? EVisibility::Collapsed : EVisibility::Visible)
+			.Visibility(InputBinding.Get().IsEmpty() ? EVisibility::Collapsed : EVisibility::Visible)
 			.Padding(FMargin(16,0,4,0))
 			[
 				SNew( STextBlock )
 				.TextStyle( StyleSet, ISlateStyle::Join( StyleName, ".Keybinding" ) )
 				.ColorAndOpacity( FSlateColor::UseSubduedForeground() )
-				.Text( InputBindingText )
+				.Text( InputBinding )
 			]
 		];
 	}
@@ -979,7 +996,17 @@ void SMenuEntryBlock::BuildMultiBlockWidget(const ISlateStyle* StyleSet, const F
 	else
 	{
 		BuildParams.ToolTip = BuildParams.UICommand.IsValid() ? BuildParams.UICommand->GetDescription() : FText::GetEmpty();
-	}	
+	}
+
+	// Input bindings are optional so if the binding override is empty and there is no UI command just use the empty tool tip.
+	if (MenuEntryBlock->InputBindingOverride.IsSet())
+	{
+		BuildParams.InputBinding = MenuEntryBlock->InputBindingOverride.Get().ToUpper();
+	}
+	else
+	{
+		BuildParams.InputBinding = BuildParams.UICommand.IsValid() ? BuildParams.UICommand->GetInputText().ToUpper() : FText::GetEmpty();
+	}
 
 	if (bIsEditing)
 	{
diff --git a/Engine/Source/Runtime/Slate/Private/Framework/MultiBox/SMenuEntryBlock.h b/Engine/Source/Runtime/Slate/Private/Framework/MultiBox/SMenuEntryBlock.h
--- a/Engine/Source/Runtime/Slate/Private/Framework/MultiBox/SMenuEntryBlock.h
+++ b/Engine/Source/Runtime/Slate/Private/Framework/MultiBox/SMenuEntryBlock.h
@@ -114,6 +114,9 @@ public:
 
 	FMenuEntryBlock( const FName& InExtensionHook, const FUIAction& UIAction, const TAttribute<FText>& InLabel, const TAttribute<FText>& InToolTip, const FOnGetContent& InMenuBuilder, TSharedPtr<FExtender> InExtender, bool bInSubMenu, bool bInSubMenuOnClick, bool bInCloseSelfOnly, const FSlateIcon& InIcon = FSlateIcon(), bool bInShouldCloseWindowAfterMenuSelection = true);
 
+	/** Construct a menu entry given param struct */
+	FMenuEntryBlock(const FMenuEntryParams& InMenuEntryParams);
+
 	/** FMultiBlock interface */
 	virtual void CreateMenuEntry(class FMenuBuilder& MenuBuilder) const override;
 	virtual bool HasIcon() const override;
@@ -144,6 +147,9 @@ private:
 	/** Optional overridden tool tip for this menu entry.  If not set, then the action's tool tip will be used instead. */
 	TAttribute<FText> ToolTipOverride;
 
+	/** Optional overridden input binding text for this menu entry.  If not set, then the UI action's binding will be used if available. */
+	TAttribute<FText> InputBindingOverride;
+
 	/** Optional overridden icon for this tool bar button.  IF not set, then the action's icon will be used instead. */
 	FSlateIcon IconOverride;
 
@@ -233,6 +239,8 @@ protected:
 		TAttribute<FText> Label;
 		/** The tooltip to display */
 		TAttribute<FText> ToolTip;
+		/** The input binding to display */
+		TAttribute<FText> InputBinding;
 		/** The style set to use */
 		const ISlateStyle* StyleSet;
 		/** The style name to use */
diff --git a/Engine/Source/Runtime/Slate/Public/Framework/Commands/InputChord.h b/Engine/Source/Runtime/Slate/Public/Framework/Commands/InputChord.h
--- a/Engine/Source/Runtime/Slate/Public/Framework/Commands/InputChord.h
+++ b/Engine/Source/Runtime/Slate/Public/Framework/Commands/InputChord.h
@@ -165,14 +165,14 @@ struct SLATE_API FInputChord
 	 *
 	 * @return A localized string.
 	 */
-	FText GetInputText() const;
+	FText GetInputText(const bool bLongDisplayName = true) const;
 	
 	/**
 	 * Gets the key represented as a localized string.
 	 *
 	 * @return A localized string.
 	 */
-	FText GetKeyText() const;
+	FText GetKeyText(const bool bLongDisplayName = true) const;
 
 	/**
 	 * Gets the localized string for the modifier portion of the coord. 
diff --git a/Engine/Source/Runtime/Slate/Public/Framework/Commands/UICommandInfo.h b/Engine/Source/Runtime/Slate/Public/Framework/Commands/UICommandInfo.h
--- a/Engine/Source/Runtime/Slate/Public/Framework/Commands/UICommandInfo.h
+++ b/Engine/Source/Runtime/Slate/Public/Framework/Commands/UICommandInfo.h
@@ -196,6 +196,7 @@ public:
 	FUICommandInfo( const FName InBindingContext )
 		: BindingContext( InBindingContext )
 		, UserInterfaceType( EUserInterfaceActionType::Button )
+		, bUseLongDisplayName( true )
 	{
 		ActiveChords.Empty(2);
 		ActiveChords.Add(TSharedRef<FInputChord>(new FInputChord));
@@ -268,6 +269,12 @@ public:
 	/** @return The name of the bundle this command is assigned to */
 	FName GetBundle() const { return Bundle; }
 
+	/** @return True if should we use long names for when getting text for input chords */
+	bool GetUseLongDisplayName() const { return bUseLongDisplayName; }
+
+	/** Sets if we should use a long display name or not */
+	void SetUseLongDisplayName(const bool bInUseLongDisplayName) { bUseLongDisplayName = bInUseLongDisplayName; }
+
 	/** Sets the new active chord for this command */
 	void SetActiveChord( const FInputChord& NewChord, const EMultipleKeyBindingIndex InChordIndex );
 
@@ -312,4 +319,7 @@ private:
 
 	/** The type of user interface to associated with this action */
 	EUserInterfaceActionType UserInterfaceType;
+
+	/** True if should we use long names for when getting text for input chords */
+	bool bUseLongDisplayName;
 };
diff --git a/Engine/Source/Runtime/Slate/Public/Framework/MultiBox/MultiBox.h b/Engine/Source/Runtime/Slate/Public/Framework/MultiBox/MultiBox.h
--- a/Engine/Source/Runtime/Slate/Public/Framework/MultiBox/MultiBox.h
+++ b/Engine/Source/Runtime/Slate/Public/Framework/MultiBox/MultiBox.h
@@ -60,6 +60,27 @@ class SLATE_API FMultiBlock
 
 public:
 
+	struct FMultiBlockParams
+	{
+		/** Direct processing of actions. Will use these actions if there is not UICommand associated with this block that handles actions*/
+		FUIAction DirectActions;
+
+		/** The action associated with this block (can be null for some actions) */
+		const TSharedPtr< const FUICommandInfo > Action;
+
+		/** The list of mappings from command info to delegates that should be called. This is here for quick access. Can be null for some widgets*/
+		const TSharedPtr< const FUICommandList > ActionList;
+
+		/** Optional extension hook which is used for debug display purposes, so users can see what hooks are where */
+		FName ExtensionHook;
+
+		/** Type of MultiBlock */
+		EMultiBlockType Type = EMultiBlockType::None;
+
+		/** Whether this block is part of the heading blocks for a section */
+		bool bIsPartOfHeading = false;
+	};
+
 	/**
 	 * Constructor
  	 *
@@ -80,7 +101,7 @@ public:
 	/**
 	 * Constructor
 	 *
-	 * @InAction UI action delegates that should be used in place of UI commands (dynamic menu items)
+	 * @param InAction	UI action delegates that should be used in place of UI commands (dynamic menu items)
 	 */
 	FMultiBlock( const FUIAction& InAction,  FName InExtensionHook = NAME_None, EMultiBlockType InType = EMultiBlockType::None, bool bInIsPartOfHeading = false, TSharedPtr< const FUICommandList > InCommandList = nullptr )
 		: DirectActions( InAction )
@@ -93,6 +114,24 @@ public:
 	{
 	}
 
+	/**
+	 * Constructor
+	 *
+	 * @param InMultiBlockParams	Bundle of params for construction
+	 */
+	FMultiBlock(const FMultiBlockParams& InMultiBlockParams)
+		: DirectActions(InMultiBlockParams.DirectActions)
+		, Action(InMultiBlockParams.Action)
+		, ActionList(InMultiBlockParams.ActionList)
+		, ExtensionHook(InMultiBlockParams.ExtensionHook)
+		, Type(InMultiBlockParams.Type)
+		, TutorialHighlightName(NAME_None)
+		, bSearchable(true)
+		, bIsPartOfHeading(InMultiBlockParams.bIsPartOfHeading)
+	{
+		check(Action.IsValid() || DirectActions.IsBound());
+	}
+
 	virtual ~FMultiBlock()
 	{
 	}
diff --git a/Engine/Source/Runtime/Slate/Public/Framework/MultiBox/MultiBoxBuilder.h b/Engine/Source/Runtime/Slate/Public/Framework/MultiBox/MultiBoxBuilder.h
--- a/Engine/Source/Runtime/Slate/Public/Framework/MultiBox/MultiBoxBuilder.h
+++ b/Engine/Source/Runtime/Slate/Public/Framework/MultiBox/MultiBoxBuilder.h
@@ -154,7 +154,55 @@ protected:
 	bool bExtendersEnabled;
 };
 
+/** Helper struct that holds FMenuEntry params for construction */
+struct FMenuEntryParams : public FMultiBlock::FMultiBlockParams
+{
+	/** Optional overridden text label for this menu entry.  If not set, then the action's label will be used instead. */
+	TAttribute<FText> LabelOverride;
+
+	/** Optional overridden tool tip for this menu entry.  If not set, then the action's tool tip will be used instead. */
+	TAttribute<FText> ToolTipOverride;
+
+	/** Optional overridden input binding text for this menu entry.  If not set, then the UI action's binding will be used if available. */
+	TAttribute<FText> InputBindingOverride;
+
+	/** Optional overridden icon for this tool bar button.  IF not set, then the action's icon will be used instead. */
+	FSlateIcon IconOverride;
+
+	/** Optional menu entry builder associated with this entry for building sub-menus and pull down menus */
+	FNewMenuDelegate EntryBuilder;
+
+	/** Delegate that returns an entire menu */
+	FOnGetContent MenuBuilder;
+
+	/** Widget to be added to the menu */
+	TSharedPtr<SWidget> EntryWidget;
+
+	/** True if this menu entry opens a sub-menu */
+	bool bIsSubMenu = false;
+
+	/** True if the search algorithm should walk down this menu sub menus. Usually true, unless the menu has circular/infinite expansion (happens in some menus generated on the fly by reflection). */
+	bool bIsRecursivelySearchable = true;;
+
+	/** True if this menu entry opens a sub-menu by clicking on it only */
+	bool bOpenSubMenuOnClick = false;
+
+	/** In the case where a command is not bound, the user interface action type to use.  If a command is bound, we
+		simply use the action type associated with that command. */
+	EUserInterfaceActionType UserInterfaceActionType;
+
+	/** True if the menu should close itself and all its children or the entire open menu stack */
+	bool bCloseSelfOnly = false;
+
+	/** An extender that this menu entry should pass down to its children, so they get extended properly */
+	TSharedPtr<FExtender> Extender;
 
+	/** For submenus, whether the menu should be closed after something is selected */
+	bool bShouldCloseWindowAfterMenuSelection = true;
+
+	/** Display name for tutorials */
+	FName TutorialHighlightName;
+};
 
 /**
  * Base menu builder
@@ -199,9 +247,22 @@ public:
 	 * @param	InTutorialHighlightName	Optional name to identify this widget and highlight during tutorials
 	 */
 	void AddMenuEntry( const TAttribute<FText>& InLabel, const TAttribute<FText>& InToolTip, const FSlateIcon& InIcon, const FUIAction& UIAction, FName InExtensionHook = NAME_None, const EUserInterfaceActionType UserInterfaceActionType = EUserInterfaceActionType::Button, FName InTutorialHighlightName = NAME_None );
-
+	
+	/**
+	 * Adds a menu entry with a custom widget
+	 *
+	 * @param	UIAction				Actions to execute on this menu item.
+	 * @param	Contents				Custom widget to display
+	 * @param	InExtensionHook			The section hook. Can be NAME_None
+	 * @param	InToolTip				Tool tip used when hovering over the menu entry
+	 * @param	UserInterfaceActionType	Type of interface action
+	 * @param	InTutorialHighlightName	Optional name to identify this widget and highlight during tutorials
+	 */
 	void AddMenuEntry( const FUIAction& UIAction, const TSharedRef< SWidget > Contents, const FName& InExtensionHook = NAME_None, const TAttribute<FText>& InToolTip = TAttribute<FText>(), const EUserInterfaceActionType UserInterfaceActionType = EUserInterfaceActionType::Button, FName InTutorialHighlightName = NAME_None );
 
+	/** Adds a menu entry with given param struct */
+	void AddMenuEntry( const FMenuEntryParams& InMenuEntryParams );
+
 protected:
 	/** True if clicking on a menu entry closes itself only and its children and not the entire stack */
 	bool bCloseSelfOnly;
