Fix default value setting of scroll bar thumb thickness
#tests: PIE: Tested changing scroll bar thickness of various elements in widget editor

diff --git a/Engine/Source/Runtime/Slate/Private/Widgets/Layout/SScrollBar.cpp b/Engine/Source/Runtime/Slate/Private/Widgets/Layout/SScrollBar.cpp
--- a/Engine/Source/Runtime/Slate/Private/Widgets/Layout/SScrollBar.cpp
+++ b/Engine/Source/Runtime/Slate/Private/Widgets/Layout/SScrollBar.cpp
@@ -18,6 +18,8 @@ void SScrollBar::Construct(const FArguments& InArgs)
 	check(InArgs._Style);
 	SetStyle(InArgs._Style);
 
+	const TAttribute<FVector2D> Thickness = InArgs._Thickness.IsSet() ? InArgs._Thickness : FVector2D(InArgs._Style->Thickness, InArgs._Style->Thickness);
+
 	EHorizontalAlignment HorizontalAlignment = Orientation == Orient_Vertical ? HAlign_Center : HAlign_Fill;
 	EVerticalAlignment VerticalAlignment = Orientation == Orient_Vertical ? VAlign_Fill : VAlign_Center;
 
@@ -60,7 +62,7 @@ void SScrollBar::Construct(const FArguments& InArgs)
 						.Padding(0.f)
 						[
 							SAssignNew(ThicknessSpacer, SSpacer)
-							.Size(InArgs._Thickness)
+							.Size(Thickness)
 						]
 					]
 					.BottomSlot()
diff --git a/Engine/Source/Runtime/Slate/Private/Widgets/Views/STableViewBase.cpp b/Engine/Source/Runtime/Slate/Private/Widgets/Views/STableViewBase.cpp
--- a/Engine/Source/Runtime/Slate/Private/Widgets/Views/STableViewBase.cpp
+++ b/Engine/Source/Runtime/Slate/Private/Widgets/Views/STableViewBase.cpp
@@ -99,7 +99,7 @@ void STableViewBase::ConstructChildren( const TAttribute<float>& InItemWidth, co
 				.AutoWidth()
 				[
 					SNew(SBox)
-					.WidthOverride( FOptionalSize( 16 ) )
+					.WidthOverride(ScrollBarSize)
 					[
 						ScrollBar.ToSharedRef()
 					]
diff --git a/Engine/Source/Runtime/Slate/Public/Widgets/Layout/SScrollBar.h b/Engine/Source/Runtime/Slate/Public/Widgets/Layout/SScrollBar.h
--- a/Engine/Source/Runtime/Slate/Public/Widgets/Layout/SScrollBar.h
+++ b/Engine/Source/Runtime/Slate/Public/Widgets/Layout/SScrollBar.h
@@ -40,7 +40,7 @@ public:
 #endif
 		, _Orientation( Orient_Vertical )
 		, _DragFocusCause( EFocusCause::Mouse )
-		, _Thickness( FVector2D( _Style->Thickness, _Style->Thickness ) )
+		, _Thickness()
 		, _Padding( 2.0f )
 		{}
 
