Exposed a property for AllowOverscroll in ListViewBase.  Prevent ListViewBase from capturing the cursor if no scroll is possible.
#rnx

diff --git a/Engine/Source/Runtime/Slate/Private/Widgets/Views/STableViewBase.cpp b/Engine/Source/Runtime/Slate/Private/Widgets/Views/STableViewBase.cpp
--- a/Engine/Source/Runtime/Slate/Private/Widgets/Views/STableViewBase.cpp
+++ b/Engine/Source/Runtime/Slate/Private/Widgets/Views/STableViewBase.cpp
@@ -618,14 +618,16 @@ FReply STableViewBase::OnTouchMoved( const FGeometry& MyGeometry, const FPointer
 				RegisterActiveTimer(0.f, FWidgetActiveTimerDelegate::CreateSP(this, &STableViewBase::UpdateInertialScroll));
 			}
 
-			const float AmountScrolled = this->ScrollBy( MyGeometry, -ScrollByAmount, EAllowOverscroll::Yes );
-
-			ScrollBar->BeginScrolling();
+			const float AmountScrolled = this->ScrollBy( MyGeometry, -ScrollByAmount, AllowOverscroll );
+			if (AmountScrolled != 0)
+			{
+				ScrollBar->BeginScrolling();
 
-			// The user has moved the list some amount; they are probably
-			// trying to scroll. From now on, the list assumes the user is scrolling
-			// until they lift their finger.
-			return HasMouseCapture() ? FReply::Handled() : FReply::Handled().CaptureMouse(AsShared());
+				// The user has moved the list some amount; they are probably
+				// trying to scroll. From now on, the list assumes the user is scrolling
+				// until they lift their finger.
+				return HasMouseCapture() ? FReply::Handled() : FReply::Handled().CaptureMouse(AsShared());
+			}
 		}
 
 		return FReply::Unhandled();
@@ -865,6 +867,11 @@ void STableViewBase::SetIsScrollAnimationEnabled(bool bInEnableScrollAnimation)
 	bEnableAnimatedScrolling = bInEnableScrollAnimation;
 }
 
+void STableViewBase::SetAllowOverscroll(EAllowOverscroll InAllowOverscroll)
+{
+	AllowOverscroll = InAllowOverscroll;
+}
+
 void STableViewBase::SetIsRightClickScrollingEnabled(const bool bInEnableRightClickScrolling)
 {
 	bEnableRightClickScrolling = bInEnableRightClickScrolling;
diff --git a/Engine/Source/Runtime/Slate/Public/Widgets/Views/STableViewBase.h b/Engine/Source/Runtime/Slate/Public/Widgets/Views/STableViewBase.h
--- a/Engine/Source/Runtime/Slate/Public/Widgets/Views/STableViewBase.h
+++ b/Engine/Source/Runtime/Slate/Public/Widgets/Views/STableViewBase.h
@@ -170,6 +170,9 @@ public:
 	/** Sets whether the list should lerp between scroll offsets or jump instantly between them. */
 	void SetIsScrollAnimationEnabled(bool bInEnableScrollAnimation);
 
+	/** Sets whether to permit overscroll on this list view */
+	void SetAllowOverscroll(EAllowOverscroll InAllowOverscroll);
+
 	/** Enables/disables being able to scroll with the right mouse button. */
 	void SetIsRightClickScrollingEnabled(const bool bInEnableRightClickScrolling);
 
diff --git a/Engine/Source/Runtime/UMG/Private/Components/ListViewBase.cpp b/Engine/Source/Runtime/UMG/Private/Components/ListViewBase.cpp
--- a/Engine/Source/Runtime/UMG/Private/Components/ListViewBase.cpp
+++ b/Engine/Source/Runtime/UMG/Private/Components/ListViewBase.cpp
@@ -158,6 +158,7 @@ void UListViewBase::SynchronizeProperties()
 	{
 		MyTableViewBase->SetIsScrollAnimationEnabled(bEnableScrollAnimation);
 		MyTableViewBase->SetIsRightClickScrollingEnabled(bEnableRightClickScrolling);
+		MyTableViewBase->SetAllowOverscroll(AllowOverscroll ? EAllowOverscroll::Yes : EAllowOverscroll::No);
 		MyTableViewBase->SetFixedLineScrollOffset(bEnableFixedLineOffset ? TOptional<double>(FixedLineScrollOffset) : TOptional<double>());
 		MyTableViewBase->SetWheelScrollMultiplier(GetGlobalScrollAmount() * WheelScrollMultiplier);
 	}
diff --git a/Engine/Source/Runtime/UMG/Public/Components/ListViewBase.h b/Engine/Source/Runtime/UMG/Public/Components/ListViewBase.h
--- a/Engine/Source/Runtime/UMG/Public/Components/ListViewBase.h
+++ b/Engine/Source/Runtime/UMG/Public/Components/ListViewBase.h
@@ -694,6 +694,10 @@ protected:
 	UPROPERTY(EditAnywhere, BlueprintReadOnly, Category = Scrolling)
 	bool bEnableScrollAnimation = false;
 
+	/**  Disable to stop scrollbars from activating inertial overscrolling */
+	UPROPERTY(EditAnywhere, Category = Scrolling)
+	bool AllowOverscroll = true;
+
 	/** True to allow right click drag scrolling. */
 	UPROPERTY(EditAnywhere, Category = Scrolling)
 	bool bEnableRightClickScrolling = true;
