Build fix to SRotatorInputBox

diff --git a/Engine/Source/Runtime/Slate/Public/Widgets/Input/SRotatorInputBox.h b/Engine/Source/Runtime/Slate/Public/Widgets/Input/SRotatorInputBox.h
--- a/Engine/Source/Runtime/Slate/Public/Widgets/Input/SRotatorInputBox.h
+++ b/Engine/Source/Runtime/Slate/Public/Widgets/Input/SRotatorInputBox.h
@@ -11,6 +11,7 @@
 #include "Styling/CoreStyle.h"
 #include "Framework/SlateDelegates.h"
 #include "Widgets/Input/NumericTypeInterface.h"
+#include "Widgets/Input/SNumericEntryBox.h"
 
 class FArrangedChildren;
 
