Fix missing include in SlateFwd

diff --git a/Engine/Source/Runtime/Slate/Public/SlateFwd.h b/Engine/Source/Runtime/Slate/Public/SlateFwd.h
--- a/Engine/Source/Runtime/Slate/Public/SlateFwd.h
+++ b/Engine/Source/Runtime/Slate/Public/SlateFwd.h
@@ -3,6 +3,7 @@
 #pragma once
 
 #include "CoreTypes.h"
+#include "Math/Vector.h"
 
 class SSearchBox;
 class SVolumeControl;
