Detail layouts now correctly load from JSON file.

Previously, this was falling back to loading from INI because the filename was incorrect. Functionally, this is currently no different except that the file is loaded from AppData instead of Saved if it exists.

diff --git a/Engine/Source/Runtime/Slate/Private/Framework/Docking/LayoutService.cpp b/Engine/Source/Runtime/Slate/Private/Framework/Docking/LayoutService.cpp
--- a/Engine/Source/Runtime/Slate/Private/Framework/Docking/LayoutService.cpp
+++ b/Engine/Source/Runtime/Slate/Private/Framework/Docking/LayoutService.cpp
@@ -194,9 +194,9 @@ static void SaveLayoutToJson(const FString& InConfigFileName, const TSharedRef<F
 
 static bool LoadLayoutFromJson(const FString& InConfigFileName, const FString& InLayoutName, TSharedPtr<FTabManager::FLayout>& OutLayout)
 {
-	const FString UserSettingsPath = GetLayoutJsonFileName(InConfigFileName);
+	const FString LayoutJsonFile = GetLayoutJsonFileName(InConfigFileName);
 
-	TSharedPtr<FJsonObject> JsonObject = LoadJsonFile(InConfigFileName);
+	TSharedPtr<FJsonObject> JsonObject = LoadJsonFile(LayoutJsonFile);
 	if (!JsonObject.IsValid())
 	{
 		return false;
