Slate: Upgrade some SWidget to use TSlateAttribute (SExpandableButton, SFxWidget, SRadialBox, SSafeZone, SSplitter, SUniformGridPanel, SUniformWrapPanel). Fix unused SLATE_ATTRIBUTE args (STextEntryPopup, SSuggestionTextBox)
#p4v-preflight-copy 18294687

diff --git a/Engine/Source/Runtime/Slate/Private/Widgets/Input/SExpandableButton.cpp b/Engine/Source/Runtime/Slate/Private/Widgets/Input/SExpandableButton.cpp
--- a/Engine/Source/Runtime/Slate/Private/Widgets/Input/SExpandableButton.cpp
+++ b/Engine/Source/Runtime/Slate/Private/Widgets/Input/SExpandableButton.cpp
@@ -9,6 +9,19 @@
 
 //////////////////////////////////////////////////////////////////////////
 // SExpandableButton
+SLATE_IMPLEMENT_WIDGET(SExpandableButton)
+void SExpandableButton::PrivateRegisterAttributes(FSlateAttributeInitializer& AttributeInitializer)
+{
+	SLATE_ADD_MEMBER_ATTRIBUTE_DEFINITION(AttributeInitializer, IsExpanded, EInvalidateWidgetReason::Paint)
+		.OnValueChanged(FSlateAttributeDescriptor::FAttributeValueChangedDelegate::CreateLambda([](SWidget& Widget)
+			{
+				static_cast<SExpandableButton&>(Widget).UpdateVisibility();
+			}));
+}
+
+SExpandableButton::SExpandableButton()
+	: IsExpanded(*this, false)
+{}
 
 EVisibility SExpandableButton::GetCollapsedVisibility() const
 {
@@ -22,7 +35,8 @@ EVisibility SExpandableButton::GetExpandedVisibility() const
 
 void SExpandableButton::Construct(const FArguments& InArgs)
 {
-	IsExpanded = InArgs._IsExpanded;
+	IsExpanded.Assign(*this, InArgs._IsExpanded);
+	ExpandedChildContent = InArgs._ExpandedChildContent.Widget;
 
 	// Set up the collapsed button content
 	TSharedRef<SWidget> CollapsedButtonContent = (InArgs._CollapsedButtonContent.Widget == SNullWidget::NullWidget)
@@ -33,9 +47,6 @@ void SExpandableButton::Construct(const FArguments& InArgs)
 	TSharedRef<SWidget> ExpandedButtonContent = (InArgs._ExpandedButtonContent.Widget == SNullWidget::NullWidget)
 		? StaticCastSharedRef<SWidget>( SNew(STextBlock) .Text( InArgs._ExpandedText ) )
 		: InArgs._ExpandedButtonContent.Widget;
-		
-	// The child content will be optionally visible depending on the state of the expandable button.
-	InArgs._ExpandedChildContent.Widget->SetVisibility( TAttribute<EVisibility>(this, &SExpandableButton::GetExpandedVisibility) );
 
 	SBorder::Construct(SBorder::FArguments()
 		.BorderImage( FCoreStyle::Get().GetBrush( "ExpandableButton.Background" ) )
@@ -48,9 +59,8 @@ void SExpandableButton::Construct(const FArguments& InArgs)
 			.AutoWidth()
 			.VAlign(VAlign_Center)
 			[
-				SNew(SButton)
+				SAssignNew(ToggleButtonClosed, SButton)
 				.VAlign(VAlign_Center)
-				.Visibility( this, &SExpandableButton::GetCollapsedVisibility )
 				.OnClicked( InArgs._OnExpansionClicked )
 				.ButtonStyle( FCoreStyle::Get(), "NoBorder" )
 				.ContentPadding( 0.f )
@@ -63,8 +73,7 @@ void SExpandableButton::Construct(const FArguments& InArgs)
 			+SHorizontalBox::Slot()
 			.AutoWidth()
 			[
-				SNew(SButton)
-				.Visibility( this, &SExpandableButton::GetExpandedVisibility )
+				SAssignNew(ToggleButtonExpanded, SButton)
 				.ButtonStyle( FCoreStyle::Get(), "NoBorder" )
 				.ContentPadding( 0.f )
 				.VAlign(VAlign_Center)
@@ -88,10 +97,9 @@ void SExpandableButton::Construct(const FArguments& InArgs)
 			.HAlign(HAlign_Right)
 			[
 				// Close expansion button
-				SNew(SButton)
+				SAssignNew(CloseExpansionButton, SButton)
 				.ButtonStyle( FCoreStyle::Get(), "NoBorder" )
 				.ContentPadding( 0.f )
-				.Visibility( this, &SExpandableButton::GetExpandedVisibility )
 				.OnClicked(InArgs._OnCloseClicked)
 				[
 					SNew(SImage)
@@ -100,4 +108,18 @@ void SExpandableButton::Construct(const FArguments& InArgs)
 			]
 		]
 	);
+
+	UpdateVisibility();
+}
+
+void SExpandableButton::UpdateVisibility()
+{
+	// The child content will be optionally visible depending on the state of the expandable button.
+	if (ExpandedChildContent)
+	{
+		ExpandedChildContent->SetVisibility(GetExpandedVisibility());
+	}
+	ToggleButtonClosed->SetVisibility(GetCollapsedVisibility());
+	ToggleButtonExpanded->SetVisibility(GetExpandedVisibility());
+	CloseExpansionButton->SetVisibility(GetExpandedVisibility());
 }
diff --git a/Engine/Source/Runtime/Slate/Private/Widgets/Input/SSuggestionTextBox.cpp b/Engine/Source/Runtime/Slate/Private/Widgets/Input/SSuggestionTextBox.cpp
--- a/Engine/Source/Runtime/Slate/Private/Widgets/Input/SSuggestionTextBox.cpp
+++ b/Engine/Source/Runtime/Slate/Private/Widgets/Input/SSuggestionTextBox.cpp
@@ -38,16 +38,16 @@ void SSuggestionTextBox::Construct( const FArguments& InArgs )
 			[
 				SAssignNew(TextBox, SEditableTextBox)
 					.BackgroundColor(InArgs._BackgroundColor)
-					.ClearKeyboardFocusOnCommit(InArgs._ClearKeyboardFocusOnCommit.Get())
+					.ClearKeyboardFocusOnCommit(InArgs._ClearKeyboardFocusOnCommit)
 					.ErrorReporting(InArgs._ErrorReporting)
 					.Font(InArgs._Font)
 					.ForegroundColor(InArgs._ForegroundColor)
 					.HintText(InArgs._HintText)
 					.IsCaretMovedWhenGainFocus(InArgs._IsCaretMovedWhenGainFocus)
 					.MinDesiredWidth(InArgs._MinDesiredWidth)
-					.RevertTextOnEscape(InArgs._RevertTextOnEscape.Get())
+					.RevertTextOnEscape(InArgs._RevertTextOnEscape)
 					.SelectAllTextOnCommit(InArgs._SelectAllTextOnCommit)
-					.SelectAllTextWhenFocused(InArgs._SelectAllTextWhenFocused.Get())
+					.SelectAllTextWhenFocused(InArgs._SelectAllTextWhenFocused)
 					.Style(InArgs._TextStyle)
 					.Text(InArgs._Text)
 					.OnKeyDownHandler(this, &SSuggestionTextBox::OnKeyDown)
diff --git a/Engine/Source/Runtime/Slate/Private/Widgets/Input/STextEntryPopup.cpp b/Engine/Source/Runtime/Slate/Private/Widgets/Input/STextEntryPopup.cpp
--- a/Engine/Source/Runtime/Slate/Private/Widgets/Input/STextEntryPopup.cpp
+++ b/Engine/Source/Runtime/Slate/Private/Widgets/Input/STextEntryPopup.cpp
@@ -39,7 +39,7 @@ void STextEntryPopup::Construct( const FArguments& InArgs )
 					.OnTextCommitted( InArgs._OnTextCommitted )
 					.OnTextChanged( InArgs._OnTextChanged )
 					.HintText( InArgs._HintText )
-					.SelectAllTextWhenFocused( InArgs._SelectAllTextWhenFocused.Get() )
+					.SelectAllTextWhenFocused( InArgs._SelectAllTextWhenFocused )
 					.ClearKeyboardFocusOnCommit( InArgs._ClearKeyboardFocusOnCommit )
 				]
 			]
diff --git a/Engine/Source/Runtime/Slate/Private/Widgets/Layout/SFxWidget.cpp b/Engine/Source/Runtime/Slate/Private/Widgets/Layout/SFxWidget.cpp
--- a/Engine/Source/Runtime/Slate/Private/Widgets/Layout/SFxWidget.cpp
+++ b/Engine/Source/Runtime/Slate/Private/Widgets/Layout/SFxWidget.cpp
@@ -5,13 +5,22 @@
 #include "Layout/ArrangedChildren.h"
 
 
+SFxWidget::SFxWidget()
+	: RenderScale(*this, 1.f)
+	, RenderScaleOrigin(*this, FVector2D::ZeroVector)
+	, LayoutScale(*this, 1.f)
+	, VisualOffset(*this, FVector2D::ZeroVector)
+	, bIgnoreClipping(*this, true)
+{
+}
+
 void SFxWidget::Construct( const FArguments& InArgs )
 {
-	RenderScale = InArgs._RenderScale;
-	RenderScaleOrigin = InArgs._RenderScaleOrigin;
-	LayoutScale = InArgs._LayoutScale;
-	VisualOffset = InArgs._VisualOffset;
-	bIgnoreClipping = InArgs._IgnoreClipping;
+	RenderScale.Assign(*this, InArgs._RenderScale);
+	RenderScaleOrigin.Assign(*this, InArgs._RenderScaleOrigin);
+	LayoutScale.Assign(*this, InArgs._LayoutScale);
+	VisualOffset.Assign(*this, InArgs._VisualOffset);
+	bIgnoreClipping.Assign(*this, InArgs._IgnoreClipping);
 	SetColorAndOpacity(InArgs._ColorAndOpacity);
 	
 	this->ChildSlot
@@ -22,24 +31,24 @@ void SFxWidget::Construct( const FArguments& InArgs )
 	];
 }
 
-void SFxWidget::SetVisualOffset( const TAttribute<FVector2D>& InOffset )
+void SFxWidget::SetVisualOffset( TAttribute<FVector2D> InOffset )
 {
-	VisualOffset = InOffset;
+	VisualOffset.Assign(*this, MoveTemp(InOffset));
 }
 
 void SFxWidget::SetVisualOffset( FVector InOffset )
 {
-	VisualOffset = FVector2D(InOffset.X, InOffset.Y);
+	VisualOffset.Set(*this, FVector2D(InOffset.X, InOffset.Y));
 }
 
-void SFxWidget::SetRenderScale( const TAttribute<float>& InScale )
+void SFxWidget::SetRenderScale( TAttribute<float> InScale )
 {
-	RenderScale = InScale;
+	RenderScale.Assign(*this, InScale);
 }
 
 void SFxWidget::SetRenderScale( float InScale )
 {
-	RenderScale = InScale;
+	RenderScale.Set(*this, InScale);
 }
 
 /**
diff --git a/Engine/Source/Runtime/Slate/Private/Widgets/Layout/SRadialBox.cpp b/Engine/Source/Runtime/Slate/Private/Widgets/Layout/SRadialBox.cpp
--- a/Engine/Source/Runtime/Slate/Private/Widgets/Layout/SRadialBox.cpp
+++ b/Engine/Source/Runtime/Slate/Private/Widgets/Layout/SRadialBox.cpp
@@ -4,7 +4,8 @@
 #include "Layout/LayoutUtils.h"
 
 SRadialBox::SRadialBox()
-: Slots(this)
+	: Slots(this)
+	, PreferredWidth(*this, 100.f)
 {
 }
 
@@ -25,7 +26,7 @@ int32 SRadialBox::RemoveSlot( const TSharedRef<SWidget>& SlotWidget )
 
 void SRadialBox::Construct( const FArguments& InArgs )
 {
-	PreferredWidth = InArgs._PreferredWidth;
+	PreferredWidth.Assign(*this, InArgs._PreferredWidth);
 	bUseAllottedWidth = InArgs._UseAllottedWidth;
 	StartingAngle = InArgs._StartingAngle;
 	bDistributeItemsEvenly = InArgs._bDistributeItemsEvenly;
@@ -39,7 +40,7 @@ void SRadialBox::Tick( const FGeometry& AllottedGeometry, const double InCurrent
 {
 	if (bUseAllottedWidth)
 	{
-		PreferredWidth = AllottedGeometry.GetLocalSize().X;
+		PreferredWidth.Set(*this, AllottedGeometry.GetLocalSize().X);
 	}
 }
 
diff --git a/Engine/Source/Runtime/Slate/Private/Widgets/Layout/SSafeZone.cpp b/Engine/Source/Runtime/Slate/Private/Widgets/Layout/SSafeZone.cpp
--- a/Engine/Source/Runtime/Slate/Private/Widgets/Layout/SSafeZone.cpp
+++ b/Engine/Source/Runtime/Slate/Private/Widgets/Layout/SSafeZone.cpp
@@ -22,6 +22,12 @@ static FAutoConsoleVariableRef CVarEnableSafeZoneScale(
 	ECVF_Default
 );
 
+SSafeZone::SSafeZone()
+	: Padding(*this, 0.f)
+{
+	SetCanTick(false);
+	bCanSupportFocus = false;
+}
 
 void SSafeZone::SetGlobalSafeZoneScale(TOptional<float> InScale)
 {
@@ -55,7 +61,7 @@ void SSafeZone::Construct( const FArguments& InArgs )
 		]
 	);
 
-	Padding = InArgs._Padding;
+	Padding.Assign(*this, InArgs._Padding);
 	SafeAreaScale = InArgs._SafeAreaScale;
 	bIsTitleSafe = InArgs._IsTitleSafe;
 	bPadLeft = InArgs._PadLeft;
diff --git a/Engine/Source/Runtime/Slate/Private/Widgets/Layout/SSplitter.cpp b/Engine/Source/Runtime/Slate/Private/Widgets/Layout/SSplitter.cpp
--- a/Engine/Source/Runtime/Slate/Private/Widgets/Layout/SSplitter.cpp
+++ b/Engine/Source/Runtime/Slate/Private/Widgets/Layout/SSplitter.cpp
@@ -78,7 +78,7 @@ void SSplitter::Construct( const SSplitter::FArguments& InArgs )
 	HitDetectionSplitterHandleSize = InArgs._HitDetectionSplitterHandleSize;
 	Orientation = InArgs._Orientation;
 	HoveredHandleIndex = INDEX_NONE;
-	HighlightedHandleIndex = InArgs._HighlightedHandleIndex;
+	HighlightedHandleIndex.Assign(*this, InArgs._HighlightedHandleIndex);
 	bIsResizing = false;
 	Style = InArgs._Style;
 	OnGetMaxSlotSize = InArgs._OnGetMaxSlotSize;
@@ -271,7 +271,7 @@ int32 SSplitter::OnPaint( const FPaintArgs& Args, const FGeometry& AllottedGeome
 	};
 
 	int32 ArrangedHoveredHandleIndex = FindArrangedHandleIndex(HoveredHandleIndex);
-	int32 ArrangedHighlightedHandleIndex = FindArrangedHandleIndex(HighlightedHandleIndex.Get(INDEX_NONE));
+	int32 ArrangedHighlightedHandleIndex = FindArrangedHandleIndex(HighlightedHandleIndex.Get());
 
 	for (int32 ChildIndex = 0; ChildIndex < ArrangedChildren.Num() - 1; ++ChildIndex)
 	{
@@ -569,7 +569,8 @@ EOrientation SSplitter::GetOrientation() const
 
 SSplitter::SSplitter()
 	: Children(this)
-	, HoveredHandleIndex( INDEX_NONE )
+	, HoveredHandleIndex(INDEX_NONE)
+	, HighlightedHandleIndex(*this, INDEX_NONE)
 	, bIsResizing( false )
 	, Orientation( Orient_Horizontal )
 	, Style( nullptr )
@@ -854,7 +855,7 @@ void SSplitter2x2::FSlot::Construct(const FChildren& SlotOwner, FSlotArguments&&
 }
 
 SSplitter2x2::SSplitter2x2()
-: Children(this)
+	: Children(this)
 {
 }
 
diff --git a/Engine/Source/Runtime/Slate/Private/Widgets/Layout/SUniformGridPanel.cpp b/Engine/Source/Runtime/Slate/Private/Widgets/Layout/SUniformGridPanel.cpp
--- a/Engine/Source/Runtime/Slate/Private/Widgets/Layout/SUniformGridPanel.cpp
+++ b/Engine/Source/Runtime/Slate/Private/Widgets/Layout/SUniformGridPanel.cpp
@@ -4,7 +4,10 @@
 #include "Layout/LayoutUtils.h"
 
 SUniformGridPanel::SUniformGridPanel()
-: Children(this)
+	: Children(this)
+	, SlotPadding(*this, FMargin(0.0f))
+	, MinDesiredSlotWidth(*this, 0.f)
+	, MinDesiredSlotHeight(*this, 0.f)
 {
 }
 
@@ -15,11 +18,11 @@ SUniformGridPanel::FSlot::FSlotArguments SUniformGridPanel::Slot(int32 Column, i
 
 void SUniformGridPanel::Construct( const FArguments& InArgs )
 {
-	SlotPadding = InArgs._SlotPadding;
+	SlotPadding.Assign(*this, InArgs._SlotPadding);
 	NumColumns = 0;
 	NumRows = 0;
-	MinDesiredSlotWidth = InArgs._MinDesiredSlotWidth.Get();
-	MinDesiredSlotHeight = InArgs._MinDesiredSlotHeight.Get();
+	MinDesiredSlotWidth.Assign(*this, InArgs._MinDesiredSlotWidth);
+	MinDesiredSlotHeight.Assign(*this, InArgs._MinDesiredSlotHeight);
 
 	Children.AddSlots(MoveTemp(const_cast<TArray<FSlot::FSlotArguments>&>(InArgs._Slots)));
 }
@@ -93,17 +96,17 @@ FChildren* SUniformGridPanel::GetChildren()
 
 void SUniformGridPanel::SetSlotPadding(TAttribute<FMargin> InSlotPadding)
 {
-	SlotPadding = InSlotPadding;
+	SlotPadding.Assign(*this, MoveTemp(InSlotPadding));
 }
 
 void SUniformGridPanel::SetMinDesiredSlotWidth(TAttribute<float> InMinDesiredSlotWidth)
 {
-	MinDesiredSlotWidth = InMinDesiredSlotWidth;
+	MinDesiredSlotWidth.Assign(*this, MoveTemp(InMinDesiredSlotWidth));
 }
 
 void SUniformGridPanel::SetMinDesiredSlotHeight(TAttribute<float> InMinDesiredSlotHeight)
 {
-	MinDesiredSlotHeight = InMinDesiredSlotHeight;
+	MinDesiredSlotHeight.Assign(*this, MoveTemp(InMinDesiredSlotHeight));
 }
 
 SUniformGridPanel::FScopedWidgetSlotArguments SUniformGridPanel::AddSlot( int32 Column, int32 Row )
diff --git a/Engine/Source/Runtime/Slate/Private/Widgets/Layout/SUniformWrapPanel.cpp b/Engine/Source/Runtime/Slate/Private/Widgets/Layout/SUniformWrapPanel.cpp
--- a/Engine/Source/Runtime/Slate/Private/Widgets/Layout/SUniformWrapPanel.cpp
+++ b/Engine/Source/Runtime/Slate/Private/Widgets/Layout/SUniformWrapPanel.cpp
@@ -3,24 +3,42 @@
 #include "Widgets/Layout/SUniformWrapPanel.h"
 #include "Layout/LayoutUtils.h"
 
+SLATE_IMPLEMENT_WIDGET(SUniformWrapPanel)
+void SUniformWrapPanel::PrivateRegisterAttributes(FSlateAttributeInitializer& AttributeInitializer)
+{
+	SLATE_ADD_MEMBER_ATTRIBUTE_DEFINITION(AttributeInitializer, SlotPadding, EInvalidateWidgetReason::Layout);
+	SLATE_ADD_MEMBER_ATTRIBUTE_DEFINITION(AttributeInitializer, MinDesiredSlotWidth, EInvalidateWidgetReason::Layout);
+	SLATE_ADD_MEMBER_ATTRIBUTE_DEFINITION(AttributeInitializer, MinDesiredSlotHeight, EInvalidateWidgetReason::Layout);
+	SLATE_ADD_MEMBER_ATTRIBUTE_DEFINITION(AttributeInitializer, MaxDesiredSlotWidth, EInvalidateWidgetReason::Layout);
+	SLATE_ADD_MEMBER_ATTRIBUTE_DEFINITION(AttributeInitializer, MaxDesiredSlotHeight, EInvalidateWidgetReason::Layout);
+	SLATE_ADD_MEMBER_ATTRIBUTE_DEFINITION(AttributeInitializer, HAlign, EInvalidateWidgetReason::Paint);
+	SLATE_ADD_MEMBER_ATTRIBUTE_DEFINITION(AttributeInitializer, EvenRowDistribution, EInvalidateWidgetReason::Paint);
+}
+
 SUniformWrapPanel::SUniformWrapPanel()
-: Children(this)
-, HAlign(EHorizontalAlignment::HAlign_Left)
-, EvenRowDistribution(false)
+	: Children(this)
+	, SlotPadding(*this, FMargin(0.0f))
+	, MinDesiredSlotWidth(*this, 0.f)
+	, MinDesiredSlotHeight(*this, 0.f)
+	, MaxDesiredSlotWidth(*this, FLT_MAX)
+	, MaxDesiredSlotHeight(*this, FLT_MAX)
+	, HAlign(*this, EHorizontalAlignment::HAlign_Left)
+	, EvenRowDistribution(*this, false)
 {
 }
 
 void SUniformWrapPanel::Construct( const FArguments& InArgs )
 {
-	SlotPadding = InArgs._SlotPadding;
+	SlotPadding.Assign(*this, InArgs._SlotPadding);
 	NumColumns = 0;
 	NumRows = 0;
-	MinDesiredSlotWidth = InArgs._MinDesiredSlotWidth.Get();
-	MinDesiredSlotHeight = InArgs._MinDesiredSlotHeight.Get();
-	MaxDesiredSlotWidth = InArgs._MaxDesiredSlotWidth.Get();
-	MaxDesiredSlotHeight = InArgs._MaxDesiredSlotHeight.Get();
-	EvenRowDistribution =  InArgs._EvenRowDistribution.Get();
-	HAlign = InArgs._HAlign.Get();
+	NumVisibleChildren = 0;
+	MinDesiredSlotWidth.Assign(*this, InArgs._MinDesiredSlotWidth);
+	MinDesiredSlotHeight.Assign(*this, InArgs._MinDesiredSlotHeight);
+	MaxDesiredSlotWidth.Assign(*this, InArgs._MaxDesiredSlotWidth);
+	MaxDesiredSlotHeight.Assign(*this, InArgs._MaxDesiredSlotHeight);
+	HAlign.Assign(*this, InArgs._HAlign);
+	EvenRowDistribution.Assign(*this, InArgs._EvenRowDistribution);
 
 	Children.AddSlots(MoveTemp(const_cast<TArray<FSlot::FSlotArguments>&>(InArgs._Slots)));
 }
@@ -178,37 +196,37 @@ FChildren* SUniformWrapPanel::GetChildren()
 
 void SUniformWrapPanel::SetSlotPadding(TAttribute<FMargin> InSlotPadding)
 {
-	SlotPadding = InSlotPadding;
+	SlotPadding.Assign(*this, MoveTemp(InSlotPadding));
 }
 
 void SUniformWrapPanel::SetMinDesiredSlotWidth(TAttribute<float> InMinDesiredSlotWidth)
 {
-	MinDesiredSlotWidth = InMinDesiredSlotWidth;
+	MinDesiredSlotWidth.Assign(*this, MoveTemp(InMinDesiredSlotWidth));
 }
 
 void SUniformWrapPanel::SetMinDesiredSlotHeight(TAttribute<float> InMinDesiredSlotHeight)
 {
-	MinDesiredSlotHeight = InMinDesiredSlotHeight;
+	MinDesiredSlotHeight.Assign(*this, MoveTemp(InMinDesiredSlotHeight));
 }
 
 void SUniformWrapPanel::SetMaxDesiredSlotWidth(TAttribute<float> InMaxDesiredSlotWidth)
 {
-	MaxDesiredSlotWidth = InMaxDesiredSlotWidth;
+	MaxDesiredSlotWidth.Assign(*this, MoveTemp(InMaxDesiredSlotWidth));
 }
 
 void SUniformWrapPanel::SetMaxDesiredSlotHeight(TAttribute<float> InMaxDesiredSlotHeight)
 {
-	MaxDesiredSlotHeight = InMaxDesiredSlotHeight;
+	MaxDesiredSlotHeight.Assign(*this, MoveTemp(InMaxDesiredSlotHeight));
 }
 
 void SUniformWrapPanel::SetHorizontalAlignment(TAttribute<EHorizontalAlignment> InHAlignment)
 {
-	HAlign = InHAlignment;	
+	HAlign.Assign(*this, MoveTemp(InHAlignment));
 }
 
 void SUniformWrapPanel::SetEvenRowDistribution(TAttribute<bool> InEvenRowDistribution)
 {
-	EvenRowDistribution = InEvenRowDistribution;
+	EvenRowDistribution.Assign(*this, MoveTemp(InEvenRowDistribution));
 }
 
 SUniformWrapPanel::FScopedWidgetSlotArguments SUniformWrapPanel::AddSlot()
diff --git a/Engine/Source/Runtime/Slate/Public/Widgets/Input/SExpandableButton.h b/Engine/Source/Runtime/Slate/Public/Widgets/Input/SExpandableButton.h
--- a/Engine/Source/Runtime/Slate/Public/Widgets/Input/SExpandableButton.h
+++ b/Engine/Source/Runtime/Slate/Public/Widgets/Input/SExpandableButton.h
@@ -15,6 +15,8 @@
 class SLATE_API SExpandableButton
 	: public SBorder
 {
+	SLATE_DECLARE_WIDGET(SExpandableButton, SBorder)
+
 public:
 
 	SLATE_BEGIN_ARGS( SExpandableButton )
@@ -47,6 +49,7 @@ public:
 
 	SLATE_END_ARGS()
 
+	SExpandableButton();
 	/**
 	 * Construct this widget
 	 *
@@ -60,8 +63,15 @@ protected:
 	EVisibility GetCollapsedVisibility() const;
 	EVisibility GetExpandedVisibility() const;
 
-protected:
+	void UpdateVisibility();
+
+private:
 
 	/** The attribute of the current expansion state */
-	TAttribute<bool> IsExpanded;
+	TSlateAttribute<bool> IsExpanded;
+
+	TSharedPtr<SWidget> ExpandedChildContent;
+	TSharedPtr<SWidget> ToggleButtonClosed;
+	TSharedPtr<SWidget> ToggleButtonExpanded;
+	TSharedPtr<SWidget> CloseExpansionButton;
 };
diff --git a/Engine/Source/Runtime/Slate/Public/Widgets/Input/STextEntryPopup.h b/Engine/Source/Runtime/Slate/Public/Widgets/Input/STextEntryPopup.h
--- a/Engine/Source/Runtime/Slate/Public/Widgets/Input/STextEntryPopup.h
+++ b/Engine/Source/Runtime/Slate/Public/Widgets/Input/STextEntryPopup.h
@@ -33,7 +33,7 @@ public:
 		SLATE_EVENT( FOnTextCommitted, OnTextCommitted )
 
 		/** Whether to select all text when the user clicks to give focus on the widget */
-		SLATE_ATTRIBUTE( bool, SelectAllTextWhenFocused )
+		SLATE_ARGUMENT( bool, SelectAllTextWhenFocused )
 
 		/** Whether to clear keyboard focus when pressing enter to commit changes */
 		SLATE_ATTRIBUTE( bool, ClearKeyboardFocusOnCommit )
diff --git a/Engine/Source/Runtime/Slate/Public/Widgets/Layout/SFxWidget.h b/Engine/Source/Runtime/Slate/Public/Widgets/Layout/SFxWidget.h
--- a/Engine/Source/Runtime/Slate/Public/Widgets/Layout/SFxWidget.h
+++ b/Engine/Source/Runtime/Slate/Public/Widgets/Layout/SFxWidget.h
@@ -20,55 +20,58 @@ class SLATE_API SFxWidget : public SCompoundWidget
 {
 public:
 	SLATE_BEGIN_ARGS( SFxWidget )
-	: _RenderScale( 1.0f )
-	, _RenderScaleOrigin( FVector2D::ZeroVector )
-	, _LayoutScale( 1.0f )
-	, _VisualOffset( FVector2D::ZeroVector )
-	, _IgnoreClipping( true )
-	, _ColorAndOpacity( FLinearColor::White )
-	, _HAlign( HAlign_Center )
-	, _VAlign( VAlign_Center )
-	, _Content()
+		: _RenderScale( 1.0f )
+		, _RenderScaleOrigin( FVector2D::ZeroVector )
+		, _LayoutScale( 1.0f )
+		, _VisualOffset( FVector2D::ZeroVector )
+		, _IgnoreClipping( true )
+		, _ColorAndOpacity( FLinearColor::White )
+		, _HAlign( HAlign_Center )
+		, _VAlign( VAlign_Center )
+		, _Content()
 	{}
-	/** Scale the visuals of this widget. Geometry is not affected. */
-	SLATE_ATTRIBUTE( float, RenderScale )
 
-	/** The origin of the visual scale transform in the 0..1 range (0 being upper left, 1 being lower right)  */
-	SLATE_ATTRIBUTE( FVector2D, RenderScaleOrigin )
+		/** Scale the visuals of this widget. Geometry is not affected. */
+		SLATE_ATTRIBUTE( float, RenderScale )
 
-	/** Just like visual scale only affects Geometry. */
-	SLATE_ATTRIBUTE( float, LayoutScale )
+		/** The origin of the visual scale transform in the 0..1 range (0 being upper left, 1 being lower right)  */
+		SLATE_ATTRIBUTE( FVector2D, RenderScaleOrigin )
 
-	/** Offset the widget by some fraction of its size in either dimension. */
-	SLATE_ATTRIBUTE( FVector2D, VisualOffset )
+		/** Just like visual scale only affects Geometry. */
+		SLATE_ATTRIBUTE( float, LayoutScale )
 
-	/** Should the FX widget disable all clipping and show through regardless of its parents' bounds.*/
-	SLATE_ATTRIBUTE( bool, IgnoreClipping )
+		/** Offset the widget by some fraction of its size in either dimension. */
+		SLATE_ATTRIBUTE( FVector2D, VisualOffset )
 
-	/** Multiply the contents of the SFxWidget by this color and opacity when drawing */
-	SLATE_ATTRIBUTE( FLinearColor, ColorAndOpacity )
+		/** Should the FX widget disable all clipping and show through regardless of its parents' bounds.*/
+		SLATE_ATTRIBUTE( bool, IgnoreClipping )
 
-	/** The horizontal alignment of the child widget */
-	SLATE_ARGUMENT( EHorizontalAlignment, HAlign )
+		/** Multiply the contents of the SFxWidget by this color and opacity when drawing */
+		SLATE_ATTRIBUTE( FLinearColor, ColorAndOpacity )
 
-	/** The vertical alignment of the child widget */
-	SLATE_ARGUMENT( EVerticalAlignment, VAlign )
+		/** The horizontal alignment of the child widget */
+		SLATE_ARGUMENT( EHorizontalAlignment, HAlign )
 
-	/** The content that should be modified. */
-	SLATE_DEFAULT_SLOT( FArguments, Content )
+		/** The vertical alignment of the child widget */
+		SLATE_ARGUMENT( EVerticalAlignment, VAlign )
+
+		/** The content that should be modified. */
+		SLATE_DEFAULT_SLOT( FArguments, Content )
 
 	SLATE_END_ARGS()
 
+	SFxWidget();
+
 	void Construct( const FArguments& InArgs );
 
 	/** @see VisualOffset */
-	void SetVisualOffset( const TAttribute<FVector2D>& InOffset );
+	void SetVisualOffset( TAttribute<FVector2D> InOffset );
 
 	/** @see VisualOffset */
 	void SetVisualOffset( FVector InOffset );
 
 	/** @see RenderScale */
-	void SetRenderScale( const TAttribute<float>& InScale );
+	void SetRenderScale( TAttribute<float> InScale );
 
 	/** @see RenderScale */
 	void SetRenderScale( float InScale );
@@ -80,9 +83,9 @@ protected:
 	virtual FVector2D ComputeDesiredSize(float) const override;
 	virtual void OnArrangeChildren( const FGeometry& AllottedGeometry, FArrangedChildren& ArrangedChildren ) const override;
 
-	TAttribute<float> RenderScale;
-	TAttribute<FVector2D> RenderScaleOrigin;
-	TAttribute<float> LayoutScale;
-	TAttribute<FVector2D> VisualOffset;
-	TAttribute<bool> bIgnoreClipping;
+	TSlateAttribute<float, EInvalidateWidgetReason::Paint> RenderScale;
+	TSlateAttribute<FVector2D, EInvalidateWidgetReason::Paint> RenderScaleOrigin;
+	TSlateAttribute<float, EInvalidateWidgetReason::Layout> LayoutScale;
+	TSlateAttribute<FVector2D, EInvalidateWidgetReason::Paint> VisualOffset;
+	TSlateAttribute<bool, EInvalidateWidgetReason::Paint> bIgnoreClipping;
 };
diff --git a/Engine/Source/Runtime/Slate/Public/Widgets/Layout/SRadialBox.h b/Engine/Source/Runtime/Slate/Public/Widgets/Layout/SRadialBox.h
--- a/Engine/Source/Runtime/Slate/Public/Widgets/Layout/SRadialBox.h
+++ b/Engine/Source/Runtime/Slate/Public/Widgets/Layout/SRadialBox.h
@@ -109,26 +109,25 @@ private:
 	/** Mods the angle so it's between 0-360 */
 	float InternalNormalizeAngle(float Angle) const;
 
-	/** How wide this panel should appear to be. */
-	TAttribute<float> PreferredWidth;
-
 	/** The slots that contain this panel's children. */
 	TPanelChildren<FSlot> Slots;
 
-	/** If true the box will have a preferred width equal to its alloted width  */
-	bool bUseAllottedWidth;
-
+	/** How wide this panel should appear to be. */
+	TSlateAttribute<float, EInvalidateWidgetReason::Layout> PreferredWidth;
 	/** Offset of the first element in the circle in degrees */
 	float StartingAngle;
 
 	/** If we need a section of a radial (for example half-a-radial) we can define a central angle < 360 (180 in case of half-a-radial). Used when bDistributeItemsEvenly is enabled. */
 	float SectorCentralAngle;
 
+	/** How many degrees apart should the elements be? */
+	float AngleBetweenItems;
+
 	/** Ignore AngleBetweenItems and distribute items evenly inside the whole circle */
 	bool bDistributeItemsEvenly;
 
-	/** How many degrees apart should the elements be? */
-	float AngleBetweenItems;
+	/** If true the box will have a preferred width equal to its alloted width  */
+	bool bUseAllottedWidth;
 
 	class FChildArranger;
 	friend class SRadialBox::FChildArranger;
diff --git a/Engine/Source/Runtime/Slate/Public/Widgets/Layout/SSafeZone.h b/Engine/Source/Runtime/Slate/Public/Widgets/Layout/SSafeZone.h
--- a/Engine/Source/Runtime/Slate/Public/Widgets/Layout/SSafeZone.h
+++ b/Engine/Source/Runtime/Slate/Public/Widgets/Layout/SSafeZone.h
@@ -90,11 +90,7 @@ class SLATE_API SSafeZone : public SBox
 
 public:
 
-	SSafeZone()
-	{
-		SetCanTick(false);
-		bCanSupportFocus = false;
-	}
+	SSafeZone();
 	virtual ~SSafeZone();
 
 	void Construct( const FArguments& InArgs );
@@ -123,7 +119,7 @@ private:
 	FMargin ComputeScaledSafeMargin(float Scale) const;
 
 	/** Cached values from the args */
-	TAttribute<FMargin> Padding;
+	TSlateAttribute<FMargin, EInvalidateWidgetReason::Layout> Padding;
 	FMargin SafeAreaScale;
 	bool bIsTitleSafe;
 	bool bPadLeft;
diff --git a/Engine/Source/Runtime/Slate/Public/Widgets/Layout/SSplitter.h b/Engine/Source/Runtime/Slate/Public/Widgets/Layout/SSplitter.h
--- a/Engine/Source/Runtime/Slate/Public/Widgets/Layout/SSplitter.h
+++ b/Engine/Source/Runtime/Slate/Public/Widgets/Layout/SSplitter.h
@@ -366,7 +366,7 @@ protected:
 	TPanelChildren< FSlot > Children;
 
 	int32 HoveredHandleIndex;
-	TAttribute<int32> HighlightedHandleIndex;
+	TSlateAttribute<int32, EInvalidateWidgetReason::Paint> HighlightedHandleIndex;
 	bool bIsResizing;
 	EOrientation Orientation;
 	ESplitterResizeMode::Type ResizeMode;
diff --git a/Engine/Source/Runtime/Slate/Public/Widgets/Layout/SUniformGridPanel.h b/Engine/Source/Runtime/Slate/Public/Widgets/Layout/SUniformGridPanel.h
--- a/Engine/Source/Runtime/Slate/Public/Widgets/Layout/SUniformGridPanel.h
+++ b/Engine/Source/Runtime/Slate/Public/Widgets/Layout/SUniformGridPanel.h
@@ -146,12 +146,12 @@ protected:
 
 private:
 	TPanelChildren<FSlot> Children;
-	TAttribute<FMargin> SlotPadding;
+	TSlateAttribute<FMargin, EInvalidateWidgetReason::Layout> SlotPadding;
 	
 	/** These values are recomputed and cached during compute desired size, as they may have changed since the previous frame. */
 	mutable int32 NumColumns;
 	mutable int32 NumRows;
 
-	TAttribute<float> MinDesiredSlotWidth;
-	TAttribute<float> MinDesiredSlotHeight;
+	TSlateAttribute<float, EInvalidateWidgetReason::Layout> MinDesiredSlotWidth;
+	TSlateAttribute<float, EInvalidateWidgetReason::Layout> MinDesiredSlotHeight;
 };
diff --git a/Engine/Source/Runtime/Slate/Public/Widgets/Layout/SUniformWrapPanel.h b/Engine/Source/Runtime/Slate/Public/Widgets/Layout/SUniformWrapPanel.h
--- a/Engine/Source/Runtime/Slate/Public/Widgets/Layout/SUniformWrapPanel.h
+++ b/Engine/Source/Runtime/Slate/Public/Widgets/Layout/SUniformWrapPanel.h
@@ -17,6 +17,8 @@ class FArrangedChildren;
 /** A panel that evenly divides up available space between all of its children. */
 class SLATE_API SUniformWrapPanel : public SPanel
 {
+	SLATE_DECLARE_WIDGET(SUniformWrapPanel, SPanel)
+
 public:
 	/** Stores the per-child info for this panel type */
 	struct FSlot : public TSlotBase<FSlot>, public TAlignmentWidgetSlotMixin<FSlot>
@@ -137,19 +139,19 @@ protected:
 
 private:
 	TPanelChildren<FSlot> Children;
-	TAttribute<FMargin> SlotPadding;
+	TSlateAttribute<FMargin> SlotPadding;
 	
 	/** These values are recomputed and cached during compute desired size, as they may have changed since the previous frame. */
 	mutable int32 NumColumns;
 	mutable int32 NumRows;
 	mutable int32 NumVisibleChildren;
 
-	TAttribute<float> MinDesiredSlotWidth;
-	TAttribute<float> MinDesiredSlotHeight;
+	TSlateAttribute<float> MinDesiredSlotWidth;
+	TSlateAttribute<float> MinDesiredSlotHeight;
 
-	TAttribute<float> MaxDesiredSlotWidth;
-	TAttribute<float> MaxDesiredSlotHeight;
+	TSlateAttribute<float> MaxDesiredSlotWidth;
+	TSlateAttribute<float> MaxDesiredSlotHeight;
 
-	TAttribute< EHorizontalAlignment > HAlign;
-	TAttribute<bool> EvenRowDistribution;
+	TSlateAttribute<EHorizontalAlignment> HAlign;
+	TSlateAttribute<bool> EvenRowDistribution;
 };
