Slate: Upgrade some SWidget to use TSlateAttribute (SExpanderArrow). Fix unused SLATE_ATTRIBUTE args (SToolTip)

diff --git a/Engine/Plugins/Enterprise/DataprepEditor/Source/DataprepEditor/Private/Widgets/SDataprepPalette.cpp b/Engine/Plugins/Enterprise/DataprepEditor/Source/DataprepEditor/Private/Widgets/SDataprepPalette.cpp
--- a/Engine/Plugins/Enterprise/DataprepEditor/Source/DataprepEditor/Private/Widgets/SDataprepPalette.cpp
+++ b/Engine/Plugins/Enterprise/DataprepEditor/Source/DataprepEditor/Private/Widgets/SDataprepPalette.cpp
@@ -66,7 +66,7 @@ public:
 			StyleSet = InArgs._StyleSet;
 			IndentAmount = InArgs._IndentAmount;
 			BaseIndentLevel = InArgs._BaseIndentLevel;
-			ShouldDrawWires = InArgs._ShouldDrawWires;
+			ShouldDrawWires.Assign(*this, InArgs._ShouldDrawWires);
 
 			ChildSlot
 			.Padding( FMargin( 3.0f, 12.0f, 4.0f, 9.0f ) )
diff --git a/Engine/Source/Runtime/Slate/Private/Widgets/SToolTip.cpp b/Engine/Source/Runtime/Slate/Private/Widgets/SToolTip.cpp
--- a/Engine/Source/Runtime/Slate/Private/Widgets/SToolTip.cpp
+++ b/Engine/Source/Runtime/Slate/Private/Widgets/SToolTip.cpp
@@ -23,7 +23,6 @@ void SToolTip::Construct( const FArguments& InArgs )
 	TextContent = InArgs._Text;
 	bIsInteractive = InArgs._IsInteractive;
 	Font = InArgs._Font;
-	ColorAndOpacity = InArgs._ColorAndOpacity;
 	TextMargin = InArgs._TextMargin;
 	BorderImage = InArgs._BorderImage;
 	OnSetInteractiveWindowLocationDelegate = InArgs._OnSetInteractiveWindowLocation;
diff --git a/Engine/Source/Runtime/Slate/Private/Widgets/Views/SExpanderArrow.cpp b/Engine/Source/Runtime/Slate/Private/Widgets/Views/SExpanderArrow.cpp
--- a/Engine/Source/Runtime/Slate/Private/Widgets/Views/SExpanderArrow.cpp
+++ b/Engine/Source/Runtime/Slate/Private/Widgets/Views/SExpanderArrow.cpp
@@ -6,6 +6,17 @@
 #include "Widgets/Input/SButton.h"
 #include "Widgets/Views/STableRow.h"
 
+SLATE_IMPLEMENT_WIDGET(SExpanderArrow)
+void SExpanderArrow::PrivateRegisterAttributes(FSlateAttributeInitializer& AttributeInitializer)
+{
+	SLATE_ADD_MEMBER_ATTRIBUTE_DEFINITION(AttributeInitializer, ShouldDrawWires, EInvalidateWidgetReason::Paint);
+}
+
+SExpanderArrow::SExpanderArrow()
+	: ShouldDrawWires(*this, false)
+{
+
+}
 
 void SExpanderArrow::Construct( const FArguments& InArgs, const TSharedPtr<class ITableRow>& TableRow  )
 {
@@ -13,7 +24,7 @@ void SExpanderArrow::Construct( const FArguments& InArgs, const TSharedPtr<class
 	StyleSet = InArgs._StyleSet;
 	IndentAmount = InArgs._IndentAmount;
 	BaseIndentLevel = InArgs._BaseIndentLevel;
-	ShouldDrawWires = InArgs._ShouldDrawWires;
+	ShouldDrawWires.Assign(*this, InArgs._ShouldDrawWires);
 
 	this->ChildSlot
 	.Padding( TAttribute<FMargin>( this, &SExpanderArrow::GetExpanderPadding ) )
diff --git a/Engine/Source/Runtime/Slate/Public/Widgets/SToolTip.h b/Engine/Source/Runtime/Slate/Public/Widgets/SToolTip.h
--- a/Engine/Source/Runtime/Slate/Public/Widgets/SToolTip.h
+++ b/Engine/Source/Runtime/Slate/Public/Widgets/SToolTip.h
@@ -27,7 +27,6 @@ public:
 		: _Text()
 		, _Content()
 		, _Font(FCoreStyle::Get().GetFontStyle("ToolTip.Font"))
-		, _ColorAndOpacity( FSlateColor::UseForeground())
 		, _TextMargin(FMargin(8.0f))
 		, _BorderImage(FCoreStyle::Get().GetBrush("ToolTip.Background"))
 		, _IsInteractive(false)
@@ -43,9 +42,6 @@ public:
 		/** The font to use for this tool tip */
 		SLATE_ATTRIBUTE(FSlateFontInfo, Font)
 		
-		/** Font color and opacity */
-		SLATE_ATTRIBUTE(FSlateColor, ColorAndOpacity)
-		
 		/** Margin between the tool tip border and the text content */
 		SLATE_ATTRIBUTE(FMargin, TextMargin)
 
@@ -112,9 +108,6 @@ private:
 
 	// Font used for the text displayed (where applicable)
 	TAttribute<FSlateFontInfo> Font;
-	
-	// Color and opacity used for the text displayed (where applicable)
-	TAttribute<FSlateColor> ColorAndOpacity;
 
 	// Margin between the tool tip border and the text content
 	TAttribute<FMargin> TextMargin;
diff --git a/Engine/Source/Runtime/Slate/Public/Widgets/Views/SExpanderArrow.h b/Engine/Source/Runtime/Slate/Public/Widgets/Views/SExpanderArrow.h
--- a/Engine/Source/Runtime/Slate/Public/Widgets/Views/SExpanderArrow.h
+++ b/Engine/Source/Runtime/Slate/Public/Widgets/Views/SExpanderArrow.h
@@ -20,6 +20,7 @@ class SButton;
  */
 class SLATE_API SExpanderArrow : public SCompoundWidget
 {
+	SLATE_DECLARE_WIDGET(SExpanderArrow, SCompoundWidget)
 public:
 
 	SLATE_BEGIN_ARGS( SExpanderArrow )
@@ -37,6 +38,7 @@ public:
 		SLATE_ATTRIBUTE(bool, ShouldDrawWires)
 	SLATE_END_ARGS()
 
+	SExpanderArrow();
 	void Construct( const FArguments& InArgs, const TSharedPtr<class ITableRow>& TableRow );
 
 protected:
@@ -70,5 +72,5 @@ protected:
 	TAttribute<int32> BaseIndentLevel;
 
 	/** Whether to draw the wires that visually reinforce the tree hierarchy. */
-	TAttribute<bool> ShouldDrawWires;
+	TSlateAttribute<bool> ShouldDrawWires;
 };
