Allow links in interactive tooltips to be clickable when the tooltip is spawned from modal windows and menu items.

diff --git a/Engine/Source/Runtime/Slate/Private/Framework/Application/MenuStack.cpp b/Engine/Source/Runtime/Slate/Private/Framework/Application/MenuStack.cpp
--- a/Engine/Source/Runtime/Slate/Private/Framework/Application/MenuStack.cpp
+++ b/Engine/Source/Runtime/Slate/Private/Framework/Application/MenuStack.cpp
@@ -780,7 +780,7 @@ void FMenuStack::OnWindowDestroyed(TSharedRef<SWindow> InWindow)
 
 void FMenuStack::OnWindowActivated( TSharedRef<SWindow> ActivatedWindow )
 {
-	if (ActivatedWindow != PendingNewWindow && HasMenus())
+	if (ActivatedWindow != PendingNewWindow && HasMenus() && !FSlateApplication::Get().IsWindowHousingInteractiveTooltip(ActivatedWindow))
 	{
 		TWeakPtr<IMenu> ActivatedMenu = FindMenuFromWindow(ActivatedWindow);
 
diff --git a/Engine/Source/Runtime/Slate/Private/Framework/Application/SlateApplication.cpp b/Engine/Source/Runtime/Slate/Private/Framework/Application/SlateApplication.cpp
--- a/Engine/Source/Runtime/Slate/Private/Framework/Application/SlateApplication.cpp
+++ b/Engine/Source/Runtime/Slate/Private/Framework/Application/SlateApplication.cpp
@@ -2692,6 +2692,13 @@ bool FSlateApplication::SetUserFocus(FSlateUser& User, const FWidgetPath& InFocu
 	if (InFocusPath.IsValid())
 	{
 		TSharedRef<SWindow> Window = InFocusPath.GetWindow();
+
+		// Prevent interactions with tooltips from disrupting the current focus state and closing open menus.
+		if (IsWindowHousingInteractiveTooltip(Window))
+		{
+			return false;
+		}
+
 		if (ActiveModalWindows.Num() != 0 && !(Window->IsDescendantOf(GetActiveModalWindow()) || ActiveModalWindows.Top() == Window))
 		{
 #if WITH_SLATE_DEBUGGING
@@ -3411,7 +3418,18 @@ TArray< TSharedRef<SWindow> > FSlateApplication::GetInteractiveTopLevelWindows()
 	{
 		// If we have modal windows, only the topmost modal window and its children are interactive.
 		TArray< TSharedRef<SWindow>, TInlineAllocator<1> > OutWindows;
-		OutWindows.Add( ActiveModalWindows.Last().ToSharedRef() );
+		OutWindows.Add(ActiveModalWindows.Last().ToSharedRef());
+
+		// If there is an interactive tooltip open from a modal window, include it too.
+		for (int32 WindowIndex = SlateWindows.Num() - 1; WindowIndex >= 0; WindowIndex--)
+		{
+			TSharedRef<SWindow> CurrentWindow = SlateWindows[WindowIndex];
+			if (GetCursorUser()->IsWindowHousingInteractiveTooltip(CurrentWindow))
+			{
+				OutWindows.Add(CurrentWindow);
+			}
+		}
+
 		return TArray< TSharedRef<SWindow> >(OutWindows);
 	}
 	else
@@ -4265,7 +4283,7 @@ bool FSlateApplication::ShouldProcessUserInputMessages( const TSharedPtr< FGener
 
 	if (ActiveModalWindows.Num() == 0 ||
 		(Window.IsValid() &&
-		(Window->IsDescendantOf(GetActiveModalWindow()) || ActiveModalWindows.Top() == Window)))
+		(Window->IsDescendantOf(GetActiveModalWindow()) || ActiveModalWindows.Top() == Window || IsWindowHousingInteractiveTooltip(Window.ToSharedRef()))))
 	{
 		return true;
 	}
@@ -6327,8 +6345,9 @@ bool FSlateApplication::ProcessWindowActivatedEvent( const FWindowActivateEvent&
 		// If you change this be sure to test windows are activated properly and receive input when they are opened when a modal dialog is open.
 		FSlateWindowHelper::BringWindowToFront(SlateWindows, ActivateEvent.GetAffectedWindow());
 
-		// Do not process activation messages unless we have no modal windows or the current window is modal
-		if( !ActiveModalWindow.IsValid() || ActivateEvent.GetAffectedWindow() == ActiveModalWindow || ActivateEvent.GetAffectedWindow()->IsDescendantOf(ActiveModalWindow) )
+		// Do not process activation messages unless we have no modal windows or the current window is modal or we are over an interactive tooltip
+		if ( !ActiveModalWindow.IsValid() || ActivateEvent.GetAffectedWindow() == ActiveModalWindow || ActivateEvent.GetAffectedWindow()->IsDescendantOf(ActiveModalWindow) 
+			|| IsWindowHousingInteractiveTooltip(ActivateEvent.GetAffectedWindow()) )
 		{
 			// Window being ACTIVATED
 			{
@@ -6345,7 +6364,6 @@ bool FSlateApplication::ProcessWindowActivatedEvent( const FWindowActivateEvent&
 			// A Slate window was activated
 			bSlateWindowActive = true;
 
-
 			{
 				FScopedSwitchWorldHack SwitchWorld( ActivateEvent.GetAffectedWindow() );
 				// let the menu stack know of new window being activated.  We may need to close menus as a result
diff --git a/Engine/Source/Runtime/SlateCore/Private/Widgets/SWindow.cpp b/Engine/Source/Runtime/SlateCore/Private/Widgets/SWindow.cpp
--- a/Engine/Source/Runtime/SlateCore/Private/Widgets/SWindow.cpp
+++ b/Engine/Source/Runtime/SlateCore/Private/Widgets/SWindow.cpp
@@ -1540,7 +1540,7 @@ EWindowActivationPolicy SWindow::ActivationPolicy() const
 /** @return true if the window accepts input; false if the window is non-interactive */
 bool SWindow::AcceptsInput() const
 {
-	return Type != EWindowType::CursorDecorator && Type != EWindowType::ToolTip;
+	return Type != EWindowType::CursorDecorator && (Type != EWindowType::ToolTip || !FSlateApplicationBase::Get().IsWindowHousingInteractiveTooltip(SharedThis(this)));
 }
 
 /** @return true if the user decides the size of the window; false if the content determines the size of the window */
