Prevent tab manager from summoning empty windows when no tabs can be opened

diff --git a/Engine/Source/Runtime/Slate/Private/Framework/Docking/TabManager.cpp b/Engine/Source/Runtime/Slate/Private/Framework/Docking/TabManager.cpp
--- a/Engine/Source/Runtime/Slate/Private/Framework/Docking/TabManager.cpp
+++ b/Engine/Source/Runtime/Slate/Private/Framework/Docking/TabManager.cpp
@@ -1615,6 +1615,7 @@ TSharedPtr<SDockingNode> FTabManager::RestoreArea_Helper(const TSharedRef<FLayou
 			NewStackWidget = SNew(SDockingTabStack, NodeAsStack.ToSharedRef());
 			NewStackWidget->SetSizeCoefficient(LayoutNode->GetSizeCoefficient());
 		}
+
 		// Open Tabs
 		for (const FTab& SomeTab : NodeAsStack->Tabs)
 		{
@@ -1707,28 +1708,31 @@ TSharedPtr<SDockingNode> FTabManager::RestoreArea_Helper(const TSharedRef<FLayou
 			TArray<TSharedRef<SDockingNode>> DockingNodes;
 			if (CanRestoreSplitterContent(DockingNodes, NodeAsArea.ToSharedRef(), NewWindow, OutSidebarTabs, OutputCanBeNullptr))
 			{
-				// Create SplitterWidget only if it will be filled with at least 1 DockingNodes
-				// Any windows that were "pulled out" of a dock area should be children of the window in which the parent dock area resides.
-				if (bIsChildWindow)
+				NewWindow->SetContent(SAssignNew(NewDockAreaWidget, SDockingArea, SharedThis(this), NodeAsArea.ToSharedRef()).ParentWindow(NewWindow));
+
+				// Restore content
+				if (!bCanOutputBeNullptr)
 				{
-					FSlateApplication::Get().AddWindowAsNativeChild(NewWindow, ParentWindow.ToSharedRef())->SetContent(
-						SAssignNew(NewDockAreaWidget, SDockingArea, SharedThis(this), NodeAsArea.ToSharedRef()).ParentWindow(NewWindow)
-					);
+					RestoreSplitterContent(NodeAsArea.ToSharedRef(), NewDockAreaWidget.ToSharedRef(), NewWindow, OutSidebarTabs);
 				}
 				else
 				{
-					FSlateApplication::Get().AddWindow(NewWindow)->SetContent(
-						SAssignNew(NewDockAreaWidget, SDockingArea, SharedThis(this), NodeAsArea.ToSharedRef()).ParentWindow(NewWindow)
-					);
+					RestoreSplitterContent(DockingNodes, NewDockAreaWidget.ToSharedRef());
 				}
-				// Restore content
-				if (!bCanOutputBeNullptr)
+
+				if (bIsChildWindow)
 				{
-					RestoreSplitterContent(NodeAsArea.ToSharedRef(), NewDockAreaWidget.ToSharedRef(), NewWindow, OutSidebarTabs);
+					// Recursively check to see how many actually spawned tabs there are in this dock area. If there are none we will not spawn a useless window
+					const int32 TotalNumTabs = NewDockAreaWidget->GetNumTabs();
+
+					if (TotalNumTabs > 0)
+					{
+						FSlateApplication::Get().AddWindowAsNativeChild(NewWindow, ParentWindow.ToSharedRef());
+					}
 				}
 				else
 				{
-					RestoreSplitterContent(DockingNodes, NewDockAreaWidget.ToSharedRef());
+					FSlateApplication::Get().AddWindow(NewWindow);
 				}
 			}
 		}
