[Backout] - CL18312422
Original CL Desc
-----------------------------------------------------------------
Prevent tab manager from summoning empty windows when no tabs can be opened

diff --git a/Engine/Source/Runtime/Slate/Private/Framework/Docking/TabManager.cpp b/Engine/Source/Runtime/Slate/Private/Framework/Docking/TabManager.cpp
--- a/Engine/Source/Runtime/Slate/Private/Framework/Docking/TabManager.cpp
+++ b/Engine/Source/Runtime/Slate/Private/Framework/Docking/TabManager.cpp
@@ -1615,7 +1615,6 @@ TSharedPtr<SDockingNode> FTabManager::RestoreArea_Helper(const TSharedRef<FLayou
 			NewStackWidget = SNew(SDockingTabStack, NodeAsStack.ToSharedRef());
 			NewStackWidget->SetSizeCoefficient(LayoutNode->GetSizeCoefficient());
 		}
-
 		// Open Tabs
 		for (const FTab& SomeTab : NodeAsStack->Tabs)
 		{
@@ -1708,31 +1707,28 @@ TSharedPtr<SDockingNode> FTabManager::RestoreArea_Helper(const TSharedRef<FLayou
 			TArray<TSharedRef<SDockingNode>> DockingNodes;
 			if (CanRestoreSplitterContent(DockingNodes, NodeAsArea.ToSharedRef(), NewWindow, OutSidebarTabs, OutputCanBeNullptr))
 			{
-				NewWindow->SetContent(SAssignNew(NewDockAreaWidget, SDockingArea, SharedThis(this), NodeAsArea.ToSharedRef()).ParentWindow(NewWindow));
-
-				// Restore content
-				if (!bCanOutputBeNullptr)
+				// Create SplitterWidget only if it will be filled with at least 1 DockingNodes
+				// Any windows that were "pulled out" of a dock area should be children of the window in which the parent dock area resides.
+				if (bIsChildWindow)
 				{
-					RestoreSplitterContent(NodeAsArea.ToSharedRef(), NewDockAreaWidget.ToSharedRef(), NewWindow, OutSidebarTabs);
+					FSlateApplication::Get().AddWindowAsNativeChild(NewWindow, ParentWindow.ToSharedRef())->SetContent(
+						SAssignNew(NewDockAreaWidget, SDockingArea, SharedThis(this), NodeAsArea.ToSharedRef()).ParentWindow(NewWindow)
+					);
 				}
 				else
 				{
-					RestoreSplitterContent(DockingNodes, NewDockAreaWidget.ToSharedRef());
+					FSlateApplication::Get().AddWindow(NewWindow)->SetContent(
+						SAssignNew(NewDockAreaWidget, SDockingArea, SharedThis(this), NodeAsArea.ToSharedRef()).ParentWindow(NewWindow)
+					);
 				}
-
-				if (bIsChildWindow)
+				// Restore content
+				if (!bCanOutputBeNullptr)
 				{
-					// Recursively check to see how many actually spawned tabs there are in this dock area. If there are none we will not spawn a useless window
-					const int32 TotalNumTabs = NewDockAreaWidget->GetNumTabs();
-
-					if (TotalNumTabs > 0)
-					{
-						FSlateApplication::Get().AddWindowAsNativeChild(NewWindow, ParentWindow.ToSharedRef());
-					}
+					RestoreSplitterContent(NodeAsArea.ToSharedRef(), NewDockAreaWidget.ToSharedRef(), NewWindow, OutSidebarTabs);
 				}
 				else
 				{
-					FSlateApplication::Get().AddWindow(NewWindow);
+					RestoreSplitterContent(DockingNodes, NewDockAreaWidget.ToSharedRef());
 				}
 			}
 		}
