Input Widgets: Overhaul numeric, vector and rotator widgets

* Introduce mapping to vectors as well as numeric components
* Introduce limits / slider limits to the vector input box
* Introduce constraining component to vector input box

diff --git a/Engine/Source/Runtime/Slate/Public/Widgets/Input/SNumericEntryBox.h b/Engine/Source/Runtime/Slate/Public/Widgets/Input/SNumericEntryBox.h
--- a/Engine/Source/Runtime/Slate/Public/Widgets/Input/SNumericEntryBox.h
+++ b/Engine/Source/Runtime/Slate/Public/Widgets/Input/SNumericEntryBox.h
@@ -23,6 +23,7 @@
 #include "Widgets/Text/STextBlock.h"
 #include "Widgets/Input/SEditableText.h"
 #include "Widgets/Input/SSpinBox.h"
+#include "Widgets/Input/SCheckBox.h"
 #include "Widgets/Colors/SColorBlock.h"
 #include "Containers/UnrealString.h"
 
@@ -39,6 +40,7 @@ public:
 	static const FLinearColor RedLabelBackgroundColor;
 	static const FLinearColor GreenLabelBackgroundColor;
 	static const FLinearColor BlueLabelBackgroundColor;
+	static const FLinearColor LilacLabelBackgroundColor;
 	static const FText DefaultUndeterminedString;
 
 	/** Notification for numeric value change */
@@ -88,6 +90,9 @@ public:
 		, _MaxSliderValue(NumericType(100))
 		, _SliderExponent(1.f)
 		, _MinDesiredValueWidth(0.f)
+		, _DisplayToggle(false)
+		, _ToggleChecked(ECheckBoxState::Checked)
+		, _TogglePadding(FMargin(1.f,0.f,1.f,0.f) )
 	{}		
 
 		/** Style to use for the editable text box within this widget */
@@ -164,6 +169,14 @@ public:
 		SLATE_EVENT( FMenuExtensionDelegate, ContextMenuExtender )
 		/** Provide custom type conversion functionality to this spin box */
 		SLATE_ARGUMENT( TSharedPtr< INumericTypeInterface<NumericType> >, TypeInterface )
+		/** Whether or not to include a toggle checkbox to the left of the widget */
+		SLATE_ARGUMENT( bool, DisplayToggle )
+		/** The value of the toggle checkbox */
+		SLATE_ATTRIBUTE( ECheckBoxState, ToggleChecked )
+		/** Called whenever the toggle changes state */
+		SLATE_EVENT( FOnCheckStateChanged, OnToggleChanged )
+		/** Padding around the toggle checkbox */
+		SLATE_ARGUMENT( FMargin, TogglePadding )
 
 	SLATE_END_ARGS()
 	SNumericEntryBox()
@@ -192,6 +205,14 @@ public:
 			CachedValueString = Interface->ToString(CachedExternalValue.GetValue());
 		}
 
+		const bool bDisplayToggle = InArgs._DisplayToggle;
+		if(bDisplayToggle)
+		{
+			SAssignNew(ToggleCheckBox, SCheckBox)
+			.Padding(FMargin(0.f, 0.f, 2.f, 0.f))
+			.IsChecked(InArgs._ToggleChecked)
+			.OnCheckStateChanged(this, &SNumericEntryBox::HandleToggleCheckBoxChanged, InArgs._OnToggleChanged);
+		}
 
 		const bool bAllowSpin = InArgs._AllowSpin;
 		TSharedPtr<SWidget> FinalWidget;
@@ -299,30 +320,73 @@ public:
 
 		if(!bHasLabel || bHasInsideLabel)
 		{
-			ChildSlot
-			[
-				MainContents.ToSharedRef()
-			];
+			if(bDisplayToggle)
+			{
+				ChildSlot
+				[
+					SNew(SHorizontalBox)
+					+SHorizontalBox::Slot()
+					.AutoWidth()
+					.HAlign(HAlign_Fill) 
+					.VAlign(VAlign_Center) 
+					.Padding(InArgs._TogglePadding)
+					[
+						ToggleCheckBox.ToSharedRef()
+					]
+					+ SHorizontalBox::Slot()
+					[
+						MainContents.ToSharedRef()
+					]
+				];
+			}
+			else
+			{
+				ChildSlot
+				[
+					MainContents.ToSharedRef()
+				];
+			}
 		}
 		else
 		{
-			ChildSlot
+			TSharedRef<SHorizontalBox> HorizontalBox = SNew(SHorizontalBox);
+
+			HorizontalBox->AddSlot()
+			.AutoWidth()
+			.HAlign(HAlign_Left)
+			.VAlign(InArgs._LabelVAlign)
+			.Padding(InArgs._LabelPadding)
 			[
-				SNew(SHorizontalBox)
-				+SHorizontalBox::Slot()
+				InArgs._Label.Widget
+			];
+
+			if(bDisplayToggle)
+			{
+				HorizontalBox->AddSlot()
 				.AutoWidth()
-				.HAlign(HAlign_Left)
-				.VAlign(InArgs._LabelVAlign)
-				.Padding(InArgs._LabelPadding)
-				[
-					InArgs._Label.Widget
-				]
-				+ SHorizontalBox::Slot()
+				.HAlign(HAlign_Fill)
+				.VAlign(VAlign_Center)
+				.Padding(InArgs._TogglePadding)
 				[
-					MainContents.ToSharedRef()
-				]
+					ToggleCheckBox.ToSharedRef()
+				];
+			}
+			
+			HorizontalBox->AddSlot()
+			[
+				MainContents.ToSharedRef()
+			];
+
+			ChildSlot
+			[
+				HorizontalBox
 			];
 		}
+
+		if (bDisplayToggle)
+		{
+			HandleToggleCheckBoxChanged(InArgs._ToggleChecked.Get(), InArgs._OnToggleChanged);
+		}
 	}
 
 	static TSharedRef<SWidget> BuildLabel(TAttribute<FText> LabelText, const FSlateColor& ForegroundColor, const FSlateColor& BackgroundColor)
@@ -610,12 +674,39 @@ private:
 		}
 	}
 
+	bool IsToggleEnabled() const
+	{
+		if(!IsEnabled())
+		{
+			return false;
+		}
+		return ToggleCheckBox->IsChecked();
+	}
+	
+	void HandleToggleCheckBoxChanged(ECheckBoxState InCheckState, FOnCheckStateChanged OnToggleChanged) const
+	{
+		if(SpinBox.IsValid())
+		{
+			SpinBox->SetEnabled(InCheckState == ECheckBoxState::Checked);
+		}
+		if(EditableText.IsValid())
+		{
+			EditableText->SetEnabled(InCheckState == ECheckBoxState::Checked);
+		}
+		if(OnToggleChanged.IsBound())
+		{
+			OnToggleChanged.Execute(InCheckState);
+		}
+	}
+
 private:
 
 	/** Attribute for getting the label */
 	TAttribute< TOptional<FString > > LabelAttribute;
 	/** Attribute for getting the value.  If the value is not set we display the undetermined string */
 	TAttribute< TOptional<NumericType> > ValueAttribute;
+	/** Toggle checkbox */
+	TSharedPtr<SCheckBox> ToggleCheckBox;
 	/** Spinbox widget */
 	TSharedPtr<SWidget> SpinBox;
 	/** Editable widget */
@@ -656,5 +747,8 @@ const FLinearColor SNumericEntryBox<NumericType>::GreenLabelBackgroundColor(0.13
 template <typename NumericType>
 const FLinearColor SNumericEntryBox<NumericType>::BlueLabelBackgroundColor(0.0251f,0.207f,0.85f);
 
+template <typename NumericType>
+const FLinearColor SNumericEntryBox<NumericType>::LilacLabelBackgroundColor(0.8f,0.121f,0.8f);
+
 template <typename NumericType>
 const FText SNumericEntryBox<NumericType>::DefaultUndeterminedString = FText::FromString(TEXT("---"));
diff --git a/Engine/Source/Runtime/Slate/Public/Widgets/Input/SRotatorInputBox.h b/Engine/Source/Runtime/Slate/Public/Widgets/Input/SRotatorInputBox.h
--- a/Engine/Source/Runtime/Slate/Public/Widgets/Input/SRotatorInputBox.h
+++ b/Engine/Source/Runtime/Slate/Public/Widgets/Input/SRotatorInputBox.h
@@ -32,6 +32,11 @@ public:
 		: _bColorAxisLabels(false)
 		, _Font(FAppStyle::Get().GetFontStyle("NormalFont"))
 		, _AllowSpin(true)
+		, _DisplayToggle( false )
+		, _TogglePitchChecked( ECheckBoxState::Checked )
+		, _ToggleYawChecked( ECheckBoxState::Checked )
+		, _ToggleRollChecked( ECheckBoxState::Checked )
+		, _TogglePadding( FMargin( 1.f,0.f,1.f,0.f ) )
 		{}
 
 		/** Roll component of the rotator */
@@ -85,6 +90,30 @@ public:
 		/** Provide custom type functionality for the rotator */
 		SLATE_ARGUMENT( TSharedPtr< INumericTypeInterface<NumericType> >, TypeInterface )
 
+		/** Whether or not to include a toggle checkbox to the left of the widget */
+		SLATE_ARGUMENT( bool, DisplayToggle )
+			
+		/** The value of the toggle Pitch checkbox */
+		SLATE_ATTRIBUTE( ECheckBoxState, TogglePitchChecked )
+
+		/** The value of the toggle Yaw checkbox */
+		SLATE_ATTRIBUTE( ECheckBoxState, ToggleYawChecked )
+
+		/** The value of the toggle Roll checkbox */
+		SLATE_ATTRIBUTE( ECheckBoxState, ToggleRollChecked )
+
+		/** Called whenever the toggle Pitch changes state */
+		SLATE_EVENT( FOnCheckStateChanged, OnTogglePitchChanged )
+
+		/** Called whenever the toggle Yaw changes state */
+		SLATE_EVENT( FOnCheckStateChanged, OnToggleYawChanged )
+
+		/** Called whenever the toggle Roll changes state */
+		SLATE_EVENT( FOnCheckStateChanged, OnToggleRollChanged )
+
+		/** Padding around the toggle checkbox */
+		SLATE_ARGUMENT( FMargin, TogglePadding )
+
 	SLATE_END_ARGS()
 
 	/**
@@ -124,6 +153,10 @@ public:
 						: NSLOCTEXT("SRotatorInputBox", "MultipleValues", "Multiple Values");
 				})
 				.TypeInterface(InArgs._TypeInterface)
+				.DisplayToggle(InArgs._DisplayToggle)
+				.ToggleChecked(InArgs._ToggleRollChecked)
+				.OnToggleChanged(InArgs._OnToggleRollChanged)
+				.TogglePadding(InArgs._TogglePadding)
 			]
 			+SHorizontalBox::Slot()
 			[
@@ -152,6 +185,10 @@ public:
 						: NSLOCTEXT("SRotatorInputBox", "MultipleValues", "Multiple Values");
 				})
 				.TypeInterface(InArgs._TypeInterface)
+				.DisplayToggle(InArgs._DisplayToggle)
+				.ToggleChecked(InArgs._TogglePitchChecked)
+				.OnToggleChanged(InArgs._OnTogglePitchChanged)
+				.TogglePadding(InArgs._TogglePadding)
 			]
 			+SHorizontalBox::Slot()
 			[
@@ -180,6 +217,10 @@ public:
 						: NSLOCTEXT("SRotatorInputBox", "MultipleValues", "Multiple Values");
 				})
 				.TypeInterface(InArgs._TypeInterface)
+				.DisplayToggle(InArgs._DisplayToggle)
+				.ToggleChecked(InArgs._ToggleYawChecked)
+				.OnToggleChanged(InArgs._OnToggleYawChanged)
+				.TogglePadding(InArgs._TogglePadding)
 			]
 		];
 
diff --git a/Engine/Source/Runtime/Slate/Public/Widgets/Input/SVectorInputBox.h b/Engine/Source/Runtime/Slate/Public/Widgets/Input/SVectorInputBox.h
--- a/Engine/Source/Runtime/Slate/Public/Widgets/Input/SVectorInputBox.h
+++ b/Engine/Source/Runtime/Slate/Public/Widgets/Input/SVectorInputBox.h
@@ -25,12 +25,21 @@ template<typename NumericType, typename VectorType = UE::Math::TVector<NumericTy
 class SNumericVectorInputBox : public SCompoundWidget
 {
 public:
-	/** Notification for float value change */
+	/** Notification for numeric value change */
 	DECLARE_DELEGATE_OneParam(FOnNumericValueChanged, NumericType);
 
-	/** Notification for float value committed */
+	/** Notification for numeric value committed */
 	DECLARE_DELEGATE_TwoParams(FOnNumericValueCommitted, NumericType, ETextCommit::Type);
 
+	/** Notification for vector value change */
+	DECLARE_DELEGATE_OneParam(FOnVectorValueChanged, VectorType);
+
+	/** Notification for vector value committed */
+	DECLARE_DELEGATE_TwoParams(FOnVectorValueCommitted, VectorType, ETextCommit::Type);
+
+	/** Delegate to constrain the vector during a change */
+	DECLARE_DELEGATE_ThreeParams(FOnConstrainVector, int32 /* Component */, VectorType /* old */ , VectorType& /* new */);
+
 	struct FArguments;
 
 private:
@@ -42,6 +51,10 @@ private:
 	{
 		using WidgetArgsType = ArgumentType;
 
+		FORCENOINLINE FVectorXArguments()
+			: _ToggleXChecked(ECheckBoxState::Checked)
+		{}
+
 		/** X Component of the vector */
 		SLATE_ATTRIBUTE(TOptional<NumericType>, X)
 
@@ -51,6 +64,12 @@ private:
 		/** Called when the x value of the vector is committed */
 		SLATE_EVENT(FOnNumericValueCommitted, OnXCommitted)
 
+		/** The value of the toggle X checkbox */
+		SLATE_ATTRIBUTE( ECheckBoxState, ToggleXChecked )
+
+		/** Called whenever the toggle X changes state */
+		SLATE_EVENT( FOnCheckStateChanged, OnToggleXChanged )
+
 		/** Menu extender delegate for the X value */
 		SLATE_EVENT(FMenuExtensionDelegate, ContextMenuExtenderX)
 	};
@@ -61,6 +80,10 @@ private:
 	{
 		using WidgetArgsType = ArgumentType;
 
+		FORCENOINLINE FVectorYArguments()
+			: _ToggleYChecked(ECheckBoxState::Checked)
+		{}
+
 		/** Y Component of the vector */
 		SLATE_ATTRIBUTE(TOptional<NumericType>, Y)
 
@@ -70,6 +93,12 @@ private:
 		/** Called when the Y value of the vector is committed */
 		SLATE_EVENT(FOnNumericValueCommitted, OnYCommitted)
 
+		/** The value of the toggle Y checkbox */
+		SLATE_ATTRIBUTE( ECheckBoxState, ToggleYChecked )
+
+		/** Called whenever the toggle Y changes state */
+		SLATE_EVENT( FOnCheckStateChanged, OnToggleYChanged )
+
 		/** Menu extender delegate for the Y value */
 		SLATE_EVENT(FMenuExtensionDelegate, ContextMenuExtenderY)
 	};
@@ -80,6 +109,10 @@ private:
 	{
 		using WidgetArgsType = ArgumentType;
 
+		FORCENOINLINE FVectorZArguments()
+			: _ToggleZChecked(ECheckBoxState::Checked)
+		{}
+
 		/** Z Component of the vector */
 		SLATE_ATTRIBUTE(TOptional<NumericType>, Z)
 
@@ -89,6 +122,12 @@ private:
 		/** Called when the Z value of the vector is committed */
 		SLATE_EVENT(FOnNumericValueCommitted, OnZCommitted)
 
+		/** The value of the toggle Z checkbox */
+		SLATE_ATTRIBUTE( ECheckBoxState, ToggleZChecked )
+
+		/** Called whenever the toggle Z changes state */
+		SLATE_EVENT( FOnCheckStateChanged, OnToggleZChanged )
+
 		/** Menu extender delegate for the Z value */
 		SLATE_EVENT(FMenuExtensionDelegate, ContextMenuExtenderZ)
 	};
@@ -99,6 +138,10 @@ private:
 	{
 		using WidgetArgsType = ArgumentType;
 
+		FORCENOINLINE FVectorWArguments()
+			: _ToggleWChecked(ECheckBoxState::Checked)
+		{}
+
 		/** W Component of the vector */
 		SLATE_ATTRIBUTE(TOptional<NumericType>, W)
 
@@ -108,6 +151,12 @@ private:
 		/** Called when the W value of the vector is committed */
 		SLATE_EVENT(FOnNumericValueCommitted, OnWCommitted)
 
+		/** The value of the toggle W checkbox */
+		SLATE_ATTRIBUTE( ECheckBoxState, ToggleWChecked )
+
+		/** Called whenever the toggle W changes state */
+		SLATE_EVENT( FOnCheckStateChanged, OnToggleWChanged )
+
 		/** Menu extender delegate for the W value */
 		SLATE_EVENT(FMenuExtensionDelegate, ContextMenuExtenderW)
 	};
@@ -122,12 +171,43 @@ public:
 	{
 		typedef FArguments WidgetArgsType;
 		FORCENOINLINE FArguments()
-			: _Font(FAppStyle::Get().GetFontStyle("NormalFont"))
+			: _EditableTextBoxStyle( &FAppStyle::Get().GetWidgetStyle<FEditableTextBoxStyle>("NormalEditableTextBox") )
+			, _SpinBoxStyle(&FAppStyle::Get().GetWidgetStyle<FSpinBoxStyle>("NumericEntrySpinBox") )
+			, _Font(FAppStyle::Get().GetFontStyle("NormalFont"))
 			, _AllowSpin(false)
 			, _SpinDelta(1)
 			, _bColorAxisLabels(false)
+			, _DisplayToggle(false)
+			, _TogglePadding(FMargin(1.f,0.f,1.f,0.f) )
 		{}
 
+		/** Optional Value of the vector */
+		SLATE_ATTRIBUTE(TOptional<VectorType>, Vector)
+
+		/** Optional minimum value of the vector */
+		SLATE_ATTRIBUTE(TOptional<VectorType>, MinVector)
+
+		/** Optional maximum value of the vector */
+		SLATE_ATTRIBUTE(TOptional<VectorType>, MaxVector)
+
+		/** Optional minimum (slider) value of the vector */
+		SLATE_ATTRIBUTE(TOptional<VectorType>, MinSliderVector)
+
+		/** Optional maximum (slider) value of the vector */
+		SLATE_ATTRIBUTE(TOptional<VectorType>, MaxSliderVector)
+
+		/** Called when the vector is changed */
+		SLATE_EVENT(FOnVectorValueChanged, OnVectorChanged)
+
+		/** Called when the vector is committed */
+		SLATE_EVENT(FOnVectorValueCommitted, OnVectorCommitted)
+
+		/** Style to use for the editable text box within this widget */
+		SLATE_STYLE_ARGUMENT( FEditableTextBoxStyle, EditableTextBoxStyle )
+
+		/** Style to use for the spin box within this widget */
+		SLATE_STYLE_ARGUMENT( FSpinBoxStyle, SpinBoxStyle )
+		
 		/** Font to use for the text in this box */
 		SLATE_ATTRIBUTE( FSlateFontInfo, Font )
 
@@ -156,6 +236,15 @@ public:
 		/** Provide custom type functionality for the vector */
 		SLATE_ARGUMENT( TSharedPtr< INumericTypeInterface<NumericType> >, TypeInterface )
 
+		/** Whether or not to include a toggle checkbox to the left of the widget */
+		SLATE_ARGUMENT( bool, DisplayToggle )
+			
+		/** Padding around the toggle checkbox */
+		SLATE_ARGUMENT( FMargin, TogglePadding )
+
+		/** Delegate to constrain the vector */
+		SLATE_ARGUMENT( FOnConstrainVector, ConstrainVector )
+		
 	SLATE_END_ARGS()
 
 	/**
@@ -194,13 +283,16 @@ private:
 	/**
 	 * Construct the widget component
 	 */
-	void ConstructComponent(const FArguments& InArgs,
+	void ConstructComponent(int32 ComponentIndex,
+		const FArguments& InArgs,
 		const FLinearColor& LabelColor,
 		const FText& TooltipText,
 		TSharedRef<SHorizontalBox>& HorizontalBox,
 		const TAttribute<TOptional<NumericType>>& Component,
 		const FOnNumericValueChanged& OnComponentChanged,
 		const FOnNumericValueCommitted& OnComponentCommitted,
+		const TAttribute<ECheckBoxState> ToggleChecked,
+		const FOnCheckStateChanged& OnToggleChanged,
 		const FMenuExtensionDelegate& OnContextMenuExtenderComponent)
 	{
 		TSharedRef<SWidget> LabelWidget = SNullWidget::NullWidget;
@@ -209,31 +301,33 @@ private:
 			LabelWidget = SNumericEntryBox<NumericType>::BuildNarrowColorLabel(LabelColor);
 		}
 
-		TAttribute<TOptional<NumericType>> Value = Component;
+		TAttribute<TOptional<NumericType>> Value = CreatePerComponentGetter(ComponentIndex, Component, InArgs._Vector);
 
 		HorizontalBox->AddSlot()
 		[
 			SNew(SNumericEntryBox<NumericType>)
 			.AllowSpin(InArgs._AllowSpin)
+			.EditableTextBoxStyle(InArgs._EditableTextBoxStyle)
+			.SpinBoxStyle(InArgs._SpinBoxStyle)
 			.Font(InArgs._Font)
-			.Value(Component)
-			.OnValueChanged(OnComponentChanged)
-			.OnValueCommitted(OnComponentCommitted)
-			.ToolTipText(MakeAttributeLambda([Component, TooltipText]
+			.Value(Value)
+			.OnValueChanged(CreatePerComponentChanged(ComponentIndex, OnComponentChanged, InArgs._Vector, InArgs._OnVectorChanged, InArgs._ConstrainVector))
+			.OnValueCommitted(CreatePerComponentCommitted(ComponentIndex, OnComponentCommitted, InArgs._Vector, InArgs._OnVectorCommitted, InArgs._ConstrainVector))
+			.ToolTipText(MakeAttributeLambda([Value, TooltipText]
 			{
-				if (Component.Get().IsSet())
+				if (Value.Get().IsSet())
 				{
-					return FText::Format(TooltipText, Component.Get().GetValue());
+					return FText::Format(TooltipText, Value.Get().GetValue());
 				}
 				return NSLOCTEXT("SVectorInputBox", "MultipleValues", "Multiple Values");
 			}))
 			.UndeterminedString(NSLOCTEXT("SVectorInputBox", "MultipleValues", "Multiple Values"))
 			.ContextMenuExtender(OnContextMenuExtenderComponent)
 			.TypeInterface(InArgs._TypeInterface)
-			.MinValue(TOptional<NumericType>())
-			.MaxValue(TOptional<NumericType>())
-			.MinSliderValue(TOptional<NumericType>())
-			.MaxSliderValue(TOptional<NumericType>())
+			.MinValue(CreatePerComponentGetter(ComponentIndex, TOptional<NumericType>(), InArgs._MinVector))
+			.MaxValue(CreatePerComponentGetter(ComponentIndex, TOptional<NumericType>(), InArgs._MaxVector))
+			.MinSliderValue(CreatePerComponentGetter(ComponentIndex, TOptional<NumericType>(), InArgs._MinSliderVector))
+			.MaxSliderValue(CreatePerComponentGetter(ComponentIndex, TOptional<NumericType>(), InArgs._MaxSliderVector))
 			.LinearDeltaSensitivity(1)
 			.Delta(InArgs._SpinDelta)
 			.OnBeginSliderMovement(InArgs._OnBeginSliderMovement)
@@ -244,21 +338,28 @@ private:
 			[
 				LabelWidget
 			]
-	];
-}
+			.DisplayToggle(InArgs._DisplayToggle)
+			.TogglePadding(InArgs._TogglePadding)
+			.ToggleChecked(ToggleChecked)
+			.OnToggleChanged(OnToggleChanged)
+		];
+	}
 
 	/**
 	 * Construct widgets for the X Value
 	 */
 	void ConstructX(const FArguments& InArgs, TSharedRef<SHorizontalBox> HorizontalBox)
 	{
-		ConstructComponent(InArgs,
+		ConstructComponent(0,
+			InArgs,
 			SNumericEntryBox<NumericType>::RedLabelBackgroundColor,
 			NSLOCTEXT("SVectorInputBox", "X_ToolTip", "X: {0}"),
 			HorizontalBox,
 			InArgs._X,
 			InArgs._OnXChanged,
 			InArgs._OnXCommitted,
+			InArgs._ToggleXChecked,
+			InArgs._OnToggleXChanged,
 			InArgs._ContextMenuExtenderX
 		);
 	}
@@ -268,13 +369,16 @@ private:
 	 */
 	void ConstructY(const FArguments& InArgs, TSharedRef<SHorizontalBox> HorizontalBox)
 	{
-		ConstructComponent(InArgs,
+		ConstructComponent(1,
+			InArgs,
 			SNumericEntryBox<NumericType>::GreenLabelBackgroundColor,
 			NSLOCTEXT("SVectorInputBox", "Y_ToolTip", "Y: {0}"),
 			HorizontalBox,
 			InArgs._Y,
 			InArgs._OnYChanged,
 			InArgs._OnYCommitted,
+			InArgs._ToggleYChecked,
+			InArgs._OnToggleYChanged,
 			InArgs._ContextMenuExtenderY
 		);
 	}
@@ -284,13 +388,16 @@ private:
 	 */
 	void ConstructZ(const FArguments& InArgs, TSharedRef<SHorizontalBox> HorizontalBox)
 	{
-		ConstructComponent(InArgs,
+		ConstructComponent(2,
+			InArgs,
 			SNumericEntryBox<NumericType>::BlueLabelBackgroundColor,
 			NSLOCTEXT("SVectorInputBox", "Z_ToolTip", "Z: {0}"),
 			HorizontalBox,
 			InArgs._Z,
 			InArgs._OnZChanged,
 			InArgs._OnZCommitted,
+			InArgs._ToggleZChecked,
+			InArgs._OnToggleZChanged,
 			InArgs._ContextMenuExtenderZ
 		);
 	}
@@ -300,16 +407,113 @@ private:
 	 */
 	void ConstructW(const FArguments& InArgs, TSharedRef<SHorizontalBox> HorizontalBox)
 	{
-		ConstructComponent(InArgs,
-			FLinearColor::Yellow,
+		ConstructComponent(3,
+			InArgs,
+			SNumericEntryBox<NumericType>::LilacLabelBackgroundColor,
 			NSLOCTEXT("SVectorInputBox", "W_ToolTip", "W: {0}"),
 			HorizontalBox,
 			InArgs._W,
 			InArgs._OnWChanged,
 			InArgs._OnWCommitted,
+			InArgs._ToggleWChecked,
+			InArgs._OnToggleWChanged,
 			InArgs._ContextMenuExtenderW
 		);
 	}
+
+	/*
+	 * Creates a lambda to retrieve a component off a vector
+	 */
+	TAttribute<TOptional<NumericType>> CreatePerComponentGetter(
+		int32 ComponentIndex,
+		const TAttribute<TOptional<NumericType>>& Component,
+		const TAttribute<TOptional<VectorType>>& Vector)
+	{
+		if(Vector.IsBound() || Vector.IsSet())
+		{
+			return TAttribute<TOptional<NumericType>>::CreateLambda(
+				[ComponentIndex, Component, Vector]() -> TOptional<NumericType>
+				{
+					const TOptional<VectorType> OptionalVectorValue = Vector.Get();
+					if(OptionalVectorValue.IsSet())
+					{
+						return OptionalVectorValue.GetValue()[ComponentIndex];
+					}
+					return Component.Get();
+				});
+		}
+		return Component;
+	}		
+
+	/*
+	 * Creates a lambda to react to a change event
+	 */
+	FOnNumericValueChanged CreatePerComponentChanged(
+		int32 ComponentIndex,
+		const FOnNumericValueChanged OnComponentChanged,
+		const TAttribute<TOptional<VectorType>>& Vector,
+		const FOnVectorValueChanged OnVectorValueChanged,
+		const FOnConstrainVector ConstrainVector)
+	{
+		if(OnVectorValueChanged.IsBound())
+		{
+			return FOnNumericValueChanged::CreateLambda(
+				[ComponentIndex, OnComponentChanged, Vector, OnVectorValueChanged, ConstrainVector](NumericType InValue)
+				{
+					OnComponentChanged.ExecuteIfBound(InValue);
+					
+					const TOptional<VectorType> OptionalVectorValue = Vector.Get();
+					if(OptionalVectorValue.IsSet())
+					{
+						VectorType VectorValue = OptionalVectorValue.GetValue();
+						VectorValue[ComponentIndex] = InValue;
+
+						if(ConstrainVector.IsBound())
+						{
+							ConstrainVector.Execute(ComponentIndex, OptionalVectorValue.GetValue(), VectorValue);
+						}
+
+						OnVectorValueChanged.Execute(VectorValue);
+					}
+				});
+		}
+		return OnComponentChanged;
+	}		
+
+	/*
+	 * Creates a lambda to react to a commit event
+	 */
+	FOnNumericValueCommitted CreatePerComponentCommitted(
+		int32 ComponentIndex,
+		const FOnNumericValueCommitted OnComponentCommitted,
+		const TAttribute<TOptional<VectorType>>& Vector,
+		const FOnVectorValueCommitted OnVectorValueCommitted,
+		const FOnConstrainVector ConstrainVector)
+	{
+		if(OnVectorValueCommitted.IsBound())
+		{
+			return FOnNumericValueCommitted::CreateLambda(
+				[ComponentIndex, OnComponentCommitted, Vector, OnVectorValueCommitted, ConstrainVector](NumericType InValue, ETextCommit::Type CommitType)
+				{
+					OnComponentCommitted.ExecuteIfBound(InValue, CommitType);
+					
+					const TOptional<VectorType> OptionalVectorValue = Vector.Get();
+					if(OptionalVectorValue.IsSet())
+					{
+						VectorType VectorValue = OptionalVectorValue.GetValue();
+						VectorValue[ComponentIndex] = InValue;
+
+						if(ConstrainVector.IsBound())
+						{
+							ConstrainVector.Execute(ComponentIndex, OptionalVectorValue.GetValue(), VectorValue);
+						}
+
+						OnVectorValueCommitted.Execute(VectorValue, CommitType);
+					}
+				});
+		}
+		return OnComponentCommitted;
+	}		
 };
 
 /**
