Fix Non-Unity: SSegmentedControl.h

diff --git a/Engine/Source/Runtime/Slate/Public/Widgets/Input/SSegmentedControl.h b/Engine/Source/Runtime/Slate/Public/Widgets/Input/SSegmentedControl.h
--- a/Engine/Source/Runtime/Slate/Public/Widgets/Input/SSegmentedControl.h
+++ b/Engine/Source/Runtime/Slate/Public/Widgets/Input/SSegmentedControl.h
@@ -12,6 +12,7 @@
 #include "Widgets/Text/STextBlock.h"
 #include "Widgets/Layout/SBorder.h"
 #include "Widgets/Layout/SUniformGridPanel.h"
+#include "Framework/Application/SlateApplication.h"
 
 /**
  * A Slate Segmented Control is functionally similar to a group of Radio Buttons.
