SSegmentedControl - fix multi selection issue

diff --git a/Engine/Source/Runtime/Slate/Public/Widgets/Input/SSegmentedControl.h b/Engine/Source/Runtime/Slate/Public/Widgets/Input/SSegmentedControl.h
--- a/Engine/Source/Runtime/Slate/Public/Widgets/Input/SSegmentedControl.h
+++ b/Engine/Source/Runtime/Slate/Public/Widgets/Input/SSegmentedControl.h
@@ -459,28 +459,28 @@ private:
 		}
 
 		// if the attribute is not bound update our internal state
-		if (!CurrentValuesIsBound)
+		if(bModifierIsDown)
 		{
-			if(bModifierIsDown)
+			if (InCheckState == ECheckBoxState::Checked)
 			{
-				if (InCheckState == ECheckBoxState::Checked)
-				{
-					Values.AddUnique(InValue);
-				}
-				else
-				{
-					Values.Remove(InValue);
-				}
+				Values.AddUnique(InValue);
 			}
 			else
 			{
-				if(InCheckState == ECheckBoxState::Checked)
-				{
-					Values.Reset();
-					Values.Add(InValue);
-				}
+				Values.Remove(InValue);
 			}
+		}
+		else
+		{
+			if((InCheckState == ECheckBoxState::Checked) || Values.Contains(InValue))
+			{
+				Values.Reset();
+				Values.Add(InValue);
+			}
+		}
 
+		if (!CurrentValuesIsBound)
+		{
 			CurrentValues.Set(*this, Values);
 
 			UpdateCheckboxValuesIfNeeded();
