It is no longer possible to dock tabs into docking areas that are contained inside the tab.

The Blueprint Debugger contains a docking panel, and it was possible to drag the tab such that the drag preview would add the tab into itself, thus causing an infinite recursion in FWidgetPath::SearchForWidgetRecursively during Slate layout.

diff --git a/Engine/Source/Runtime/Slate/Private/Framework/Docking/FDockingDragOperation.cpp b/Engine/Source/Runtime/Slate/Private/Framework/Docking/FDockingDragOperation.cpp
--- a/Engine/Source/Runtime/Slate/Private/Framework/Docking/FDockingDragOperation.cpp
+++ b/Engine/Source/Runtime/Slate/Private/Framework/Docking/FDockingDragOperation.cpp
@@ -4,6 +4,24 @@
 #include "Framework/Application/SlateApplication.h"
 #include "HAL/PlatformApplicationMisc.h"
 
+/** 
+  * Check if the given PotentialParent is a parent of PotentialChild. 
+  */
+static bool IsParentWidgetOf(const TSharedPtr<SWidget>& PotentialParent, const TSharedPtr<SWidget>& PotentialChild)
+{
+	const SWidget* Parent = PotentialChild->GetParentWidget().Get();
+	while (Parent != nullptr)
+	{
+		if (PotentialParent.Get() == Parent)
+		{
+			return true;
+		}
+		Parent = Parent->GetParentWidget().Get();
+	}
+
+	return false;
+}
+
 /**
  * Invoked when the drag and drop operation has ended.
  * 
@@ -59,6 +77,11 @@ void FDockingDragOperation::OnDragged( const FDragDropEvent& DragDropEvent )
  */
 void FDockingDragOperation::OnTabWellEntered( const TSharedRef<class SDockingTabWell>& ThePanel )
 {
+	if (IsParentWidgetOf(this->TabOwnerAreaOfOrigin, ThePanel->GetDockArea()))
+	{
+		return;
+	}
+
 	// We just pulled the tab into some TabWell (in some DockNode).
 	// Hide our decorator window and let the DockNode handle previewing what will happen if we drop the node.
 	HoveredTabPanelPtr = ThePanel;
@@ -179,6 +202,11 @@ bool FDockingDragOperation::CanDockInNode(const TSharedRef<SDockingNode>& DockNo
 	const TSharedRef<FTabManager> TargetTabManager = DockNode->GetDockArea()->GetTabManager();
 	const TSharedRef<FTabManager> TabManagerOfOrigin = this->TabOwnerAreaOfOrigin->GetTabManager();
 
+	if (IsParentWidgetOf(TabBeingDragged->GetContent(), DockNode))
+	{
+		return false;
+	}
+
 	if (TabBeingDragged->GetTabRole() == ETabRole::NomadTab)
 	{
 		if ( IsDockingViaTabwell == FDockingDragOperation::DockingViaTabWell )
