Prevent tooltips from jumping between the mouse location and the widget anchor point when the interactive key combo is pressed

diff --git a/Engine/Source/Runtime/Slate/Private/Framework/Application/SlateUser.cpp b/Engine/Source/Runtime/Slate/Private/Framework/Application/SlateUser.cpp
--- a/Engine/Source/Runtime/Slate/Private/Framework/Application/SlateUser.cpp
+++ b/Engine/Source/Runtime/Slate/Private/Framework/Application/SlateUser.cpp
@@ -1258,7 +1258,7 @@ void FSlateUser::UpdateTooltip(const FMenuStack& MenuStack, bool bCanSpawnNewToo
 	FVector2D DesiredLocation = ActiveTooltipInfo.DesiredLocation;
 	if ((ActiveTooltip && !ActiveTooltip->IsInteractive()) || (NewTooltip && NewTooltip != ActiveTooltip))
 	{
-		// Keep track of where we want tooltips to be positioned
+		// New tooltips and non-interactive tooltips appear offset from the cursor position, and they follow the cursor as it moves.	
 		DesiredLocation = GetPreviousCursorPosition() + SlateDefs::TooltipOffsetFromMouse;
 
 		// Allow interactive tooltips to adjust the window location
@@ -1301,6 +1301,9 @@ void FSlateUser::UpdateTooltip(const FMenuStack& MenuStack, bool bCanSpawnNewToo
 		}
 	}
 
+	// Update the desired location so that interactive tooltips can continue to target it in future frames even after the mouse moves away
+	ActiveTooltipInfo.DesiredLocation = DesiredLocation;
+
 	// The tool tip changed...
 	if (bTooltipChanged)
 	{
