Fix layout not loading json files

diff --git a/Engine/Source/Runtime/Slate/Private/Framework/Docking/LayoutService.cpp b/Engine/Source/Runtime/Slate/Private/Framework/Docking/LayoutService.cpp
--- a/Engine/Source/Runtime/Slate/Private/Framework/Docking/LayoutService.cpp
+++ b/Engine/Source/Runtime/Slate/Private/Framework/Docking/LayoutService.cpp
@@ -11,6 +11,8 @@
 
 DEFINE_LOG_CATEGORY_STATIC(LogLayoutService, Log, All);
 
+PRAGMA_DISABLE_OPTIMIZATION
+
 const TCHAR* DefaultEditorLayoutsSectionName = TEXT("EditorLayouts");
 
 static const TCHAR* GetEditorLayoutsSectionName()
@@ -438,3 +440,5 @@ bool FLayoutSaveRestore::IsValidConfig(const FString& InConfigFileName, bool bAl
 	const FString JsonFileName = GetLayoutJsonFileName(InConfigFileName);
 	return IFileManager::Get().FileExists(*JsonFileName);
 }
+
+PRAGMA_ENABLE_OPTIMIZATION
\ No newline at end of file
