Fix Scrollbox Padding doing nothing on horizontal scroll bars and not updating per calls to SetScrollBarPadding

diff --git a/Engine/Source/Runtime/Slate/Private/Widgets/Layout/SScrollBox.cpp b/Engine/Source/Runtime/Slate/Private/Widgets/Layout/SScrollBox.cpp
--- a/Engine/Source/Runtime/Slate/Private/Widgets/Layout/SScrollBox.cpp
+++ b/Engine/Source/Runtime/Slate/Private/Widgets/Layout/SScrollBox.cpp
@@ -4,7 +4,6 @@
 #include "Rendering/DrawElements.h"
 #include "Types/SlateConstants.h"
 #include "Layout/LayoutUtils.h"
-#include "Widgets/SBoxPanel.h"
 #include "Widgets/SOverlay.h"
 #include "Framework/Application/SlateApplication.h"
 #include "Widgets/Images/SImage.h"
@@ -125,6 +124,7 @@ float SScrollPanel::ArrangeChildHorizontalAndReturnOffset(const FGeometry& Allot
 
 SScrollBox::SScrollBox()
 {
+	VerticalScrollBarSlot = nullptr;
 	bClippingProxy = true;
 }
 
@@ -174,7 +174,7 @@ void SScrollBox::Construct( const FArguments& InArgs )
 		ScrollBar->SetThickness(InArgs._ScrollBarThickness);
 		ScrollBar->SetUserVisibility(InArgs._ScrollBarVisibility);
 		ScrollBar->SetScrollBarAlwaysVisible(InArgs._ScrollBarAlwaysVisible);
-		ScrollBarPadding = InArgs._ScrollBarPadding;
+		ScrollBarSlotPadding = InArgs._ScrollBarPadding;
 
 		bScrollBarIsExternal = false;
 	}
@@ -253,11 +253,13 @@ void SScrollBox::ConstructVerticalLayout()
 		]
 	];
 
+	VerticalScrollBarSlot = nullptr;
 	if (!bScrollBarIsExternal)
 	{
 		PanelAndScrollbar->AddSlot()
-		.Padding(ScrollBarPadding)
+		.Padding(ScrollBarSlotPadding)
 		.AutoWidth()
+		.Expose(VerticalScrollBarSlot)
 		[
 			ScrollBar.ToSharedRef()
 		];
@@ -307,13 +309,16 @@ void SScrollBox::ConstructHorizontalLayout()
 		]
 	];
 
+	HorizontalScrollBarSlot = nullptr;
 	if (!bScrollBarIsExternal)
 	{
 		PanelAndScrollbar->AddSlot()
-			.AutoHeight()
-			[
-				ScrollBar.ToSharedRef()
-			];
+		.Padding(ScrollBarSlotPadding)
+		.AutoHeight()
+		.Expose(HorizontalScrollBarSlot)
+		[
+			ScrollBar.ToSharedRef()
+		];
 	}
 }
 
@@ -533,7 +538,22 @@ void SScrollBox::SetScrollBarThickness(FVector2D InThickness)
 
 void SScrollBox::SetScrollBarPadding(const FMargin& InPadding)
 {
-	ScrollBarPadding = InPadding;
+	ScrollBarSlotPadding = InPadding;
+
+	if (Orientation == Orient_Vertical)
+	{
+		if (VerticalScrollBarSlot)
+		{
+			VerticalScrollBarSlot->SetPadding(ScrollBarSlotPadding);
+		}
+	}
+	else
+	{
+		if (HorizontalScrollBarSlot)
+		{
+			HorizontalScrollBarSlot->SetPadding(ScrollBarSlotPadding);
+		}
+	}
 }
 
 void SScrollBox::SetScrollBarRightClickDragAllowed(bool bIsAllowed)
diff --git a/Engine/Source/Runtime/Slate/Public/Widgets/Layout/SScrollBox.h b/Engine/Source/Runtime/Slate/Public/Widgets/Layout/SScrollBox.h
--- a/Engine/Source/Runtime/Slate/Public/Widgets/Layout/SScrollBox.h
+++ b/Engine/Source/Runtime/Slate/Public/Widgets/Layout/SScrollBox.h
@@ -12,6 +12,7 @@
 #include "Input/NavigationReply.h"
 #include "Widgets/SWidget.h"
 #include "Widgets/SPanel.h"
+#include "Widgets/SBoxPanel.h"
 #include "Layout/Children.h"
 #include "Widgets/DeclarativeSyntaxSupport.h"
 #include "Widgets/SCompoundWidget.h"
@@ -307,6 +308,16 @@ private:
 
 	void BeginInertialScrolling();
 
+	/** Padding to the scrollbox */
+	FMargin ScrollBarSlotPadding;
+
+	union
+	{
+		// vertical scroll bar is stored in horizontal box and vice versa
+		SHorizontalBox::FSlot* VerticalScrollBarSlot; // valid when Orientation == Orient_Vertical
+		SVerticalBox::FSlot* HorizontalScrollBarSlot; // valid when Orientation == Orient_Horizontal
+	};
+
 protected:
 	/** Scrolls or begins scrolling a widget into view, only valid to call when we have layout geometry. */
 	bool InternalScrollDescendantIntoView(const FGeometry& MyGeometry, const TSharedPtr<SWidget>& WidgetToFind, bool InAnimateScroll = true, EDescendantScrollDestination InDestination = EDescendantScrollDestination::IntoView, float Padding = 0);
@@ -341,8 +352,11 @@ protected:
 	/** Whether to permit overscroll on this scroll box */
 	EAllowOverscroll AllowOverscroll;
 
+#if WITH_EDITORONLY_DATA
 	/** Padding to the scrollbox */
+	UE_DEPRECATED(5.0, "ScrollBarPadding is deprecated, Use SetScrollBarPadding")
 	FMargin ScrollBarPadding;
+#endif
 
 	/** Whether to back pad this scroll box, allowing user to scroll backward until child contents are no longer visible */
 	bool BackPadScrolling;
