Fix lingering error when losing focus of an editable text box.

diff --git a/Engine/Source/Runtime/Slate/Private/Widgets/Input/SEditableTextBox.cpp b/Engine/Source/Runtime/Slate/Private/Widgets/Input/SEditableTextBox.cpp
--- a/Engine/Source/Runtime/Slate/Private/Widgets/Input/SEditableTextBox.cpp
+++ b/Engine/Source/Runtime/Slate/Private/Widgets/Input/SEditableTextBox.cpp
@@ -475,22 +475,17 @@ void SEditableTextBox::OnEditableTextCommitted(const FText& InText, ETextCommit:
 		FText OutErrorMessage;
 		if (!OnVerifyTextChanged.Execute(InText, OutErrorMessage))
 		{
-           		// Display as an error.
+           	// Display as an error.
 			if (InCommitType == ETextCommit::OnEnter)
 			{
 				SetError(OutErrorMessage);
 			}
 			return;
 		}
-		else
-		{
-			if (InCommitType == ETextCommit::OnEnter)
-			{
-				SetError(FText::GetEmpty());
-			}
-			
-		}		
 	}
 
+	// Text commited without errors, so clear error text
+	SetError(FText::GetEmpty());
+
 	OnTextCommitted.ExecuteIfBound(InText, InCommitType);
 }
