Slate: Fix the Overlay ZOrder and ConstraintConvas ZOrder in invalidation. The children's order needs to match the order they are painted in for the GlobalInvalidation to work.

diff --git a/Engine/Source/Runtime/Slate/Private/Widgets/Layout/SConstraintCanvas.cpp b/Engine/Source/Runtime/Slate/Private/Widgets/Layout/SConstraintCanvas.cpp
--- a/Engine/Source/Runtime/Slate/Private/Widgets/Layout/SConstraintCanvas.cpp
+++ b/Engine/Source/Runtime/Slate/Private/Widgets/Layout/SConstraintCanvas.cpp
@@ -25,12 +25,76 @@ void SConstraintCanvas::FSlot::Construct(const FChildren& SlotOwner, FSlotArgume
 	{
 		AutoSizeAttr = MoveTemp(InArgs._AutoSize);
 	}
-	if (InArgs._ZOrder.IsSet())
+	ZOrder = InArgs._ZOrder.Get(ZOrder);
+}
+
+void SConstraintCanvas::FSlot::SetZOrder(float InZOrder)
+{
+	if (SWidget* OwnerWidget = FSlotBase::GetOwnerWidget())
+	{
+		if (ZOrder != InZOrder)
+		{
+			TPanelChildren<FSlot>& OwnerChildren = static_cast<SConstraintCanvas*>(OwnerWidget)->Children;
+
+			int32 ChildIndexToMove = 0;
+			{
+				bool bFound = false;
+				const int32 ChildrenNum = OwnerChildren.Num();
+				for (; ChildIndexToMove < ChildrenNum; ++ChildIndexToMove)
+				{
+					const FSlot& CurSlot = OwnerChildren[ChildIndexToMove];
+					if (&CurSlot == this)
+					{
+						bFound = true;
+						break;
+					}
+				}
+				// This slot has to be contained by the children's owner.
+				check(OwnerChildren.IsValidIndex(ChildIndexToMove) && bFound);
+			}
+
+			int32 ChildIndexDestination = 0;
+			{
+				const int32 ChildrenNum = OwnerChildren.Num();
+				for (; ChildIndexDestination < ChildrenNum; ++ChildIndexDestination)
+				{
+					const FSlot& CurSlot = OwnerChildren[ChildIndexDestination];
+					if (InZOrder < CurSlot.GetZOrder() && &CurSlot != this)
+					{
+						// Insert before
+						break;
+					}
+				}
+				if (ChildIndexDestination >= ChildrenNum)
+				{
+					const FSlot& CurSlot = OwnerChildren[ChildrenNum - 1];
+					if (&CurSlot == this)
+					{
+						// No need to move, it's already at the end.
+						ChildIndexToMove = ChildIndexDestination;
+					}
+				}
+			}
+
+			ZOrder = InZOrder;
+
+			if (ChildIndexToMove != ChildIndexDestination)
+			{
+				// TPanelChildren::Move does a remove before the insert, that may change the index location
+				if (ChildIndexDestination > ChildIndexToMove)
+				{
+					--ChildIndexDestination;
+				}
+				OwnerChildren.Move(ChildIndexToMove, ChildIndexDestination);
+			}
+		}
+	}
+	else
 	{
-		ZOrderAttr = MoveTemp(InArgs._ZOrder);
+		ZOrder = InZOrder;
+		ensureMsgf(false, TEXT("The order of the Slot could not be set because it's not added to an existing Widget."));
 	}
 }
-PRAGMA_ENABLE_DEPRECATION_WARNINGS
 
 /* SConstraintCanvas interface
  *****************************************************************************/
@@ -44,7 +108,17 @@ SConstraintCanvas::SConstraintCanvas()
 
 void SConstraintCanvas::Construct( const SConstraintCanvas::FArguments& InArgs )
 {
-	Children.AddSlots(MoveTemp(const_cast<TArray<FSlot::FSlotArguments>&>(InArgs._Slots)));
+	// Sort the children based on ZOrder.
+	TArray<FSlot::FSlotArguments>& Slots = const_cast<TArray<FSlot::FSlotArguments>&>(InArgs._Slots);
+	auto SortOperator = [](const FSlot::FSlotArguments& A, const FSlot::FSlotArguments& B)
+	{
+		int32 AZOrder = A._ZOrder.Get(0);
+		int32 BZOrder = B._ZOrder.Get(0);
+		return AZOrder == BZOrder ? reinterpret_cast<UPTRINT>(&A) < reinterpret_cast<UPTRINT>(&B) : AZOrder < BZOrder;
+	};
+	Slots.StableSort(SortOperator);
+
+	Children.AddSlots(MoveTemp(Slots));
 }
 
 SConstraintCanvas::FSlot::FSlotArguments SConstraintCanvas::Slot()
@@ -54,7 +128,50 @@ SConstraintCanvas::FSlot::FSlotArguments SConstraintCanvas::Slot()
 
 SConstraintCanvas::FScopedWidgetSlotArguments SConstraintCanvas::AddSlot()
 {
-	return FScopedWidgetSlotArguments { MakeUnique<FSlot>(), this->Children, INDEX_NONE };
+	TWeakPtr<SConstraintCanvas> WeakWidget = SharedThis(this);
+	return FScopedWidgetSlotArguments { MakeUnique<FSlot>(), this->Children, INDEX_NONE,
+		[WeakWidget](const FSlot* InNewSlot, int32 InsertedLocation)
+		{
+			if (TSharedPtr<SConstraintCanvas> Pinned = WeakWidget.Pin())
+			{
+				int32 NewSlotZOrder = InNewSlot->GetZOrder();
+				TPanelChildren<FSlot>& OwnerChildren = Pinned->Children;
+				int32 ChildIndexDestination = 0;
+				{
+					const int32 ChildrenNum = OwnerChildren.Num();
+					for (; ChildIndexDestination < ChildrenNum; ++ChildIndexDestination)
+					{
+						const FSlot& CurSlot = OwnerChildren[ChildIndexDestination];
+						if (NewSlotZOrder < CurSlot.GetZOrder() && &CurSlot != InNewSlot)
+						{
+							// Insert before
+							break;
+						}
+					}
+					if (ChildIndexDestination >= ChildrenNum)
+					{
+						const FSlot& CurSlot = OwnerChildren[ChildrenNum - 1];
+						if (&CurSlot == InNewSlot)
+						{
+							// No need to move, it's already at the end.
+							ChildIndexDestination = INDEX_NONE;
+						}
+					}
+				}
+
+				// Move the slot to the correct location
+				if (ChildIndexDestination != INDEX_NONE && InsertedLocation != ChildIndexDestination)
+				{
+					// TPanelChildren::Move does a remove before the insert, that may change the index location
+					if (ChildIndexDestination > InsertedLocation)
+					{
+						--ChildIndexDestination;
+						ensure(false); // we inserted at the end, that should not occurs.
+					}
+					OwnerChildren.Move(InsertedLocation, ChildIndexDestination);
+				}
+			}
+		}};
 }
 
 void SConstraintCanvas::ClearChildren()
@@ -78,20 +195,6 @@ int32 SConstraintCanvas::RemoveSlot( const TSharedRef<SWidget>& SlotWidget )
 	return -1;
 }
 
-struct FChildZOrder
-{
-	int32 ChildIndex;
-	float ZOrder;
-};
-
-struct FSortSlotsByZOrder
-{
-	FORCEINLINE bool operator()(const FChildZOrder& A, const FChildZOrder& B) const
-	{
-		return A.ZOrder == B.ZOrder ? A.ChildIndex < B.ChildIndex : A.ZOrder < B.ZOrder;
-	}
-};
-
 /* SWidget overrides
  *****************************************************************************/
 
@@ -113,28 +216,13 @@ void SConstraintCanvas::ArrangeLayeredChildren(const FGeometry& AllottedGeometry
 #else
 		const static bool bExplicitChildZOrder = GetDefault<USlateSettings>()->bExplicitCanvasChildZOrder;
 #endif
-		// Sort the children based on zorder.
-		TArray< FChildZOrder, TInlineAllocator<64> > SlotOrder;
-		SlotOrder.Reserve(Children.Num());
-
-		for (int32 ChildIndex = 0; ChildIndex < Children.Num(); ++ChildIndex)
-		{
-			const SConstraintCanvas::FSlot& CurChild = Children[ChildIndex];
-
-			FChildZOrder Order;
-			Order.ChildIndex = ChildIndex;
-			Order.ZOrder = CurChild.GetZOrder();
-			SlotOrder.Add(Order);
-		}
 
-		SlotOrder.StableSort(FSortSlotsByZOrder());
 		float LastZOrder = -FLT_MAX;
 
 		// Arrange the children now in their proper z-order.
 		for (int32 ChildIndex = 0; ChildIndex < Children.Num(); ++ChildIndex)
 		{
-			const FChildZOrder& CurSlot = SlotOrder[ChildIndex];
-			const SConstraintCanvas::FSlot& CurChild = Children[CurSlot.ChildIndex];
+			const SConstraintCanvas::FSlot& CurChild = Children[ChildIndex];
 			const TSharedRef<SWidget>& CurWidget = CurChild.GetWidget();
 
 			const EVisibility ChildVisibility = CurWidget->GetVisibility();
@@ -204,13 +292,13 @@ void SConstraintCanvas::ArrangeLayeredChildren(const FGeometry& AllottedGeometry
 				{
 					// Split children into discrete layers for the paint method
 					bNewLayer = false;
-					if (CurSlot.ZOrder > LastZOrder + DELTA)
+					if (CurChild.GetZOrder() > LastZOrder + DELTA)
 					{
 						if (ArrangedChildLayers.Num() > 0)
 						{
 							bNewLayer = true;
 						}
-						LastZOrder = CurSlot.ZOrder;
+						LastZOrder = CurChild.GetZOrder();
 					}
 
 				}
diff --git a/Engine/Source/Runtime/Slate/Public/Widgets/Layout/SConstraintCanvas.h b/Engine/Source/Runtime/Slate/Public/Widgets/Layout/SConstraintCanvas.h
--- a/Engine/Source/Runtime/Slate/Public/Widgets/Layout/SConstraintCanvas.h
+++ b/Engine/Source/Runtime/Slate/Public/Widgets/Layout/SConstraintCanvas.h
@@ -40,7 +40,7 @@ public:
 			, AnchorsAttr(FAnchors(0.0f, 0.0f))
 			, AlignmentAttr(FVector2D(0.5f, 0.5f))
 			, AutoSizeAttr(false)
-			, ZOrderAttr(0)
+			, ZOrder(0)
 		{ }
 
 		SLATE_SLOT_BEGIN_ARGS(FSlot, TSlotBase<FSlot>)
@@ -48,7 +48,7 @@ public:
 			SLATE_ATTRIBUTE(FAnchors, Anchors)
 			SLATE_ATTRIBUTE(FVector2D, Alignment)
 			SLATE_ATTRIBUTE(bool, AutoSize)
-			SLATE_ATTRIBUTE(float, ZOrder)
+			SLATE_ARGUMENT(TOptional<float>, ZOrder)
 		SLATE_SLOT_END_ARGS()
 
 		void Construct(const FChildren& SlotOwner, FSlotArguments&& InArgs);
@@ -93,18 +93,11 @@ public:
 			return AutoSizeAttr.Get();
 		}
 
-		void SetZOrder(const TAttribute<float>& InZOrder)
-		{
-			// Layout isn't entirely correct here, but Paint wouldn't be either.
-			// We need the parent to redraw the children because the logical order didn't change
-			// but the paint order may have if ZOrder changed, so the painted elements need to
-			// be resorted.
-			SetAttribute(ZOrderAttr, InZOrder, EInvalidateWidgetReason::Layout);
-		}
+		void SetZOrder(float InZOrder);
 
 		float GetZOrder() const
 		{
-			return ZOrderAttr.Get();
+			return ZOrder;
 		}
 
 	public:
@@ -124,9 +117,14 @@ public:
 		UE_DEPRECATED(5.0, "Direct access to AutoSizeAttr is now deprecated. Use the getter or setter.")
 		TAttribute<bool> AutoSizeAttr;
 
+#if WITH_EDITORONLY_DATA
 		/** Z-Order */
 		UE_DEPRECATED(5.0, "Direct access to ZOrderAttr is now deprecated. Use the getter or setter.")
 		TAttribute<float> ZOrderAttr;
+#endif
+	private:
+		/** Z-Order */
+		float ZOrder;
 	};
 	PRAGMA_ENABLE_DEPRECATION_WARNINGS
 
diff --git a/Engine/Source/Runtime/SlateCore/Private/Widgets/SOverlay.cpp b/Engine/Source/Runtime/SlateCore/Private/Widgets/SOverlay.cpp
--- a/Engine/Source/Runtime/SlateCore/Private/Widgets/SOverlay.cpp
+++ b/Engine/Source/Runtime/SlateCore/Private/Widgets/SOverlay.cpp
@@ -13,11 +13,97 @@ void SOverlay::PrivateRegisterAttributes(FSlateAttributeInitializer& AttributeIn
 	FOverlaySlot::RegisterAttributes(Initializer);
 }
 
-PRAGMA_DISABLE_DEPRECATION_WARNINGS
 void SOverlay::FOverlaySlot::Construct(const FChildren& SlotOwner, FSlotArguments&& InArgs)
 {
 	TBasicLayoutWidgetSlot<FOverlaySlot>::Construct(SlotOwner, MoveTemp(InArgs));
-	ZOrder = InArgs._ZOrder.Get(ZOrder);
+}
+
+PRAGMA_DISABLE_DEPRECATION_WARNINGS
+void SOverlay::FOverlaySlot::SetZOrder(int32 InZOrder)
+{
+	if (SWidget* OwnerWidget = FSlotBase::GetOwnerWidget())
+	{
+		if (ZOrder != InZOrder)
+		{
+			TPanelChildren<FOverlaySlot>& OwnerChildren = static_cast<SOverlay*>(OwnerWidget)->Children;
+
+			// if bellow 0 or INDEX_NONE, then move it to the end.
+			if (InZOrder <= INDEX_NONE)
+			{
+				// We have at least one child since GetOwnerWidget returned a valid Widget (parent/child relationship).
+				check(OwnerChildren.Num() != 0);
+
+				const FOverlaySlot& LastSlot = OwnerChildren[OwnerChildren.Num() - 1];
+				if (&LastSlot == this)
+				{
+					return; // We are already the last FSlot
+				}
+
+				InZOrder = LastSlot.GetZOrder() + 1;
+				if (ZOrder == InZOrder)
+				{
+					return; 
+				}
+			}
+
+			int32 ChildIndexToMove = 0;
+			{
+				bool bFound = false;
+				const int32 ChildrenNum = OwnerChildren.Num();
+				for (; ChildIndexToMove < ChildrenNum; ++ChildIndexToMove)
+				{
+					const FOverlaySlot& CurSlot = OwnerChildren[ChildIndexToMove];
+					if (&CurSlot == this)
+					{
+						bFound = true;
+						break;
+					}
+				}
+				// This slot has to be contained by the children's owner.
+				check(OwnerChildren.IsValidIndex(ChildIndexToMove) && bFound);
+			}
+
+			int32 ChildIndexDestination = 0;
+			{
+				const int32 ChildrenNum = OwnerChildren.Num();
+				for (; ChildIndexDestination < ChildrenNum; ++ChildIndexDestination)
+				{
+					const FOverlaySlot& CurSlot = OwnerChildren[ChildIndexDestination];
+					if (InZOrder < CurSlot.GetZOrder() && &CurSlot != this)
+					{
+						// Insert before
+						break;
+					}
+				}
+				if (ChildIndexDestination >= ChildrenNum)
+				{
+					const FOverlaySlot& CurSlot = OwnerChildren[ChildrenNum-1];
+					if (&CurSlot == this)
+					{
+						// No need to move, it's already at the end.
+						ChildIndexToMove = ChildIndexDestination;
+					}
+				}
+			}
+
+			ZOrder = InZOrder;
+
+			if (ChildIndexToMove != ChildIndexDestination)
+			{
+				// TPanelChildren::Move does a remove before the insert, that may change the index location
+				if (ChildIndexDestination > ChildIndexToMove)
+				{
+					--ChildIndexDestination;
+				}
+				OwnerChildren.Move(ChildIndexToMove, ChildIndexDestination);
+			}
+		}
+	}
+	else
+	{
+		ZOrder = InZOrder;
+		ensureMsgf(false, TEXT("The order of the OverlaySlot could not be set because it's not added to an existing Widget."));
+	}
 }
 PRAGMA_ENABLE_DEPRECATION_WARNINGS
 
@@ -30,7 +116,7 @@ SOverlay::SOverlay()
 
 SOverlay::FOverlaySlot::FSlotArguments SOverlay::Slot()
 {
-	return FOverlaySlot::FSlotArguments(MakeUnique<FOverlaySlot>());
+	return FOverlaySlot::FSlotArguments(MakeUnique<FOverlaySlot>(0));
 }
 
 void SOverlay::Construct( const SOverlay::FArguments& InArgs )
@@ -168,8 +254,7 @@ SOverlay::FScopedWidgetSlotArguments SOverlay::AddSlot( int32 ZOrder )
 		Index = CurSlotIndex;
 	}
 
-	FScopedWidgetSlotArguments Result {MakeUnique<FOverlaySlot>(), this->Children, Index};
-	Result.ZOrder(ZOrder);
+	FScopedWidgetSlotArguments Result {MakeUnique<FOverlaySlot>(ZOrder), this->Children, Index};
 	return MoveTemp(Result);
 }
 
diff --git a/Engine/Source/Runtime/SlateCore/Public/Widgets/SOverlay.h b/Engine/Source/Runtime/SlateCore/Public/Widgets/SOverlay.h
--- a/Engine/Source/Runtime/SlateCore/Public/Widgets/SOverlay.h
+++ b/Engine/Source/Runtime/SlateCore/Public/Widgets/SOverlay.h
@@ -41,13 +41,12 @@ public:
 	class SLATECORE_API FOverlaySlot : public TBasicLayoutWidgetSlot<FOverlaySlot>
 	{
 	public:
-		FOverlaySlot()
+		FOverlaySlot(int32 InZOrder)
 			: TBasicLayoutWidgetSlot<FOverlaySlot>(HAlign_Fill, VAlign_Fill)
-			, ZOrder(0)
+			, ZOrder(InZOrder)
 		{ }
 
 		SLATE_SLOT_BEGIN_ARGS(FOverlaySlot, TBasicLayoutWidgetSlot<FOverlaySlot>)
-			SLATE_ARGUMENT(TOptional<int32>, ZOrder)
 		SLATE_SLOT_END_ARGS()
 
 		void Construct(const FChildren& SlotOwner, FSlotArguments&& InArgs);
@@ -57,20 +56,12 @@ public:
 			return ZOrder;
 		}
 
-		void SetZOrder(int32 InOrder)
-		{
-			if (ZOrder != InOrder)
-			{
-				ZOrder = InOrder;
-				FSlotBase::Invalidate(EInvalidateWidgetReason::Layout);
-			}
-		}
+		void SetZOrder(int32 InOrder);
 
 	public:
 		/**
-		 * Slots with larger ZOrder values will draw above slots with smaller ZOrder values.  Slots
-		 * with the same ZOrder will simply draw in the order they were added.  Currently this only
-		 * works for overlay slots that are added dynamically with AddWidget() and RemoveWidget()
+		 * Slots with larger ZOrder values will draw above slots with smaller ZOrder values. Slots
+		 * with the same ZOrder will simply draw in the order they were added.
 		 */
 		UE_DEPRECATED(5.0, "Direct access to ZOrder is now deprecated. Use the getter or setter.")
 		int32 ZOrder;
