Slate: Add test to see if widgets have a valid LayerId.

diff --git a/Engine/Source/Runtime/Slate/Private/Framework/Application/SlateApplication.cpp b/Engine/Source/Runtime/Slate/Private/Framework/Application/SlateApplication.cpp
--- a/Engine/Source/Runtime/Slate/Private/Framework/Application/SlateApplication.cpp
+++ b/Engine/Source/Runtime/Slate/Private/Framework/Application/SlateApplication.cpp
@@ -69,13 +69,23 @@ bool GSlateVerifyParentChildrenRelationship = false;
 static FAutoConsoleVariableRef CVarSlateVerifyParentChildrenRelationship(
 	TEXT("Slate.VerifyParentChildrenRelationship"),
 	GSlateVerifyParentChildrenRelationship,
-	TEXT("Each frame, verify that a widget has only one parent.")
+	TEXT("Every tick, verify that a widget has only one parent.")
 );
-namespace UE {
-	namespace Slate
-	{
-		void VerifyParentChildrenRelationship(const TSharedRef<SWindow>& WindowToDraw);
-	}
+
+bool GSlateVerifyWidgetLayerId = false;
+static FAutoConsoleVariableRef CVarSlateVerifyWidgetLayerId(
+	TEXT("Slate.VerifyWidgetLayerId"),
+	GSlateVerifyWidgetLayerId,
+	TEXT("Every tick, verify that widgets have a LayerId range that fits with their siblings and their parent.")
+);
+
+namespace UE
+{
+namespace Slate
+{
+	void VerifyParentChildrenRelationship(const TSharedRef<SWindow>& WindowToDraw);
+	void VerifyWidgetLayerId(const TSharedRef<SWindow>& WindowToDraw);
+}
 }
 #endif
 
@@ -1133,6 +1143,13 @@ void FSlateApplication::DrawWindowAndChildren( const TSharedRef<SWindow>& Window
 
 #endif
 
+#if WITH_SLATE_DEBUGGING
+		if (GSlateVerifyWidgetLayerId)
+		{
+			UE::Slate::VerifyWidgetLayerId(WindowToDraw);
+		}
+#endif
+
 		// This window is visible, so draw its child windows as well
 		bDrawChildWindows = true;
 	}
@@ -7106,6 +7123,7 @@ namespace Private
 				{
 					if (AllWidgets.Find(&ChildWidget))
 					{
+						bResult = false;
 						UE_LOG(LogSlate, Warning, TEXT("The widget '%s' is owned by more than one parent. 1:'%s' 2:'%s'.")
 							, *FReflectionMetaData::GetWidgetDebugInfo(ChildWidget)
 							, *FReflectionMetaData::GetWidgetDebugInfo(AllWidgets[&ChildWidget])
@@ -7119,11 +7137,56 @@ namespace Private
 					if (ChildWidget.GetParentWidget().Get() != Parent)
 					{
 						bResult = false;
-						UE_LOG(LogSlate, Warning, TEXT("The widget '%s' ."), *FReflectionMetaData::GetWidgetDebugInfo(ChildWidget));
+						UE_LOG(LogSlate, Warning, TEXT("The widget '%s' has the wrong parent."), *FReflectionMetaData::GetWidgetDebugInfo(ChildWidget));
 					}
-					else
+
+					bResult = bResult && VerifyParentChildrenRelationship_Recursive(&ChildWidget, AllWidgets);
+				}
+			});
+
+		return bResult;
+	}
+
+	bool VerifyWidgetLayerId_Recursive(SWidget& Widget)
+	{
+		bool bResult = true;
+
+		const uint32 LastPaintFrame = Widget.Debug_GetLastPaintFrame();
+		const bool bIsDeferredPaint = Widget.GetPersistentState().bDeferredPainting;
+		const int32 InLayerId = Widget.GetPersistentState().LayerId;
+		const int32 OutLayerId = Widget.GetPersistentState().OutgoingLayerId;
+
+		int32 PreviousInLayerId = InLayerId;
+		int32 PreviousOutLayerId = InLayerId;
+
+		Widget.GetAllChildren()->ForEachWidget([&bResult, LastPaintFrame, bIsDeferredPaint, InLayerId, OutLayerId](SWidget& ChildWidget)
+			{
+				if (&ChildWidget != &SNullWidget::NullWidget.Get())
+				{
+					if (ChildWidget.GetVisibility().IsVisible() && ChildWidget.Debug_GetLastPaintFrame() == LastPaintFrame)
 					{
-						VerifyParentChildrenRelationship_Recursive(&ChildWidget, AllWidgets);
+						const bool bIsChildDeferredPaint = ChildWidget.GetPersistentState().bDeferredPainting;
+						if (bIsChildDeferredPaint == bIsDeferredPaint)
+						{
+							const int32 ChildInLayerId = ChildWidget.GetPersistentState().LayerId;
+							const int32 ChildOutLayerId = ChildWidget.GetPersistentState().OutgoingLayerId;
+
+							const bool bLayerInIsValid = InLayerId <= ChildInLayerId && InLayerId <= ChildOutLayerId;
+							const bool bLayerOutIsValid = OutLayerId >= ChildInLayerId && OutLayerId >= ChildOutLayerId;
+
+							if (!bLayerInIsValid || !bLayerOutIsValid)
+							{
+								bResult = false;
+								UE_LOG(LogSlate, Warning, TEXT("The widget '%s' LayerId is invalid. Parent: [%d,%d] Child: [%d,%d].")
+									, *FReflectionMetaData::GetWidgetDebugInfo(ChildWidget)
+									, InLayerId
+									, OutLayerId
+									, ChildInLayerId
+									, ChildOutLayerId);
+							}
+						}
+
+						bResult = bResult && VerifyWidgetLayerId_Recursive(ChildWidget);
 					}
 				}
 			});
@@ -7141,12 +7204,26 @@ void VerifyParentChildrenRelationship(const TSharedRef<SWindow>& WindowToDraw)
 		AllWidgets.Add(&WindowToDraw.Get(), nullptr);
 		if (!Private::VerifyParentChildrenRelationship_Recursive(&WindowToDraw.Get(), AllWidgets))
 		{
-			GSlateVerifyParentChildrenRelationship = false;
+			CVarSlateVerifyParentChildrenRelationship->Set(false);
 			ensureMsgf(false, TEXT("VerifyParentChildrenRelationship failed. See log for more info."));
 		}
 	}
 }
 
+
+void VerifyWidgetLayerId(const TSharedRef<SWindow>& WindowToDraw)
+{
+	if (WindowToDraw != SNullWidget::NullWidget)
+	{
+		if (!Private::VerifyWidgetLayerId_Recursive(WindowToDraw.Get()))
+		{
+			CVarSlateVerifyWidgetLayerId->Set(false);
+			ensureMsgf(false, TEXT("VerifyWidgetLayerId failed. See log for more info."));
+		}
+	}
+}
+
+
 #endif
 }// Slate
 }//UE
diff --git a/Engine/Source/Runtime/SlateCore/Private/FastUpdate/WidgetProxy.cpp b/Engine/Source/Runtime/SlateCore/Private/FastUpdate/WidgetProxy.cpp
--- a/Engine/Source/Runtime/SlateCore/Private/FastUpdate/WidgetProxy.cpp
+++ b/Engine/Source/Runtime/SlateCore/Private/FastUpdate/WidgetProxy.cpp
@@ -212,6 +212,7 @@ int32 FWidgetProxy::Repaint(const FPaintArgs& PaintArgs, FSlateWindowElementList
 	
 	FPaintArgs UpdatedArgs = PaintArgs.WithNewParent(MyState.PaintParent.Pin().Get());
 	UpdatedArgs.SetInheritedHittestability(MyState.bInheritedHittestability);
+	UpdatedArgs.SetDeferredPaint(MyState.bDeferredPainting);
 
 	int32 PrevLayerId = MyState.OutgoingLayerId;
 
diff --git a/Engine/Source/Runtime/SlateCore/Private/Rendering/DrawElements.cpp b/Engine/Source/Runtime/SlateCore/Private/Rendering/DrawElements.cpp
--- a/Engine/Source/Runtime/SlateCore/Private/Rendering/DrawElements.cpp
+++ b/Engine/Source/Runtime/SlateCore/Private/Rendering/DrawElements.cpp
@@ -765,6 +765,8 @@ FSlateWindowElementList::FDeferredPaint::FDeferredPaint( const TSharedRef<const
 	, WidgetStyle( InWidgetStyle )
 	, bParentEnabled( InParentEnabled )
 {
+	const_cast<FPaintArgs&>(Args).SetDeferredPaint(true);
+
 #if WITH_SLATE_DEBUGGING
 	// We need to perform this update here, because otherwise we'll warn that this widget
 	// was not painted along the fast path, which, it will be, but later because it's deferred,
@@ -781,6 +783,7 @@ FSlateWindowElementList::FDeferredPaint::FDeferredPaint(const FDeferredPaint& Co
 	, WidgetStyle(Copy.WidgetStyle)
 	, bParentEnabled(Copy.bParentEnabled)
 {
+	const_cast<FPaintArgs&>(Args).SetDeferredPaint(true);
 }
 
 int32 FSlateWindowElementList::FDeferredPaint::ExecutePaint(int32 LayerId, FSlateWindowElementList& OutDrawElements, const FSlateRect& MyCullingRect) const
@@ -802,7 +805,7 @@ FSlateWindowElementList::FDeferredPaint FSlateWindowElementList::FDeferredPaint:
 
 void FSlateWindowElementList::QueueDeferredPainting( const FDeferredPaint& InDeferredPaint )
 {
-	DeferredPaintList.Add(MakeShareable(new FDeferredPaint(InDeferredPaint)));
+	DeferredPaintList.Add(MakeShared<FDeferredPaint>(InDeferredPaint));
 }
 
 int32 FSlateWindowElementList::PaintDeferred(int32 LayerId, const FSlateRect& MyCullingRect)
diff --git a/Engine/Source/Runtime/SlateCore/Private/Types/PaintArgs.cpp b/Engine/Source/Runtime/SlateCore/Private/Types/PaintArgs.cpp
--- a/Engine/Source/Runtime/SlateCore/Private/Types/PaintArgs.cpp
+++ b/Engine/Source/Runtime/SlateCore/Private/Types/PaintArgs.cpp
@@ -13,6 +13,7 @@ FPaintArgs::FPaintArgs(const SWidget* PaintParent, FHittestGrid& InRootHittestGr
 	, CurrentTime(InCurrentTime)
 	, DeltaTime(InDeltaTime)
 	, bInheritedHittestability(true)
+	, bDeferredPainting(false)
 {
 }
 
diff --git a/Engine/Source/Runtime/SlateCore/Private/Widgets/SWidget.cpp b/Engine/Source/Runtime/SlateCore/Private/Widgets/SWidget.cpp
--- a/Engine/Source/Runtime/SlateCore/Private/Widgets/SWidget.cpp
+++ b/Engine/Source/Runtime/SlateCore/Private/Widgets/SWidget.cpp
@@ -1472,6 +1472,7 @@ int32 SWidget::Paint(const FPaintArgs& Args, const FGeometry& AllottedGeometry,
 	PersistentState.LayerId = LayerId;
 	PersistentState.bParentEnabled = bParentEnabled;
 	PersistentState.bInheritedHittestability = bInheritedHittestability;
+	PersistentState.bDeferredPainting = Args.GetDeferredPaint();
 	PersistentState.AllottedGeometry = AllottedGeometry;
 	PersistentState.DesktopGeometry = DesktopSpaceGeometry;
 	PersistentState.WidgetStyle = InWidgetStyle;
@@ -1521,15 +1522,14 @@ int32 SWidget::Paint(const FPaintArgs& Args, const FGeometry& AllottedGeometry,
 	TGuardValue<EFlowDirection> FlowGuard(GSlateFlowDirection, ComputeFlowDirection());
 
 #if WITH_SLATE_DEBUGGING
-	TArray<TWeakPtr<const SWidget>> DebugChildWidgetsToPaint;
+	TArray<TWeakPtr<const SWidget>, TInlineAllocator<16>> DebugChildWidgetsToPaint;
 
 	if (GSlateIsOnFastUpdatePath && GSlateEnsureAllVisibleWidgetsPaint)
 	{
-		// Don't check things that are invalidation roots, or volatile, or volatile indirectly, a completely different set
-		// of rules apply to those widgets.
-		if (!IsVolatile() && !IsVolatileIndirectly() && !Advanced_IsInvalidationRoot())
+		// Don't check for invalidation roots a completely different set of rules apply to those widgets.
+		if (!Advanced_IsInvalidationRoot())
 		{
-			MutableThis->GetChildren()->ForEachWidget([&DebugChildWidgetsToPaint](const SWidget& Child)
+			MutableThis->GetAllChildren()->ForEachWidget([&DebugChildWidgetsToPaint](const SWidget& Child)
 				{
 					if (Child.GetVisibility().IsVisible())
 					{
@@ -1552,7 +1552,6 @@ int32 SWidget::Paint(const FPaintArgs& Args, const FGeometry& AllottedGeometry,
 	// Detect children that should have been painted, but were skipped during the paint process.
 	// this will result in geometry being left on screen and not cleared, because it's visible, yet wasn't painted.
 #if WITH_SLATE_DEBUGGING
-	if (GSlateIsOnFastUpdatePath && GSlateEnsureAllVisibleWidgetsPaint)
 	{
 		for (TWeakPtr<const SWidget>& DebugChildThatShouldHaveBeenPaintedPtr : DebugChildWidgetsToPaint)
 		{
@@ -1560,9 +1559,14 @@ int32 SWidget::Paint(const FPaintArgs& Args, const FGeometry& AllottedGeometry,
 			{
 				if (DebugChild->GetVisibility().IsVisible())
 				{
-					ensureMsgf(DebugChild->Debug_GetLastPaintFrame() == GFrameNumber, TEXT("The Widget '%s' was visible, but never painted.  This means it was skipped during painting, without alerting the fast path."), *FReflectionMetaData::GetWidgetPath(DebugChild.Get()));
+					ensureAlwaysMsgf(DebugChild->Debug_GetLastPaintFrame() == GFrameNumber, TEXT("The Widget '%s' was visible, but never painted. This means it was skipped during painting, without alerting the fast path."), *FReflectionMetaData::GetWidgetPath(DebugChild.Get()));
 				}
 			}
+			else
+			{
+				CVarSlateEnsureAllVisibleWidgetsPaint->Set(false);
+				ensureAlwaysMsgf(false, TEXT("A widget was destroyed while painting it's parent. This is not supported by the Invalidation system."));
+			}
 		}
 	}
 #endif
diff --git a/Engine/Source/Runtime/SlateCore/Public/FastUpdate/WidgetProxy.h b/Engine/Source/Runtime/SlateCore/Public/FastUpdate/WidgetProxy.h
--- a/Engine/Source/Runtime/SlateCore/Public/FastUpdate/WidgetProxy.h
+++ b/Engine/Source/Runtime/SlateCore/Public/FastUpdate/WidgetProxy.h
@@ -225,6 +225,7 @@ struct FSlateWidgetPersistentState
 		, IncomingFlowDirection(EFlowDirection::LeftToRight)
 		, bParentEnabled(true)
 		, bInheritedHittestability(false)
+		, bDeferredPainting(false)
 	{}
 
 	TWeakPtr<SWidget> PaintParent;
@@ -241,6 +242,7 @@ struct FSlateWidgetPersistentState
 	EFlowDirection IncomingFlowDirection;
 	uint8 bParentEnabled : 1;
 	uint8 bInheritedHittestability : 1;
+	uint8 bDeferredPainting : 1;
 
 	static const FSlateWidgetPersistentState NoState;
 };
diff --git a/Engine/Source/Runtime/SlateCore/Public/Types/PaintArgs.h b/Engine/Source/Runtime/SlateCore/Public/Types/PaintArgs.h
--- a/Engine/Source/Runtime/SlateCore/Public/Types/PaintArgs.h
+++ b/Engine/Source/Runtime/SlateCore/Public/Types/PaintArgs.h
@@ -27,7 +27,7 @@ public:
 	FPaintArgs(const SWidget* PaintParent, FHittestGrid& InRootHittestGrid, FHittestGrid& InCurrentHitTestGrid, FVector2D InWindowOffset, double InCurrentTime, float InDeltaTime);
 	FPaintArgs(const SWidget* PaintParent, FHittestGrid& InRootHittestGrid, FVector2D InWindowOffset, double InCurrentTime, float InDeltaTime);
 
-	FORCEINLINE_DEBUGGABLE FPaintArgs WithNewParent(const SWidget* PaintParent) const
+	UE_NODISCARD FORCEINLINE_DEBUGGABLE FPaintArgs WithNewParent(const SWidget* PaintParent) const
 	{
 		FPaintArgs Args(*this);
 		Args.PaintParentPtr = PaintParent;
@@ -35,26 +35,58 @@ public:
 		return Args;
 	}
 
-	FORCEINLINE_DEBUGGABLE FPaintArgs WithNewHitTestGrid(FHittestGrid& NewHitTestGrid) const
+	UE_NODISCARD FORCEINLINE_DEBUGGABLE FPaintArgs WithNewHitTestGrid(FHittestGrid& NewHitTestGrid) const
 	{
 		FPaintArgs NewArgs(PaintParentPtr, RootGrid, NewHitTestGrid, WindowOffset, CurrentTime, DeltaTime);
 		return NewArgs;
 	}
 
 	FPaintArgs InsertCustomHitTestPath(const SWidget* Widget, TSharedRef<ICustomHitTestPath> CustomHitTestPath) const;
-	void SetInheritedHittestability(bool InInheritedHittestability) { bInheritedHittestability = InInheritedHittestability; }
 
-	UE_DEPRECATED(4.23, "FHittestGrid::GetGrid is deprecated.  Use GetHittestGrid instead")
-	FHittestGrid& GetGrid() const { return GetHittestGrid(); }
+	void SetInheritedHittestability(bool InInheritedHittestability)
+	{
+		bInheritedHittestability = InInheritedHittestability;
+	}
+
+	bool GetInheritedHittestability() const
+	{
+		return bInheritedHittestability;
+	}
+
+	FHittestGrid& GetHittestGrid() const
+	{
+		return CurrentGrid;
+	}
 
-	FHittestGrid& GetHittestGrid() const { return CurrentGrid; }
+	const SWidget* GetPaintParent() const
+	{
+		return PaintParentPtr; 
+	}
 
-	const SWidget* GetPaintParent() const { return PaintParentPtr;  }
-	FVector2D GetWindowToDesktopTransform() const { return WindowOffset; }
-	double GetCurrentTime() const { return CurrentTime; }
-	float GetDeltaTime() const { return DeltaTime; }
+	FVector2D GetWindowToDesktopTransform() const
+	{
+		return WindowOffset;
+	}
 
-	bool GetInheritedHittestability() const { return bInheritedHittestability; }
+	double GetCurrentTime() const
+	{
+		return CurrentTime;
+	}
+
+	float GetDeltaTime() const
+	{
+		return DeltaTime;
+	}
+
+	void SetDeferredPaint(bool InDeferredPaint)
+	{
+		bDeferredPainting = InDeferredPaint;
+	}
+
+	bool GetDeferredPaint() const
+	{
+		return bDeferredPainting;
+	}
 	
 private:
 
@@ -70,5 +102,6 @@ private:
 	double CurrentTime;
 	float DeltaTime;
 	uint8 bInheritedHittestability : 1;
+	uint8 bDeferredPainting : 1;
 
 };
