Fix balanced desprecation warnings macro. Introduced with CL 18534070.

diff --git a/Engine/Source/Runtime/Slate/Private/Widgets/Layout/SConstraintCanvas.cpp b/Engine/Source/Runtime/Slate/Private/Widgets/Layout/SConstraintCanvas.cpp
--- a/Engine/Source/Runtime/Slate/Private/Widgets/Layout/SConstraintCanvas.cpp
+++ b/Engine/Source/Runtime/Slate/Private/Widgets/Layout/SConstraintCanvas.cpp
@@ -95,6 +95,7 @@ void SConstraintCanvas::FSlot::SetZOrder(float InZOrder)
 		ensureMsgf(false, TEXT("The order of the Slot could not be set because it's not added to an existing Widget."));
 	}
 }
+PRAGMA_ENABLE_DEPRECATION_WARNINGS
 
 /* SConstraintCanvas interface
  *****************************************************************************/
