Fix 2x2 splitter not having any hover cues

diff --git a/Engine/Source/Runtime/Slate/Private/Widgets/Layout/SSplitter.cpp b/Engine/Source/Runtime/Slate/Private/Widgets/Layout/SSplitter.cpp
--- a/Engine/Source/Runtime/Slate/Private/Widgets/Layout/SSplitter.cpp
+++ b/Engine/Source/Runtime/Slate/Private/Widgets/Layout/SSplitter.cpp
@@ -3,6 +3,7 @@
 #include "Widgets/Layout/SSplitter.h"
 #include "Layout/ArrangedChildren.h"
 #include "Rendering/DrawElements.h"
+#include "Framework/Application/SlateApplication.h"
 
 PRAGMA_DISABLE_DEPRECATION_WARNINGS
 void SSplitter::FSlot::Construct(const FChildren& SlotOwner, FSlotArguments&& InArgs)
@@ -866,10 +867,12 @@ void SSplitter2x2::Construct( const FArguments& InArgs )
 	Children.AddSlot(FSlot::FSlotArguments(MakeUnique<FSlot>(InArgs._TopRight.Widget)));
 	Children.AddSlot(FSlot::FSlotArguments(MakeUnique<FSlot>(InArgs._BottomRight.Widget)));
 
+	Style = InArgs._Style;
+
 	SplitterHandleSize = 5.0f;
 	MinSplitterChildLength = 20.0f;
 	bIsResizing = false;
-	ResizingAxis = INDEX_NONE;
+	ResizingAxisMask = EResizingAxis::None;
 }
 
 TArray<FLayoutGeometry> SSplitter2x2::ArrangeChildrenForLayout( const FGeometry& AllottedGeometry ) const
@@ -927,7 +930,47 @@ void SSplitter2x2::OnArrangeChildren( const FGeometry& AllottedGeometry, FArrang
 	}
 }
 
+int32 SSplitter2x2::OnPaint(const FPaintArgs& Args, const FGeometry& AllottedGeometry, const FSlateRect& MyCullingRect, FSlateWindowElementList& OutDrawElements, int32 LayerId, const FWidgetStyle& InWidgetStyle, bool bParentEnabled) const
+{
+	FArrangedChildren ArrangedChildren(EVisibility::Visible);
+	ArrangeChildren(AllottedGeometry, ArrangedChildren);
+
+	int32 MaxLayerId = PaintArrangedChildren(Args, ArrangedChildren, AllottedGeometry, MyCullingRect, OutDrawElements, LayerId, InWidgetStyle, bParentEnabled);
+
+	const FVector2D LocalMousePos = GetTickSpaceGeometry().AbsoluteToLocal(FSlateApplication::Get().GetCursorPos());
+
+	EResizingAxis HoveredAxis = CalculateResizingAxis(GetTickSpaceGeometry(), LocalMousePos);
+
+	// First two children are top left and bottom left. Use that to calculate horizontal separator
 
+	const FVector2D HorizontalSplitterPos(0, ArrangedChildren[0].Geometry.GetLocalSize().Y);
+	const FVector2D HorizontalSplitterSize(AllottedGeometry.GetLocalSize().X, ArrangedChildren[1].Geometry.Position.Y - ArrangedChildren[0].Geometry.GetLocalSize().Y);
+
+	const FSlateBrush* HorizontalBrush = EnumHasAnyFlags(HoveredAxis, EResizingAxis::UpDownMask) || EnumHasAnyFlags(ResizingAxisMask, EResizingAxis::UpDownMask) ? &Style->HandleHighlightBrush : &Style->HandleNormalBrush;
+
+	FSlateDrawElement::MakeBox(
+		OutDrawElements,
+		MaxLayerId,
+		AllottedGeometry.ToPaintGeometry(HorizontalSplitterPos, HorizontalSplitterSize),
+		HorizontalBrush,
+		ShouldBeEnabled(bParentEnabled) ? ESlateDrawEffect::None : ESlateDrawEffect::DisabledEffect,
+		InWidgetStyle.GetColorAndOpacityTint() * HorizontalBrush->TintColor.GetSpecifiedColor());
+
+	const FVector2D VerticalSplitterPos(ArrangedChildren[0].Geometry.GetLocalSize().X,0);
+	const FVector2D VerticalSplitterSize(ArrangedChildren[2].Geometry.Position.X - ArrangedChildren[0].Geometry.GetLocalSize().X, AllottedGeometry.GetLocalSize().Y);
+
+	const FSlateBrush* VerticalBrush = EnumHasAnyFlags(HoveredAxis, EResizingAxis::LeftRightMask) || EnumHasAnyFlags(ResizingAxisMask, EResizingAxis::LeftRightMask) ? &Style->HandleHighlightBrush : &Style->HandleNormalBrush;
+
+	FSlateDrawElement::MakeBox(
+		OutDrawElements,
+		MaxLayerId,
+		AllottedGeometry.ToPaintGeometry(VerticalSplitterPos, VerticalSplitterSize),
+		VerticalBrush,
+		ShouldBeEnabled(bParentEnabled) ? ESlateDrawEffect::None : ESlateDrawEffect::DisabledEffect,
+		InWidgetStyle.GetColorAndOpacityTint() * VerticalBrush->TintColor.GetSpecifiedColor());
+
+	return MaxLayerId;
+}
 
 FVector2D SSplitter2x2::ComputeDesiredSize( float ) const
 {
@@ -947,8 +990,8 @@ FReply SSplitter2x2::OnMouseButtonDown( const FGeometry& MyGeometry, const FPoin
 	if ( MouseEvent.GetEffectingButton() == EKeys::LeftMouseButton )
 	{
 		const FVector2D LocalMousePos = MyGeometry.AbsoluteToLocal( MouseEvent.GetScreenSpacePosition() );
-		ResizingAxis = CalculateResizingAxis( MyGeometry, LocalMousePos );
-		if ( ResizingAxis != INDEX_NONE )
+		ResizingAxisMask = CalculateResizingAxis( MyGeometry, LocalMousePos );
+		if (ResizingAxisMask != EResizingAxis::None )
 		{
 			bIsResizing = true;
 			Reply = FReply::Handled().CaptureMouse( SharedThis(this) );
@@ -964,6 +1007,7 @@ FReply SSplitter2x2::OnMouseButtonUp( const FGeometry& MyGeometry, const FPointe
 	if ( MouseEvent.GetEffectingButton() == EKeys::LeftMouseButton && bIsResizing == true )
 	{
 		bIsResizing = false;
+		ResizingAxisMask = EResizingAxis::None;
 		return FReply::Handled().ReleaseMouseCapture();
 	}
 	return FReply::Unhandled();
@@ -980,29 +1024,29 @@ FReply SSplitter2x2::OnMouseMove( const FGeometry& MyGeometry, const FPointerEve
 		ResizeChildren( MyGeometry, LayoutChildren, LocalMousePos );
 		return FReply::Handled();
 	}
-	else
-	{	
-		ResizingAxis = CalculateResizingAxis( MyGeometry, LocalMousePos );
-		return FReply::Unhandled();
-	}
-}
 
+	return FReply::Unhandled();
+}
 
 FCursorReply SSplitter2x2::OnCursorQuery( const FGeometry& MyGeometry, const FPointerEvent& CursorEvent ) const 
 {
-	if( ResizingAxis == 0 )
+	const FVector2D LocalMousePos = MyGeometry.AbsoluteToLocal(CursorEvent.GetScreenSpacePosition());
+	EResizingAxis HoveredAxis = CalculateResizingAxis(GetTickSpaceGeometry(), LocalMousePos);
+
+	if (EnumHasAnyFlags(HoveredAxis, EResizingAxis::CrossMask) || EnumHasAnyFlags(ResizingAxisMask, EResizingAxis::CrossMask))
 	{
-		return FCursorReply::Cursor( EMouseCursor::ResizeLeftRight );
+		return FCursorReply::Cursor(EMouseCursor::CardinalCross);
 	}
-	else if( ResizingAxis == 1 )
+	else if(EnumHasAnyFlags(HoveredAxis,EResizingAxis::LeftRightMask) || EnumHasAnyFlags(ResizingAxisMask, EResizingAxis::LeftRightMask) )
 	{
-		return FCursorReply::Cursor( EMouseCursor::ResizeUpDown );
+		return FCursorReply::Cursor( EMouseCursor::ResizeLeftRight );
 	}
-	else if( ResizingAxis == 2)
+	else if(EnumHasAnyFlags(HoveredAxis, EResizingAxis::UpDownMask) || EnumHasAnyFlags(ResizingAxisMask, EResizingAxis::UpDownMask))
 	{
-		return FCursorReply::Cursor( EMouseCursor::CardinalCross );
+		return FCursorReply::Cursor( EMouseCursor::ResizeUpDown );
 	}
 
+
 	return FCursorReply::Unhandled();
 }
 
@@ -1022,12 +1066,12 @@ void SSplitter2x2::ResizeChildren( const FGeometry& MyGeometry, const TArray<FLa
 	FSlot& TopRight = Children[2];
 	FSlot& BotRight = Children[3];
 
-	if( ResizingAxis == 0 )
+	if(ResizingAxisMask == EResizingAxis::LeftRightMask )
 	{
 		// Ensure deltas along the Y axis are not taken into account
 		Delta.Y = 0;
 	}
-	else if( ResizingAxis == 1 )
+	else if(ResizingAxisMask == EResizingAxis::UpDownMask )
 	{
 		// Ensure deltas along the X axis are not taken into account
 		Delta.X = 0;
@@ -1039,14 +1083,14 @@ void SSplitter2x2::ResizeChildren( const FGeometry& MyGeometry, const TArray<FLa
 	FVector2D NewSizeTR;
 	FVector2D NewSizeBR;
 
-	if( ResizingAxis == 0)
+	if( ResizingAxisMask == EResizingAxis::LeftRightMask)
 	{
 		NewSizeTL = TopLeftSize + Delta;
 		NewSizeBL = BotLeftSize + Delta;
 		NewSizeTR = TopRightSize - Delta;
 		NewSizeBR = BotRightSize - Delta;
 	}
-	else if( ResizingAxis == 1 )
+	else if( ResizingAxisMask == EResizingAxis::UpDownMask)
 	{
 		NewSizeTL = TopLeftSize + Delta;
 		NewSizeBL = BotLeftSize - Delta;
@@ -1097,9 +1141,9 @@ void SSplitter2x2::ResizeChildren( const FGeometry& MyGeometry, const TArray<FLa
 
 }
 
-int32 SSplitter2x2::CalculateResizingAxis( const FGeometry& MyGeometry, const FVector2D& LocalMousePos ) const
+EResizingAxis SSplitter2x2::CalculateResizingAxis( const FGeometry& MyGeometry, const FVector2D& LocalMousePos ) const
 {
-	int32 Axis = INDEX_NONE;
+	EResizingAxis Axis = EResizingAxis::None;
 
 	TArray<FLayoutGeometry> LayoutChildren = ArrangeChildrenForLayout(MyGeometry);
 
@@ -1117,13 +1161,13 @@ int32 SSplitter2x2::CalculateResizingAxis( const FGeometry& MyGeometry, const FV
 		{
 			// The mouse is in between two viewports vertically
 			// Resizing axis is the X axis
-			Axis = 0;
+			Axis = EResizingAxis::LeftRightMask;
 		}
 		else if( LocalMousePos.Y > PrevBound.Y && LocalMousePos.Y < NextBound.Y )
 		{
 			// The mouse is in between two viewports horizontally
 			// Resizing axis is the Y axis
-			Axis = 1;
+			Axis = EResizingAxis::UpDownMask;
 		}
 		else
 		{
@@ -1134,7 +1178,7 @@ int32 SSplitter2x2::CalculateResizingAxis( const FGeometry& MyGeometry, const FV
 
 	if( bInCenter )
 	{
-		Axis = 2;
+		Axis = EResizingAxis::CrossMask;
 	}
 
 	return Axis;
diff --git a/Engine/Source/Runtime/Slate/Public/Widgets/Layout/SSplitter.h b/Engine/Source/Runtime/Slate/Public/Widgets/Layout/SSplitter.h
--- a/Engine/Source/Runtime/Slate/Public/Widgets/Layout/SSplitter.h
+++ b/Engine/Source/Runtime/Slate/Public/Widgets/Layout/SSplitter.h
@@ -385,6 +385,15 @@ protected:
 	const FSplitterStyle* Style;
 };
 
+enum class EResizingAxis : uint8
+{
+	None = 0x00,
+	LeftRightMask = 0x01,
+	UpDownMask = 0x10,
+	CrossMask = UpDownMask | LeftRightMask
+};
+ENUM_CLASS_FLAGS(EResizingAxis);
+
 /**
  * SSplitter2x2													
  * A splitter which has exactly 4 children and allows simultaneous		
@@ -425,7 +434,12 @@ private:
 		TAttribute<FVector2D> PercentageAttribute;
 	};
 
-	SLATE_BEGIN_ARGS( SSplitter2x2 ){}
+	SLATE_BEGIN_ARGS( SSplitter2x2 )
+		: _Style(&FCoreStyle::Get().GetWidgetStyle<FSplitterStyle>("Splitter"))
+		{
+		}
+		/** Style used to draw this splitter */
+		SLATE_STYLE_ARGUMENT(FSplitterStyle, Style)
 		SLATE_NAMED_SLOT( FArguments, TopLeft )
 		SLATE_NAMED_SLOT( FArguments, BottomLeft )
 		SLATE_NAMED_SLOT( FArguments, TopRight )
@@ -505,6 +519,8 @@ private:
 
 	virtual void OnArrangeChildren( const FGeometry& AllottedGeometry, FArrangedChildren& ArrangedChildren ) const override;
 
+	virtual int32 OnPaint(const FPaintArgs& Args, const FGeometry& AllottedGeometry, const FSlateRect& MyCullingRect, FSlateWindowElementList& OutDrawElements, int32 LayerId, const FWidgetStyle& InWidgetStyle, bool bParentEnabled) const override;
+
 	/**
 	 * A Panel's desired size in the space required to arrange of its children on the screen while respecting all of
 	 * the children's desired sizes and any layout-related options specified by the user. See StackPanel for an example.
@@ -561,7 +577,7 @@ private:
 	 * @param MyGeometry	The geometry of this widget
 	 * @param LocalMousePos	The local space current mouse position
 	 */
-	int32 CalculateResizingAxis( const FGeometry& MyGeometry, const FVector2D& LocalMousePos ) const;
+	EResizingAxis CalculateResizingAxis( const FGeometry& MyGeometry, const FVector2D& LocalMousePos ) const;
 
 	/**
 	 * Resizes all children based on a user moving the splitter handles
@@ -577,8 +593,10 @@ private:
 	/** The children of the splitter. There can only be four */
 	TPanelChildren<FSlot> Children;
 
+	const FSplitterStyle* Style = nullptr;
+
 	/** The axis currently being resized or INDEX_NONE if no resizing */
-	int32 ResizingAxis;
+	EResizingAxis ResizingAxisMask;
 
 	/** true if a splitter axis is currently being resized. */
 	bool bIsResizing;
