SGridPanel no longer creates new columns/rows past the actual desired amount of columns/rows when Column/Row Span is set to a large value, which prevents crashes (due to bounds violations) and very long loops. Instead, all column accesses are clamped to the current number of actual columns. This should not affect any current behaviour, as the number of columns is internal to the widget.

diff --git a/Engine/Source/Runtime/Slate/Private/Widgets/Layout/SGridPanel.cpp b/Engine/Source/Runtime/Slate/Private/Widgets/Layout/SGridPanel.cpp
--- a/Engine/Source/Runtime/Slate/Private/Widgets/Layout/SGridPanel.cpp
+++ b/Engine/Source/Runtime/Slate/Private/Widgets/Layout/SGridPanel.cpp
@@ -252,11 +252,16 @@ void SGridPanel::OnArrangeChildren( const FGeometry& AllottedGeometry, FArranged
 		{
 			// Figure out the position of this cell.
 			const FVector2D ThisCellOffset( FinalColumns[CurSlot.GetColumn()], FinalRows[CurSlot.GetRow()] );
+
 			// Figure out the size of this slot; takes row span into account.
 			// We use the properties of partial sums arrays to achieve this.
+			// Row/column span may be an arbitrarily large value above 1, so clamp it to the array bounds.
+			const int32 LastColumnIdx = FMath::Min(CurSlot.GetColumn() + CurSlot.GetColumnSpan(), FinalColumns.Num() - 1);
+			const int32 LastRowIdx = FMath::Min(CurSlot.GetRow() + CurSlot.GetRowSpan(), FinalRows.Num() - 1);
 			const FVector2D CellSize(
-				FinalColumns[CurSlot.GetColumn()+CurSlot.GetColumnSpan()] - ThisCellOffset.X ,
-				FinalRows[CurSlot.GetRow()+CurSlot.GetRowSpan()] - ThisCellOffset.Y );
+				FinalColumns[LastColumnIdx] - ThisCellOffset.X,
+				FinalRows[LastRowIdx] - ThisCellOffset.Y
+			);
 
 			// Do the standard arrangement of elements within a slot
 			// Takes care of alignment and padding.
@@ -403,13 +408,13 @@ void SGridPanel::NotifySlotChanged(const FSlot* InSlot, bool bSlotLayerChanged /
 	// We need an extra cell at the end for easily figuring out the size across any number of cells
 	// by doing Columns[End] - Columns[Start] or Rows[End] - Rows[Start].
 	// The first Columns[]/Rows[] entry will be 0.
-	const int32 NumColumnsRequiredForThisSlot = InSlot->GetColumn() + InSlot->GetColumnSpan() + 1;
+	const int32 NumColumnsRequiredForThisSlot = InSlot->GetColumn() + 1;
 	if ( NumColumnsRequiredForThisSlot > Columns.Num() )
 	{
 		Columns.AddZeroed( NumColumnsRequiredForThisSlot - Columns.Num() );
 	}
 
-	const int32 NumRowsRequiredForThisSlot = InSlot->GetRow() + InSlot->GetRowSpan() + 1;
+	const int32 NumRowsRequiredForThisSlot = InSlot->GetRow() + 1;
 	if ( NumRowsRequiredForThisSlot > Rows.Num() )
 	{
 		Rows.AddZeroed( NumRowsRequiredForThisSlot - Rows.Num() );
@@ -428,10 +433,10 @@ void SGridPanel::NotifySlotChanged(const FSlot* InSlot, bool bSlotLayerChanged /
 
 void SGridPanel::ComputeDesiredCellSizes( TArray<float>& OutColumns, TArray<float>& OutRows ) const
 {
-	FMemory::Memzero( OutColumns.GetData(), OutColumns.Num() * sizeof(float) );
-	FMemory::Memzero( OutRows.GetData(), OutRows.Num() * sizeof(float) );
+	FMemory::Memzero(OutColumns.GetData(), OutColumns.Num() * sizeof(float));
+	FMemory::Memzero(OutRows.GetData(), OutRows.Num() * sizeof(float));
 
-	for (int32 SlotIndex=0; SlotIndex<Slots.Num(); ++SlotIndex)
+	for (int32 SlotIndex = 0; SlotIndex < Slots.Num(); ++SlotIndex)
 	{
 		const FSlot& CurSlot = Slots[SlotIndex];
 		if (CurSlot.GetWidget()->GetVisibility() != EVisibility::Collapsed)
@@ -439,13 +444,19 @@ void SGridPanel::ComputeDesiredCellSizes( TArray<float>& OutColumns, TArray<floa
 			// The slots wants to be as big as its content along with the required padding.
 			const FVector2D SlotDesiredSize = CurSlot.GetWidget()->GetDesiredSize() + CurSlot.GetPadding().GetDesiredSize();
 
-			// If the slot has a (colspan,rowspan) of (1,1) it will only affect that slot.
+			// If the slot has a (colspan, rowspan) of (1,1) it will only affect that slot.
 			// For larger spans, the slots size will be evenly distributed across all the affected slots.
-			const FVector2D SizeContribution( SlotDesiredSize.X / CurSlot.GetColumnSpan(), SlotDesiredSize.Y / CurSlot.GetRowSpan() );
+			// For spans that are past the end of the slot array, only distribute to the size to existing columns, eg.
+			// | [0] | [1] | [2] |
+			// | S=1 |   Span=5  | <- size is only distributed between columns 1 and 2, until more slots exist
+			const int32 MaxColumns = FMath::Min(CurSlot.GetColumn() + CurSlot.GetColumnSpan(), OutColumns.Num());
+			const int32 MaxRows = FMath::Min(CurSlot.GetRow() + CurSlot.GetRowSpan(), OutRows.Num());
+
+			const FVector2D SizeContribution(SlotDesiredSize.X / (MaxColumns - CurSlot.GetColumn()), SlotDesiredSize.Y / (MaxRows - CurSlot.GetRow()));
 
 			// Distribute the size contributions over all the columns and rows that this slot spans
-			DistributeSizeContributions( SizeContribution.X, OutColumns, CurSlot.GetColumn(), CurSlot.GetColumn() + CurSlot.GetColumnSpan());
-			DistributeSizeContributions( SizeContribution.Y, OutRows, CurSlot.GetRow(), CurSlot.GetRow() + CurSlot.GetRowSpan() );
+			DistributeSizeContributions(SizeContribution.X, OutColumns, CurSlot.GetColumn(), MaxColumns);
+			DistributeSizeContributions(SizeContribution.Y, OutRows, CurSlot.GetRow(), MaxRows);
 		}
 	}
 }
