SVectorInputBox: Support delegates / attributes for vectors directly

diff --git a/Engine/Source/Runtime/Slate/Public/Widgets/Input/SVectorInputBox.h b/Engine/Source/Runtime/Slate/Public/Widgets/Input/SVectorInputBox.h
--- a/Engine/Source/Runtime/Slate/Public/Widgets/Input/SVectorInputBox.h
+++ b/Engine/Source/Runtime/Slate/Public/Widgets/Input/SVectorInputBox.h
@@ -254,6 +254,8 @@ public:
 	 */
 	void Construct(const FArguments& InArgs)
 	{
+		bUseVectorGetter = true;
+
 		TSharedRef<SHorizontalBox> HorizontalBox = SNew(SHorizontalBox);
 
 		ChildSlot
@@ -261,6 +263,97 @@ public:
 			HorizontalBox
 		];
 
+
+		VectorAttribute = InArgs._Vector;
+		OnVectorValueChanged = InArgs._OnVectorChanged;
+		OnVectorValueCommitted = InArgs._OnVectorCommitted;
+			
+		if(!VectorAttribute.IsBound() && !VectorAttribute.IsSet())
+		{
+			bUseVectorGetter = false;
+
+			VectorAttribute = TAttribute<TOptional<VectorType>>::CreateLambda([InArgs]() -> TOptional<VectorType>
+			{
+				if constexpr (NumberOfComponents == 2)
+				{
+					TOptional<NumericType> X = InArgs._X.Get();
+					TOptional<NumericType> Y = InArgs._Y.Get();
+					if(X.IsSet() && Y.IsSet())
+					{
+						return VectorType(X.GetValue(), Y.GetValue());
+					}
+				}
+				if constexpr (NumberOfComponents == 3)
+				{
+					TOptional<NumericType> X = InArgs._X.Get();
+					TOptional<NumericType> Y = InArgs._Y.Get();
+					TOptional<NumericType> Z = InArgs._Z.Get();
+					if(X.IsSet() && Y.IsSet() && Z.IsSet())
+					{
+						return VectorType(X.GetValue(), Y.GetValue(), Z.GetValue());
+					}
+				}
+				if constexpr (NumberOfComponents == 4)
+				{
+					TOptional<NumericType> X = InArgs._X.Get();
+					TOptional<NumericType> Y = InArgs._Y.Get();
+					TOptional<NumericType> Z = InArgs._Z.Get();
+					TOptional<NumericType> W = InArgs._W.Get();
+					if(X.IsSet() && Y.IsSet() && Z.IsSet() && W.IsSet())
+					{
+						return VectorType(X.GetValue(), Y.GetValue(), Z.GetValue(), W.GetValue());
+					}
+				}
+
+				return TOptional<VectorType>();
+			});
+		}
+
+		if(!OnVectorValueChanged.IsBound())
+		{
+			OnVectorValueChanged = FOnVectorValueChanged::CreateLambda([InArgs](VectorType Vector)
+			{
+				if constexpr (NumberOfComponents >= 1)
+				{
+					InArgs._OnXChanged.ExecuteIfBound(Vector.X);
+				}
+				if constexpr (NumberOfComponents >= 2)
+				{
+					InArgs._OnYChanged.ExecuteIfBound(Vector.Y);
+				}
+				if constexpr (NumberOfComponents >= 3)
+				{
+					InArgs._OnZChanged.ExecuteIfBound(Vector.Z);
+				}
+				if constexpr (NumberOfComponents >= 4)
+				{
+					InArgs._OnWChanged.ExecuteIfBound(Vector.W);
+				}
+			});
+		}
+		if(!OnVectorValueCommitted.IsBound())
+		{
+			OnVectorValueCommitted = FOnVectorValueCommitted::CreateLambda([InArgs](VectorType Vector, ETextCommit::Type CommitType)
+			{
+				if constexpr (NumberOfComponents >= 1)
+				{
+					InArgs._OnXCommitted.ExecuteIfBound(Vector.X, CommitType);
+				}
+				if constexpr (NumberOfComponents >= 2)
+				{
+					InArgs._OnYCommitted.ExecuteIfBound(Vector.Y, CommitType);
+				}
+				if constexpr (NumberOfComponents >= 3)
+				{
+					InArgs._OnZCommitted.ExecuteIfBound(Vector.Z, CommitType);
+				}
+				if constexpr (NumberOfComponents >= 4)
+				{
+					InArgs._OnWCommitted.ExecuteIfBound(Vector.W, CommitType);
+				}
+			});
+		}
+
 		if constexpr (NumberOfComponents >= 1)
 		{
 			ConstructX(InArgs, HorizontalBox);
@@ -301,7 +394,10 @@ private:
 			LabelWidget = SNumericEntryBox<NumericType>::BuildNarrowColorLabel(LabelColor);
 		}
 
-		TAttribute<TOptional<NumericType>> Value = CreatePerComponentGetter(ComponentIndex, Component, InArgs._Vector);
+		TAttribute<TOptional<NumericType>> Value = CreatePerComponentGetter(ComponentIndex, Component, VectorAttribute);
+
+		// any other getter below can use the vector
+		TGuardValue<bool> UseVectorGetterGuard(bUseVectorGetter, true);
 
 		HorizontalBox->AddSlot()
 		[
@@ -311,8 +407,8 @@ private:
 			.SpinBoxStyle(InArgs._SpinBoxStyle)
 			.Font(InArgs._Font)
 			.Value(Value)
-			.OnValueChanged(CreatePerComponentChanged(ComponentIndex, OnComponentChanged, InArgs._Vector, InArgs._OnVectorChanged, InArgs._ConstrainVector))
-			.OnValueCommitted(CreatePerComponentCommitted(ComponentIndex, OnComponentCommitted, InArgs._Vector, InArgs._OnVectorCommitted, InArgs._ConstrainVector))
+			.OnValueChanged(CreatePerComponentChanged(ComponentIndex, OnComponentChanged, InArgs._ConstrainVector))
+			.OnValueCommitted(CreatePerComponentCommitted(ComponentIndex, OnComponentCommitted, InArgs._ConstrainVector))
 			.ToolTipText(MakeAttributeLambda([Value, TooltipText]
 			{
 				if (Value.Get().IsSet())
@@ -429,7 +525,7 @@ private:
 		const TAttribute<TOptional<NumericType>>& Component,
 		const TAttribute<TOptional<VectorType>>& Vector)
 	{
-		if(Vector.IsBound() || Vector.IsSet())
+		if(bUseVectorGetter && (Vector.IsBound() || Vector.IsSet()))
 		{
 			return TAttribute<TOptional<NumericType>>::CreateLambda(
 				[ComponentIndex, Component, Vector]() -> TOptional<NumericType>
@@ -451,30 +547,26 @@ private:
 	FOnNumericValueChanged CreatePerComponentChanged(
 		int32 ComponentIndex,
 		const FOnNumericValueChanged OnComponentChanged,
-		const TAttribute<TOptional<VectorType>>& Vector,
-		const FOnVectorValueChanged OnVectorValueChanged,
 		const FOnConstrainVector ConstrainVector)
 	{
-		if(OnVectorValueChanged.IsBound())
+		if(ConstrainVector.IsBound() && OnVectorValueChanged.IsBound())
 		{
 			return FOnNumericValueChanged::CreateLambda(
-				[ComponentIndex, OnComponentChanged, Vector, OnVectorValueChanged, ConstrainVector](NumericType InValue)
+				[ComponentIndex, OnComponentChanged, this, ConstrainVector](NumericType InValue)
 				{
-					OnComponentChanged.ExecuteIfBound(InValue);
-					
-					const TOptional<VectorType> OptionalVectorValue = Vector.Get();
+					const TOptional<VectorType> OptionalVectorValue = VectorAttribute.Get();
 					if(OptionalVectorValue.IsSet())
 					{
 						VectorType VectorValue = OptionalVectorValue.GetValue();
 						VectorValue[ComponentIndex] = InValue;
 
-						if(ConstrainVector.IsBound())
-						{
-							ConstrainVector.Execute(ComponentIndex, OptionalVectorValue.GetValue(), VectorValue);
-						}
-
+						ConstrainVector.ExecuteIfBound(ComponentIndex, OptionalVectorValue.GetValue(), VectorValue);
 						OnVectorValueChanged.Execute(VectorValue);
 					}
+					else
+					{
+						OnComponentChanged.ExecuteIfBound(InValue);
+					}
 				});
 		}
 		return OnComponentChanged;
@@ -486,18 +578,14 @@ private:
 	FOnNumericValueCommitted CreatePerComponentCommitted(
 		int32 ComponentIndex,
 		const FOnNumericValueCommitted OnComponentCommitted,
-		const TAttribute<TOptional<VectorType>>& Vector,
-		const FOnVectorValueCommitted OnVectorValueCommitted,
 		const FOnConstrainVector ConstrainVector)
 	{
-		if(OnVectorValueCommitted.IsBound())
+		if(ConstrainVector.IsBound() && OnVectorValueCommitted.IsBound())
 		{
 			return FOnNumericValueCommitted::CreateLambda(
-				[ComponentIndex, OnComponentCommitted, Vector, OnVectorValueCommitted, ConstrainVector](NumericType InValue, ETextCommit::Type CommitType)
+				[ComponentIndex, OnComponentCommitted, this, ConstrainVector](NumericType InValue, ETextCommit::Type CommitType)
 				{
-					OnComponentCommitted.ExecuteIfBound(InValue, CommitType);
-					
-					const TOptional<VectorType> OptionalVectorValue = Vector.Get();
+					const TOptional<VectorType> OptionalVectorValue = VectorAttribute.Get();
 					if(OptionalVectorValue.IsSet())
 					{
 						VectorType VectorValue = OptionalVectorValue.GetValue();
@@ -510,10 +598,19 @@ private:
 
 						OnVectorValueCommitted.Execute(VectorValue, CommitType);
 					}
+					else
+					{
+						OnComponentCommitted.ExecuteIfBound(InValue, CommitType);
+					}
 				});
 		}
 		return OnComponentCommitted;
-	}		
+	}
+
+	bool bUseVectorGetter = true;
+	TAttribute<TOptional<VectorType>> VectorAttribute;
+	FOnVectorValueChanged OnVectorValueChanged;
+	FOnVectorValueCommitted OnVectorValueCommitted;
 };
 
 /**
