Fix to SGridPanel adding one row/column less than necessary due to a mistake in CL 18632502.

diff --git a/Engine/Source/Runtime/Slate/Private/Widgets/Layout/SGridPanel.cpp b/Engine/Source/Runtime/Slate/Private/Widgets/Layout/SGridPanel.cpp
--- a/Engine/Source/Runtime/Slate/Private/Widgets/Layout/SGridPanel.cpp
+++ b/Engine/Source/Runtime/Slate/Private/Widgets/Layout/SGridPanel.cpp
@@ -408,16 +408,18 @@ void SGridPanel::NotifySlotChanged(const FSlot* InSlot, bool bSlotLayerChanged /
 	// We need an extra cell at the end for easily figuring out the size across any number of cells
 	// by doing Columns[End] - Columns[Start] or Rows[End] - Rows[Start].
 	// The first Columns[]/Rows[] entry will be 0.
-	const int32 NumColumnsRequiredForThisSlot = InSlot->GetColumn() + 1;
-	if ( NumColumnsRequiredForThisSlot > Columns.Num() )
+	// For a grid with 2 rows (index 0 and 1), we need 3 slots, thus the + 2 here.
+
+	const int32 NumColumnsRequiredForThisSlot = InSlot->GetColumn() + 2;
+	if (NumColumnsRequiredForThisSlot > Columns.Num())
 	{
-		Columns.AddZeroed( NumColumnsRequiredForThisSlot - Columns.Num() );
+		Columns.AddZeroed(NumColumnsRequiredForThisSlot - Columns.Num());
 	}
 
-	const int32 NumRowsRequiredForThisSlot = InSlot->GetRow() + 1;
-	if ( NumRowsRequiredForThisSlot > Rows.Num() )
+	const int32 NumRowsRequiredForThisSlot = InSlot->GetRow() + 2;
+	if (NumRowsRequiredForThisSlot > Rows.Num())
 	{
-		Rows.AddZeroed( NumRowsRequiredForThisSlot - Rows.Num() );
+		Rows.AddZeroed(NumRowsRequiredForThisSlot - Rows.Num());
 	}
 
 	if (bSlotLayerChanged)
