SNumericEntryBox: Avoid change event during construction

diff --git a/Engine/Source/Runtime/Slate/Public/Widgets/Input/SNumericEntryBox.h b/Engine/Source/Runtime/Slate/Public/Widgets/Input/SNumericEntryBox.h
--- a/Engine/Source/Runtime/Slate/Public/Widgets/Input/SNumericEntryBox.h
+++ b/Engine/Source/Runtime/Slate/Public/Widgets/Input/SNumericEntryBox.h
@@ -385,7 +385,7 @@ public:
 
 		if (bDisplayToggle)
 		{
-			HandleToggleCheckBoxChanged(InArgs._ToggleChecked.Get(), InArgs._OnToggleChanged);
+			HandleToggleCheckBoxChanged(InArgs._ToggleChecked.Get(), FOnCheckStateChanged());
 		}
 	}
 
