Scale rounded blurs to fix UV calculating for corner rounding

diff --git a/Engine/Source/Runtime/Slate/Private/Widgets/Layout/SBackgroundBlur.cpp b/Engine/Source/Runtime/Slate/Private/Widgets/Layout/SBackgroundBlur.cpp
--- a/Engine/Source/Runtime/Slate/Private/Widgets/Layout/SBackgroundBlur.cpp
+++ b/Engine/Source/Runtime/Slate/Private/Widgets/Layout/SBackgroundBlur.cpp
@@ -122,7 +122,7 @@ int32 SBackgroundBlur::OnPaint(const FPaintArgs& Args, const FGeometry& Allotted
 			if (Strength > 0.f)
 			{
 				FSlateRect RenderBoundingRect = AllottedGeometry.GetRenderBoundingRect();
-				FPaintGeometry PaintGeometry(RenderBoundingRect.GetTopLeft(), RenderBoundingRect.GetSize(), 1.0f);
+				FPaintGeometry PaintGeometry(RenderBoundingRect.GetTopLeft(), RenderBoundingRect.GetSize(), AllottedGeometry.GetAccumulatedLayoutTransform().GetScale());
 
 				int32 RenderTargetWidth = FMath::RoundToInt(RenderBoundingRect.GetSize().X);
 				int32 RenderTargetHeight = FMath::RoundToInt(RenderBoundingRect.GetSize().Y);
diff --git a/Engine/Source/Runtime/SlateCore/Private/Rendering/DrawElements.cpp b/Engine/Source/Runtime/SlateCore/Private/Rendering/DrawElements.cpp
--- a/Engine/Source/Runtime/SlateCore/Private/Rendering/DrawElements.cpp
+++ b/Engine/Source/Runtime/SlateCore/Private/Rendering/DrawElements.cpp
@@ -672,7 +672,7 @@ void FSlateDrawElement::MakePostProcessPass(FSlateWindowElementList& ElementList
 	FSlatePostProcessPayload& DataPayload = ElementList.CreatePayload<FSlatePostProcessPayload>(Element);
 	DataPayload.DownsampleAmount = DownsampleAmount;
 	DataPayload.PostProcessData = Params;
-	DataPayload.CornerRadius = CornerRadius;
+	DataPayload.CornerRadius = CornerRadius * PaintGeometry.DrawScale;
 
 	Element.Init(ElementList, EElementType::ET_PostProcessPass, InLayer, PaintGeometry, ESlateDrawEffect::None);
 }
