Set Scroll Bar Padding of SScrollBar within SScrollBox to 0.0f as the scroll box modifies its scroll bar padding on the slot in which the scroll bar is laid out, not the scroll bar itself.
#tests: Verified extra padding of 2.0f was removed

diff --git a/Engine/Source/Runtime/Slate/Private/Widgets/Layout/SScrollBox.cpp b/Engine/Source/Runtime/Slate/Private/Widgets/Layout/SScrollBox.cpp
--- a/Engine/Source/Runtime/Slate/Private/Widgets/Layout/SScrollBox.cpp
+++ b/Engine/Source/Runtime/Slate/Private/Widgets/Layout/SScrollBox.cpp
@@ -207,6 +207,7 @@ TSharedPtr<SScrollBar> SScrollBox::ConstructScrollBar()
 	return TSharedPtr<SScrollBar>(SNew(SScrollBar)
 		.Style(ScrollBarStyle)
 		.Orientation(Orientation)
+		.Padding(0.0f)
 		.OnUserScrolled(this, &SScrollBox::ScrollBar_OnUserScrolled));
 }
 
diff --git a/Engine/Source/Runtime/Slate/Public/Widgets/Layout/SScrollBox.h b/Engine/Source/Runtime/Slate/Public/Widgets/Layout/SScrollBox.h
--- a/Engine/Source/Runtime/Slate/Public/Widgets/Layout/SScrollBox.h
+++ b/Engine/Source/Runtime/Slate/Public/Widgets/Layout/SScrollBox.h
@@ -118,6 +118,7 @@ public:
 
 		SLATE_ARGUMENT( FVector2D, ScrollBarThickness )
 
+		/** This accounts for total internal scroll bar padding; default 2.0f padding from the scroll bar itself is removed */
 		SLATE_ARGUMENT( FMargin, ScrollBarPadding )
 
 		SLATE_ARGUMENT(EAllowOverscroll, AllowOverscroll);
