SScaleBox Prepass: scale box is getting old data and passing it along to children, mark prepass dirty so we get a new desiredsize from the child

diff --git a/Engine/Source/Runtime/Slate/Private/Widgets/Layout/SScaleBox.cpp b/Engine/Source/Runtime/Slate/Private/Widgets/Layout/SScaleBox.cpp
--- a/Engine/Source/Runtime/Slate/Private/Widgets/Layout/SScaleBox.cpp
+++ b/Engine/Source/Runtime/Slate/Private/Widgets/Layout/SScaleBox.cpp
@@ -83,6 +83,7 @@ bool SScaleBox::CustomPrepass(float LayoutScaleMultiplier)
 	// size, do that now.
 	if (bNeedsNormalizingPrepassOrLocalGeometry || !LastAllocatedArea.IsSet())
 	{
+		ChildSlotWidget.MarkPrepassAsDirty();
 		ChildSlotWidget.SlatePrepass(LayoutScaleMultiplier);
 
 		NormalizedContentDesiredSize = ChildSlotWidget.GetDesiredSize();
