Fixed default padding for the search result numbers slot. Previously, the numbers would poke out of the box if left aligned.
When right aligned, the clear button wouldn't look right either

diff --git a/Engine/Source/Runtime/Slate/Private/Widgets/Input/SSearchBox.cpp b/Engine/Source/Runtime/Slate/Private/Widgets/Input/SSearchBox.cpp
--- a/Engine/Source/Runtime/Slate/Private/Widgets/Input/SSearchBox.cpp
+++ b/Engine/Source/Runtime/Slate/Private/Widgets/Input/SSearchBox.cpp
@@ -56,10 +56,12 @@ void SSearchBox::Construct( const FArguments& InArgs )
 	// If a search delegate was bound, add a previous and next button
 	if (OnSearchDelegate.IsBound())
 	{
+		FMargin SearchResultPadding = InArgs._Style->bLeftAlignButtons ? FMargin(5, 0, 2, 0) : FMargin(2, 0, 2, 0);
+
 		// Search result data text
 		Box->InsertSlot(SlotIndex++)
 		.AutoWidth()
-		.Padding(0, 0, 2, 0)
+		.Padding(SearchResultPadding)
 		.HAlign(HAlign_Center)
 		.VAlign(VAlign_Center)
 		[
@@ -127,7 +129,9 @@ void SSearchBox::Construct( const FArguments& InArgs )
 			.Image(&InArgs._Style->GlassImage)
 			.ColorAndOpacity(FSlateColor::UseForeground())
 		];
+
 	// Add an X to clear the search whenever there is some text typed into it
+	FMargin ClearButtonMargin = InArgs._Style->bLeftAlignButtons ? FMargin(2.f, 0.f) : FMargin(2.f, 0.f, 4.f, 0.f);
 	Box->InsertSlot(SlotIndex++)
 	.AutoWidth()
 	.HAlign(HAlign_Center)
@@ -136,7 +140,7 @@ void SSearchBox::Construct( const FArguments& InArgs )
 		SNew(SButton)
 		.Visibility(this, &SSearchBox::GetXVisibility)
 		.ButtonStyle(FAppStyle::Get(), "HoverOnlyButton")
-		.ContentPadding(FMargin(2.f, 0.0f))
+		.ContentPadding(ClearButtonMargin)
 		.HAlign(HAlign_Center)
 		.VAlign(VAlign_Center)
 		.OnClicked(this, &SSearchBox::OnClearSearch)
