Fixed crash when SDockingTab::GetTabManager() dereferenced an invalid weak pointer. GetTabManager has been deprecated in favor of GetTabManagerPtr which returns a TSharedPtr.

#rnx

diff --git a/Engine/Source/Developer/OutputLog/Private/SOutputLog.cpp b/Engine/Source/Developer/OutputLog/Private/SOutputLog.cpp
--- a/Engine/Source/Developer/OutputLog/Private/SOutputLog.cpp
+++ b/Engine/Source/Developer/OutputLog/Private/SOutputLog.cpp
@@ -1808,7 +1808,10 @@ FReply SOutputLog::OnDockInLayoutClicked()
 	static const FName OutputLogTabName = FName("OutputLog");
 	if (TSharedPtr<SDockTab> ActiveTab = FGlobalTabmanager::Get()->GetActiveTab())
 	{
-		DockedTab = ActiveTab->GetTabManager()->TryInvokeTab(OutputLogTabName);
+		if (TSharedPtr<FTabManager> TabManager = ActiveTab->GetTabManagerPtr())
+		{
+			DockedTab = TabManager->TryInvokeTab(OutputLogTabName);
+		}
 	}
 	
 	if (!DockedTab)
diff --git a/Engine/Source/Editor/ContentBrowser/Private/SContentBrowser.cpp b/Engine/Source/Editor/ContentBrowser/Private/SContentBrowser.cpp
--- a/Engine/Source/Editor/ContentBrowser/Private/SContentBrowser.cpp
+++ b/Engine/Source/Editor/ContentBrowser/Private/SContentBrowser.cpp
@@ -1386,12 +1386,12 @@ bool SContentBrowser::CanSetAsPrimaryContentBrowser() const
 
 TSharedPtr<FTabManager> SContentBrowser::GetTabManager() const
 {
-	if ( ContainingTab.IsValid() )
+	if (TSharedPtr<SDockTab> Tab = ContainingTab.Pin())
 	{
-		return ContainingTab.Pin()->GetTabManager();
+		return Tab->GetTabManagerPtr();
 	}
 
-	return NULL;
+	return TSharedPtr<FTabManager>();
 }
 
 void SContentBrowser::LoadSelectedObjectsIfNeeded()
diff --git a/Engine/Source/Editor/StatusBar/Private/SStatusBar.cpp b/Engine/Source/Editor/StatusBar/Private/SStatusBar.cpp
--- a/Engine/Source/Editor/StatusBar/Private/SStatusBar.cpp
+++ b/Engine/Source/Editor/StatusBar/Private/SStatusBar.cpp
@@ -815,7 +815,7 @@ void SStatusBar::OnActiveTabChanged(TSharedPtr<SDockTab> PreviouslyActive, TShar
 		// Remove the content browser if a newly activated tab is a major tab
 		bShouldRemoveDrawer = true;
 	}
-	else if (PreviouslyActive->GetTabManager() != NewlyActivated->GetTabManager())
+	else if (PreviouslyActive->GetTabManagerPtr() != NewlyActivated->GetTabManagerPtr())
 	{
 		// Remove the content browser if we're switching tab managers (indicates a new status bar is becoming active)
 		bShouldRemoveDrawer = true;
diff --git a/Engine/Source/Editor/StatusBar/Private/StatusBarSubsystem.cpp b/Engine/Source/Editor/StatusBar/Private/StatusBarSubsystem.cpp
--- a/Engine/Source/Editor/StatusBar/Private/StatusBarSubsystem.cpp
+++ b/Engine/Source/Editor/StatusBar/Private/StatusBarSubsystem.cpp
@@ -44,9 +44,12 @@ namespace UE
 				{
 					if (TSharedPtr<SDockTab> ActiveTab = FGlobalTabmanager::Get()->GetActiveTab())
 					{
-						if (TSharedPtr<SDockTab> ActiveMajorTab = FGlobalTabmanager::Get()->GetMajorTabForTabManager(ActiveTab->GetTabManager()))
+						if (TSharedPtr<FTabManager> ActiveTabManager = ActiveTab->GetTabManagerPtr())
 						{
-							ParentWindow = ActiveMajorTab->GetParentWindow();
+							if (TSharedPtr<SDockTab> ActiveMajorTab = FGlobalTabmanager::Get()->GetMajorTabForTabManager(ActiveTabManager.ToSharedRef()))
+							{
+								ParentWindow = ActiveMajorTab->GetParentWindow();
+							}
 						}
 					}
 				}
@@ -273,7 +276,7 @@ bool UStatusBarSubsystem::ToggleDebugConsole(TSharedRef<SWindow> ParentWindow, b
 				FWidgetPath ConsoleEditBoxPath;
 				if(FSlateApplication::Get().GeneratePathToWidgetUnchecked(SBData.ConsoleEditBox.ToSharedRef(), ConsoleEditBoxPath))
 				{
-					if (!!bAlwaysToggleDrawer && MainOutputLogTab && MainOutputLogTab->GetTabManager()->GetOwnerTab() == ParentTab)
+					if (bAlwaysToggleDrawer && MainOutputLogTab && MainOutputLogTab->GetTabManagerPtr() && MainOutputLogTab->GetTabManagerPtr()->GetOwnerTab() == ParentTab)
 					{
 						OutputLogModule.FocusOutputLogConsoleBox(OutputLogModule.GetOutputLog().ToSharedRef());
 					}
diff --git a/Engine/Source/Editor/UnrealEd/Private/Toolkits/GlobalEditorCommonCommands.cpp b/Engine/Source/Editor/UnrealEd/Private/Toolkits/GlobalEditorCommonCommands.cpp
--- a/Engine/Source/Editor/UnrealEd/Private/Toolkits/GlobalEditorCommonCommands.cpp
+++ b/Engine/Source/Editor/UnrealEd/Private/Toolkits/GlobalEditorCommonCommands.cpp
@@ -190,9 +190,12 @@ void FGlobalEditorCommonCommands::OnSummonedConsoleCommandBox()
 	{
 		if (TSharedPtr<SDockTab> ActiveTab = FGlobalTabmanager::Get()->GetActiveTab())
 		{
-			if (TSharedPtr<SDockTab> ActiveMajorTab = FGlobalTabmanager::Get()->GetMajorTabForTabManager(ActiveTab->GetTabManager()))
+			if (TSharedPtr<FTabManager> ActiveTabManager = ActiveTab->GetTabManagerPtr())
 			{
-				ParentWindow = ActiveMajorTab->GetParentWindow();
+				if (TSharedPtr<SDockTab> ActiveMajorTab = FGlobalTabmanager::Get()->GetMajorTabForTabManager(ActiveTabManager.ToSharedRef()))
+				{
+					ParentWindow = ActiveMajorTab->GetParentWindow();
+				}
 			}
 		}
 	}
diff --git a/Engine/Source/Runtime/Slate/Private/Framework/Docking/FDockingDragOperation.cpp b/Engine/Source/Runtime/Slate/Private/Framework/Docking/FDockingDragOperation.cpp
--- a/Engine/Source/Runtime/Slate/Private/Framework/Docking/FDockingDragOperation.cpp
+++ b/Engine/Source/Runtime/Slate/Private/Framework/Docking/FDockingDragOperation.cpp
@@ -273,7 +273,7 @@ FDockingDragOperation::FDockingDragOperation( const TSharedRef<SDockTab>& InTabT
 		SNew(SBorder)
 		. BorderImage(FCoreStyle::Get().GetBrush("Docking.Background"))
 		[
-			SNew(SDockingArea, TabBeingDragged->GetTabManager(), FTabManager::NewPrimaryArea())
+			SNew(SDockingArea, TabBeingDragged->GetTabManagerPtr().ToSharedRef(), FTabManager::NewPrimaryArea())
 			//. OriginalDockArea(OriginalDockArea)
 			. InitialContent
 			( 
@@ -314,16 +314,18 @@ void FDockingDragOperation::DroppedOntoNothing()
 {
 	// If we dropped the tab into an existing DockNode then it would have handled the DropEvent.
 	// We are here because that didn't happen, so make a new window with a new DockNode and drop the tab into that.
+	TSharedPtr<FTabManager> MyTabManager = TabBeingDragged->GetTabManagerPtr();
+	if (!MyTabManager.IsValid())
+	{
+		return;
+	}
 
 	const FVector2D PositionToDrop = CursorDecoratorWindow->GetPositionInScreen();
 
 	const float DPIScale = FPlatformApplicationMisc::GetDPIScaleFactorAtPoint(PositionToDrop.X, PositionToDrop.Y);
 
-	TSharedRef<FTabManager> MyTabManager = TabBeingDragged->GetTabManager();
-
 	TSharedPtr<SWindow> NewWindowParent = MyTabManager->GetPrivateApi().GetParentWindow();
 
-
 	TSharedRef<SWindow> NewWindow = SNew(SWindow)
 		.Title(FGlobalTabmanager::Get()->GetApplicationTitle())
 		.AutoCenter(EAutoCenter::None)
@@ -342,7 +344,7 @@ void FDockingDragOperation::DroppedOntoNothing()
 
 	// Create a new dockarea
 	TSharedRef<SDockingArea> NewDockArea =
-		SNew(SDockingArea, TabBeingDragged->GetTabManager(), FTabManager::NewPrimaryArea())
+		SNew(SDockingArea, MyTabManager.ToSharedRef(), FTabManager::NewPrimaryArea())
 		.ParentWindow(NewWindow)
 		.InitialContent
 		(
diff --git a/Engine/Source/Runtime/Slate/Private/Framework/Docking/SDockingTabStack.cpp b/Engine/Source/Runtime/Slate/Private/Framework/Docking/SDockingTabStack.cpp
--- a/Engine/Source/Runtime/Slate/Private/Framework/Docking/SDockingTabStack.cpp
+++ b/Engine/Source/Runtime/Slate/Private/Framework/Docking/SDockingTabStack.cpp
@@ -1112,7 +1112,14 @@ bool SDockingTabStack::CanMoveTabToSideBar(TSharedRef<SDockTab> Tab) const
 bool SDockingTabStack::IsTabAllowedInSidebar(TSharedPtr<SDockTab> Tab) const
 {
 	// Major tabs are not allowed to be sidebared
-	return Tab.IsValid() && Tab->GetVisualTabRole() != ETabRole::MajorTab && Tab->GetTabManager()->IsTabAllowedInSidebar(Tab->GetLayoutIdentifier());
+	if (Tab.IsValid())
+	{
+		if (TSharedPtr<FTabManager> TabManager = Tab->GetTabManagerPtr())
+		{
+			return Tab->GetVisualTabRole() != ETabRole::MajorTab && TabManager->IsTabAllowedInSidebar(Tab->GetLayoutIdentifier());
+		}
+	}
+	return false;
 }
 
 SSplitter::ESizeRule SDockingTabStack::GetSizeRule() const
diff --git a/Engine/Source/Runtime/Slate/Private/Framework/Docking/SDockingTabWell.cpp b/Engine/Source/Runtime/Slate/Private/Framework/Docking/SDockingTabWell.cpp
--- a/Engine/Source/Runtime/Slate/Private/Framework/Docking/SDockingTabWell.cpp
+++ b/Engine/Source/Runtime/Slate/Private/Framework/Docking/SDockingTabWell.cpp
@@ -303,7 +303,13 @@ int32 SDockingTabWell::ComputeChildDropIndex(const FGeometry& MyGeometry, const
 
 FReply SDockingTabWell::StartDraggingTab( TSharedRef<SDockTab> TabToStartDragging, FVector2D InTabGrabOffsetFraction, const FPointerEvent& MouseEvent )
 {
-	const bool bCanLeaveTabWell = TabToStartDragging->GetTabManager()->GetPrivateApi().CanTabLeaveTabWell( TabToStartDragging );
+	TSharedPtr<FTabManager> TabManager = TabToStartDragging->GetTabManagerPtr();
+	if (!TabManager.IsValid())
+	{
+		return FReply::Handled();
+	}
+
+	const bool bCanLeaveTabWell = TabManager->GetPrivateApi().CanTabLeaveTabWell( TabToStartDragging );
 
 	// We are about to start dragging a tab, so make sure its offset is correct
 	this->ChildBeingDraggedOffset = ComputeDraggedTabOffset( MouseEvent.FindGeometry(SharedThis(this)), MouseEvent, InTabGrabOffsetFraction );
@@ -550,14 +556,18 @@ void SDockingTabWell::BringTabToFront( int32 TabIndexToActivate )
 	// Update the native, global menu bar if a tab is in the foreground.
 	if( ForegroundTabIndex != INDEX_NONE )
 	{
-		TSharedPtr<FTabManager> TabManager = Tabs[ForegroundTabIndex]->GetTabManager();
-		if(TabManager == FGlobalTabmanager::Get())
-		{
-			FGlobalTabmanager::Get()->UpdateMainMenu(Tabs[ForegroundTabIndex], false);
-		}
-		else
+		const TSharedRef<SDockTab>& ForegroundTab = Tabs[ForegroundTabIndex];
+		TSharedPtr<FTabManager> TabManager = ForegroundTab->GetTabManagerPtr();
+		if (TabManager.IsValid())
 		{
-			TabManager->UpdateMainMenu(Tabs[ForegroundTabIndex], false);
+			if (TabManager == FGlobalTabmanager::Get())
+			{
+				FGlobalTabmanager::Get()->UpdateMainMenu(ForegroundTab, false);
+			}
+			else
+			{
+				TabManager->UpdateMainMenu(ForegroundTab, false);
+			}
 		}
 	}
 }
diff --git a/Engine/Source/Runtime/Slate/Private/Framework/Docking/TabManager.cpp b/Engine/Source/Runtime/Slate/Private/Framework/Docking/TabManager.cpp
--- a/Engine/Source/Runtime/Slate/Private/Framework/Docking/TabManager.cpp
+++ b/Engine/Source/Runtime/Slate/Private/Framework/Docking/TabManager.cpp
@@ -1359,7 +1359,11 @@ TSharedPtr<SDockTab> FTabManager::InvokeTab_Internal(const FTabId& TabId, bool b
 
 		if (ExistingTab.IsValid())
 		{
-			TSharedPtr<SDockTab> MajorTab = FGlobalTabmanager::Get()->GetMajorTabForTabManager(ExistingTab->GetTabManager());
+			TSharedPtr<SDockTab> MajorTab;
+			if (TSharedPtr<FTabManager> ExistingTabManager = ExistingTab->GetTabManagerPtr())
+			{
+				MajorTab = FGlobalTabmanager::Get()->GetMajorTabForTabManager(ExistingTabManager.ToSharedRef());
+			}
 
 			// Rules for drawing attention to a tab:
 			// 1. Tab is not active
@@ -2587,16 +2591,16 @@ TSharedRef<FTabManager> FGlobalTabmanager::NewTabManager( const TSharedRef<SDock
 
 void FGlobalTabmanager::UpdateMainMenu(const TSharedRef<SDockTab>& ForTab, bool const bForce)
 {
-	TSharedPtr<FTabManager> Tabmanager = ForTab->GetTabManager();
-	if(Tabmanager == AsShared())
+	TSharedPtr<FTabManager> TabManager = ForTab->GetTabManagerPtr();
+	if(TabManager == AsShared())
 	{
 		const int32 TabIndex = SubTabManagers.IndexOfByPredicate(FindByTab(ForTab));
 		if (TabIndex != INDEX_NONE)
 		{
-			Tabmanager = SubTabManagers[TabIndex].TabManager.Pin();
+			TabManager = SubTabManagers[TabIndex].TabManager.Pin();
 		}
 	}
-	Tabmanager->UpdateMainMenu(ForTab, bForce);
+	TabManager->UpdateMainMenu(ForTab, bForce);
 }
 
 void FGlobalTabmanager::SaveAllVisualState()
diff --git a/Engine/Source/Runtime/Slate/Private/Widgets/Docking/SDockTab.cpp b/Engine/Source/Runtime/Slate/Private/Widgets/Docking/SDockTab.cpp
--- a/Engine/Source/Runtime/Slate/Private/Widgets/Docking/SDockTab.cpp
+++ b/Engine/Source/Runtime/Slate/Private/Widgets/Docking/SDockTab.cpp
@@ -81,11 +81,11 @@ FReply SDockTab::OnDragDetected( const FGeometry& MyGeometry, const FPointerEven
 		FMath::Clamp(TabGrabOffset.X / TabSize.X, 0.0f, 1.0f),
 		FMath::Clamp(TabGrabOffset.Y / TabSize.Y, 0.0f, 1.0f) );
 	
-	auto PinnedParent = ParentPtr.Pin();
-	if (PinnedParent.IsValid())
+	if (TSharedPtr<SDockingTabWell> PinnedParent = ParentPtr.Pin())
 	{
 		//See if we can drag tabs contain in this manager
-		if (GetTabManager()->GetCanDoDragOperation())
+		TSharedPtr<FTabManager> TabManager = GetTabManagerPtr();
+		if (TabManager.IsValid() && TabManager->GetCanDoDragOperation())
 		{
 			return PinnedParent->StartDraggingTab(SharedThis(this), TabGrabOffsetFraction, MouseEvent);
 		}
@@ -466,9 +466,18 @@ TSharedRef<FTabManager> SDockTab::GetTabManager() const
 	return MyTabManager.Pin().ToSharedRef();
 }
 
+TSharedPtr<FTabManager> SDockTab::GetTabManagerPtr() const
+{
+	return MyTabManager.Pin();
+}
+
 void SDockTab::DrawAttention()
 {
-	GetTabManager()->DrawAttention(SharedThis(this));
+	TSharedPtr<FTabManager> TabManager = GetTabManagerPtr();
+	if (TabManager.IsValid())
+	{
+		TabManager->DrawAttention(SharedThis(this));
+	}
 }
 
 void SDockTab::ProvideDefaultLabel( const FText& InDefaultLabel )
diff --git a/Engine/Source/Runtime/Slate/Public/Widgets/Docking/SDockTab.h b/Engine/Source/Runtime/Slate/Public/Widgets/Docking/SDockTab.h
--- a/Engine/Source/Runtime/Slate/Public/Widgets/Docking/SDockTab.h
+++ b/Engine/Source/Runtime/Slate/Public/Widgets/Docking/SDockTab.h
@@ -291,8 +291,12 @@ public:
 	void SetOnExtendContextMenu( const FExtendContextMenu& Handler );
 
 	/** Get the tab manager currently managing this tab. Note that a user move the tab between Tab Managers, so this return value may change. */
+	UE_DEPRECATED(5.0, "The tab manager is not guaranteed to exist, which will cause GetTabManager() to crash. Use GetTabManagerPtr() instead.")
 	TSharedRef<FTabManager> GetTabManager() const;
 
+	/** Get the tab manager currently managing this tab. Note that a user move the tab between Tab Managers, so this return value may change. */
+	TSharedPtr<FTabManager> GetTabManagerPtr() const;
+
 	/** Draws attention to the tab. */
 	void DrawAttention();
 
