Slate: Fix VerifyWidgetLayerId to only use the invalidation test when we painted in fast path. Add the VerifyOutgoingLayerId to the Widget Reflection test list.

diff --git a/Engine/Source/Developer/SlateReflector/Private/Widgets/SSlateOptions.cpp b/Engine/Source/Developer/SlateReflector/Private/Widgets/SSlateOptions.cpp
--- a/Engine/Source/Developer/SlateReflector/Private/Widgets/SSlateOptions.cpp
+++ b/Engine/Source/Developer/SlateReflector/Private/Widgets/SSlateOptions.cpp
@@ -99,6 +99,7 @@ void SSlateOptions::Construct( const FArguments& InArgs )
 #if WITH_SLATE_DEBUGGING
 			AddMenuEntry(MenuBuilder, Icon, LOCTEXT("VerifyParentChildrenRelationship", "Verify Parent/Children Relationship"), TEXT("Slate.VerifyParentChildrenRelationship"));
 			AddMenuEntry(MenuBuilder, Icon, LOCTEXT("VerifyLayerId", "Verify LayerId"), TEXT("Slate.VerifyWidgetLayerId"));
+			AddMenuEntry(MenuBuilder, Icon, LOCTEXT("VerifyOutgoingLayerId", "Verify Outgoing LayerId"), TEXT("Slate.EnsureOutgoingLayerId"));
 #endif // WITH_SLATE_DEBUGGING
 
 			{
diff --git a/Engine/Source/Runtime/Slate/Private/Framework/Application/SlateApplication.cpp b/Engine/Source/Runtime/Slate/Private/Framework/Application/SlateApplication.cpp
--- a/Engine/Source/Runtime/Slate/Private/Framework/Application/SlateApplication.cpp
+++ b/Engine/Source/Runtime/Slate/Private/Framework/Application/SlateApplication.cpp
@@ -7145,7 +7145,14 @@ bool VerifyWidgetLayerId_Recursive(SWidget& Widget, bool bInsideInvalidationRoot
 {
 	bool bResult = true;
 
-	bInsideInvalidationRoot = bInsideInvalidationRoot || Widget.Advanced_IsInvalidationRoot();
+	if (!bInsideInvalidationRoot)
+	{
+		if (const FSlateInvalidationRoot* AsInvalidationRoot = Widget.Advanced_AsInvalidationRoot())
+		{
+			bInsideInvalidationRoot = AsInvalidationRoot->GetLastPaintType() == ESlateInvalidationPaintType::Fast;
+		}
+	}
+
 	const uint32 LastPaintFrame = Widget.Debug_GetLastPaintFrame();
 	const bool bIsDeferredPaint = Widget.GetPersistentState().bDeferredPainting;
 	const int32 InLayerId = Widget.GetPersistentState().LayerId;
