Update SFxWidget to use a non-deprecated version of FGeometry::MakeChild()

diff --git a/Engine/Source/Runtime/Slate/Private/Widgets/Layout/SFxWidget.cpp b/Engine/Source/Runtime/Slate/Private/Widgets/Layout/SFxWidget.cpp
--- a/Engine/Source/Runtime/Slate/Private/Widgets/Layout/SFxWidget.cpp
+++ b/Engine/Source/Runtime/Slate/Private/Widgets/Layout/SFxWidget.cpp
@@ -62,9 +62,9 @@ int32 SFxWidget::OnPaint( const FPaintArgs& Args, const FGeometry& AllottedGeome
 	const FVector2D ScaleOrigin = RenderScaleOrigin.Get() * AllottedGeometry.GetLocalSize();
 	const FVector2D Offset = VisualOffset.Get() * AllottedGeometry.GetLocalSize();
 	// create the render transform as a scale around ScaleOrigin and offset it by Offset.
-	const auto SlateRenderTransform = Concatenate(Inverse(ScaleOrigin), FVector2D(RenderScale.Get()), ScaleOrigin, Offset);
+	const FVector2D SlateRenderTransform = Concatenate<FVector2D, FVector2D, FVector2D, FVector2D>(Inverse(ScaleOrigin), FVector2D(RenderScale.Get()), ScaleOrigin, Offset);
 	// This will append the render transform to the layout transform, and we only use it for rendering.
-	FGeometry ModifiedGeometry = AllottedGeometry.MakeChild(AllottedGeometry.GetLocalSize(), SlateRenderTransform);
+	FGeometry ModifiedGeometry = AllottedGeometry.MakeChild(AllottedGeometry.GetLocalSize(), FSlateLayoutTransform(SlateRenderTransform));
 	
 	FArrangedChildren ArrangedChildren(EVisibility::Visible);
 	this->ArrangeChildren(ModifiedGeometry, ArrangedChildren);
