Fix menu's flashing if longer than height of screen by accounting for application scale.

diff --git a/Engine/Source/Runtime/Slate/Private/Framework/Application/MenuStack.cpp b/Engine/Source/Runtime/Slate/Private/Framework/Application/MenuStack.cpp
--- a/Engine/Source/Runtime/Slate/Private/Framework/Application/MenuStack.cpp
+++ b/Engine/Source/Runtime/Slate/Private/Framework/Application/MenuStack.cpp
@@ -338,14 +338,15 @@ FMenuStack::FPrePushResults FMenuStack::PrePush(const FPrePushArgs& InArgs)
 
 	// Calc the max height available on screen for the menu
 	float MaxHeight;
+	const float ApplicationScale = FSlateApplication::Get().GetApplicationScale() * HostWindow->GetNativeWindow()->GetDPIScaleFactor();
 	if (ActiveMethod.GetPopupMethod() == EPopupMethod::CreateNewWindow)
 	{
 		FSlateRect WorkArea = FSlateApplication::Get().GetWorkArea(InArgs.Anchor);
-		MaxHeight = FMenuStackDefs::MaxMenuScreenHeightFraction * WorkArea.GetSize().Y;
+		MaxHeight = FMenuStackDefs::MaxMenuScreenHeightFraction * WorkArea.GetSize().Y / ApplicationScale;
 	}
 	else
 	{
-		MaxHeight = FMenuStackDefs::MaxMenuScreenHeightFraction * HostWindow->GetClientSizeInScreen().Y;
+		MaxHeight = FMenuStackDefs::MaxMenuScreenHeightFraction * HostWindow->GetClientSizeInScreen().Y / ApplicationScale;
 	}
 
 	bool bAnchorSetsMinWidth = InArgs.TransitionEffect.SlideDirection == FPopupTransitionEffect::ComboButton;
@@ -359,7 +360,6 @@ FMenuStack::FPrePushResults FMenuStack::PrePush(const FPrePushArgs& InArgs)
 
 	OutResults.WrappedContent = WrapContent(TempContent, OptionalMinWidth, OptionalMinHeight);
 
-	const float ApplicationScale = FSlateApplication::Get().GetApplicationScale() * HostWindow->GetNativeWindow()->GetDPIScaleFactor();
 	OutResults.WrappedContent->SlatePrepass(ApplicationScale);
 	// @todo slate: Doesn't take into account potential window border size
 	OutResults.ExpectedSize = OutResults.WrappedContent->GetDesiredSize() * ApplicationScale;
