ListPanel:  update logic for FirstLineScrollOffset to go above 1.0 without assuming all widgets are the same size
- in this case we should add up the full scrollaxis of each widget until the fractional part

diff --git a/Engine/Source/Runtime/Slate/Private/Widgets/Views/SListPanel.cpp b/Engine/Source/Runtime/Slate/Private/Widgets/Views/SListPanel.cpp
--- a/Engine/Source/Runtime/Slate/Private/Widgets/Views/SListPanel.cpp
+++ b/Engine/Source/Runtime/Slate/Private/Widgets/Views/SListPanel.cpp
@@ -88,9 +88,23 @@ void SListPanel::OnArrangeChildren( const FGeometry& AllottedGeometry, FArranged
 		else
 		{
 			// This is a normal list, arrange items in a line along the scroll axis
-			const FTableViewDimensions FirstChildDimensions(Orientation, Children[0].GetWidget()->GetDesiredSize());
-			DimensionsSoFar.ScrollAxis = -FMath::FloorToInt(FirstLineScrollOffset * FirstChildDimensions.ScrollAxis) - OverscrollAmount;
-			
+			float ScrollAxisOffset = 0.0f;
+			const int32 NumWholeWidgetsOffset = FMath::Min(FMath::Floor(FirstLineScrollOffset), Children.Num());
+			for (int32 ItemIndex = 0; ItemIndex < NumWholeWidgetsOffset; ++ItemIndex)
+			{
+				const FTableViewDimensions CurrentChildDimensions(Orientation, Children[ItemIndex].GetWidget()->GetDesiredSize());
+				ScrollAxisOffset += CurrentChildDimensions.ScrollAxis;
+			}
+
+			const int32 FractionalWidgetIndex = NumWholeWidgetsOffset;
+			if (FractionalWidgetIndex < Children.Num())
+			{
+				const FTableViewDimensions OffsetFractionChildDimensions(Orientation, Children[FractionalWidgetIndex].GetWidget()->GetDesiredSize());
+				ScrollAxisOffset += FMath::Frac(FirstLineScrollOffset) * OffsetFractionChildDimensions.ScrollAxis;
+			}
+
+			DimensionsSoFar.ScrollAxis = -FMath::FloorToInt(ScrollAxisOffset) - OverscrollAmount;
+
 			for (int32 ItemIndex = 0; ItemIndex < Children.Num(); ++ItemIndex)
 			{
 				const FTableViewDimensions WidgetDesiredDimensions(Orientation, Children[ItemIndex].GetWidget()->GetDesiredSize());
