Add functionality to STreeView to display the hierarchy of items stacked at the top
Enable said functionality for the Outliner

diff --git a/Engine/Source/Editor/EditorStyle/Private/StarshipStyle.cpp b/Engine/Source/Editor/EditorStyle/Private/StarshipStyle.cpp
--- a/Engine/Source/Editor/EditorStyle/Private/StarshipStyle.cpp
+++ b/Engine/Source/Editor/EditorStyle/Private/StarshipStyle.cpp
@@ -822,6 +822,8 @@ void FStarshipEditorStyle::FStyle::SetupGeneralStyles()
 			.SetActiveHoveredBrush(FSlateNoResource())
 			.SetInactiveHoveredBrush(FSlateNoResource())
 			);
+	
+		Set("ListView.PinnedItemShadow", new IMAGE_BRUSH("Starship/ListView/PinnedItemShadow", FVector2D(16.f, 8.f)));
 	}
 	
 	// Spinboxes
diff --git a/Engine/Source/Editor/SceneOutliner/Private/SOutlinerTreeView.cpp b/Engine/Source/Editor/SceneOutliner/Private/SOutlinerTreeView.cpp
--- a/Engine/Source/Editor/SceneOutliner/Private/SOutlinerTreeView.cpp
+++ b/Engine/Source/Editor/SceneOutliner/Private/SOutlinerTreeView.cpp
@@ -8,6 +8,7 @@
 #include "DragAndDrop/DecoratedDragDropOp.h"
 #include "SceneOutlinerDragDrop.h"
 #include "SSceneOutliner.h"
+#include "Styling/StyleColors.h"
 
 #include "FolderTreeItem.h"
 
@@ -398,4 +399,149 @@ int32 SSceneOutlinerTreeRow::OnPaint( const FPaintArgs& Args, const FGeometry& A
 	return FMath::Max(StartLayer, LayerId + TextLayer);
 }
 
+void SSceneOutlinerPinnedTreeRow::Construct(const FArguments& InArgs, const TSharedRef<SSceneOutlinerTreeView>& OutlinerTreeView, TSharedRef<SSceneOutliner> SceneOutliner)
+{
+	Item = InArgs._Item->AsShared();
+	SceneOutlinerWeak = SceneOutliner;
+	OutlinerTreeViewWeak = OutlinerTreeView;
+
+	FSuperRowType::FArguments Args = FSuperRowType::FArguments()
+		.Style(&FEditorStyle::Get().GetWidgetStyle<FTableRowStyle>("SceneOutliner.TableViewRow"));
+
+	Args.OnDragDetected_Static(HandleOnDragDetected, TWeakPtr<SSceneOutlinerTreeView>(OutlinerTreeView));
+
+	SMultiColumnTableRow<FSceneOutlinerTreeItemPtr>::Construct(Args, OutlinerTreeView);
+
+}
+
+TSharedRef<SWidget> SSceneOutlinerPinnedTreeRow::GenerateWidgetForColumn(const FName& ColumnName)
+{
+	TSharedPtr<ISceneOutlinerTreeItem> ItemPtr = Item.Pin();
+	if (!ItemPtr.IsValid())
+	{
+		return SNullWidget::NullWidget;
+	}
+
+	TSharedPtr<SSceneOutliner> Outliner = SceneOutlinerWeak.Pin();
+	check(Outliner.IsValid());
+
+	// Create the widget for this item
+	TSharedRef<SWidget> NewItemWidget = SNullWidget::NullWidget;
+
+	TSharedPtr<ISceneOutlinerColumn> Column = Outliner->GetColumns().FindRef(ColumnName);
+
+	if (Column.IsValid())
+	{
+		// Construct the actual column widget first
+		TSharedRef<SWidget> ActualWidget = Column->ConstructRowWidget(ItemPtr.ToSharedRef(), *this);
+
+		if (ColumnName == FSceneOutlinerBuiltInColumnTypes::Label())
+		{
+			// We add space for the expander arrow for the label widget to make sure the spacing/indentation is consistent with non-pinned rows
+			NewItemWidget = SNew(SBox)
+				.MinDesiredHeight(FSceneOutlinerDefaultTreeItemMetrics::RowHeight())
+				[
+					SNew(SHorizontalBox)
+					+ SHorizontalBox::Slot()
+					.AutoWidth()
+					.Padding(6, 0, 0, 0)
+					[
+						SNew(SExpanderArrow, SharedThis(this)).IndentAmount(12)
+						.Visibility(EVisibility::Hidden) // Hidden SExpanderArrow to occupy the same space as non-pinned rows
+					]
+					+ SHorizontalBox::Slot()
+					.FillWidth(1.0f)
+					[
+						ActualWidget
+					]
+				];
+		}
+		else
+		{
+			if (ActualWidget != SNullWidget::NullWidget)
+			{
+				// Get the sorted column IDs from the outliner
+				TArray<FName> SortedColumnIDs;
+				Outliner->GetSortedColumnIDs(SortedColumnIDs);
+
+				// Get the current column and label column index to compare
+				int32 ColumnIndex = SortedColumnIDs.Find(ColumnName);
+				int32 LabelColumnIndex = SortedColumnIDs.Find(FSceneOutlinerBuiltInColumnTypes::Label());
+
+				// If either of the columns don't exist, this widget does not need to occupy space
+				if (ColumnIndex == INDEX_NONE || LabelColumnIndex == INDEX_NONE)
+				{
+					ActualWidget->SetVisibility(EVisibility::Collapsed);
+				}
+
+				// If this column is to the LEFT of the label column, it is hidden but occupies space to ensure proper indentation
+				if (ColumnIndex < LabelColumnIndex)
+				{
+					ActualWidget->SetVisibility(EVisibility::Hidden);
+				}
+				// If this column is to the RIGHT of the label colum, it should not occupy space
+				else
+				{
+					ActualWidget->SetVisibility(EVisibility::Collapsed);
+				}
+
+			}
+
+			NewItemWidget = ActualWidget;
+		}
+		
+	}
+
+	return NewItemWidget;
+}
+
+FReply SSceneOutlinerPinnedTreeRow::OnDrop(const FGeometry& MyGeometry, const FDragDropEvent& DragDropEvent)
+{
+	TSharedPtr<ISceneOutlinerTreeItem> ItemPtr = Item.Pin();
+	TSharedPtr<SSceneOutliner> SceneOutlinerPtr = SceneOutlinerWeak.Pin();
+	if (ItemPtr.IsValid() && SceneOutlinerPtr.IsValid())
+	{
+		FSceneOutlinerDragValidationInfo ValidationInfo = FSceneOutlinerDragValidationInfo::Invalid();
+		return HandleDrop(SceneOutlinerPtr, DragDropEvent, *ItemPtr, ValidationInfo, true);
+	}
+
+	return FReply::Unhandled();
+}
+
+void SSceneOutlinerPinnedTreeRow::OnDragEnter(const FGeometry& MyGeometry, const FDragDropEvent& DragDropEvent)
+{
+	TSharedPtr<ISceneOutlinerTreeItem> ItemPtr = Item.Pin();
+	TSharedPtr<SSceneOutliner> SceneOutlinerPtr = SceneOutlinerWeak.Pin();
+	if (ItemPtr.IsValid() && SceneOutlinerPtr.IsValid())
+	{
+		FSceneOutlinerDragValidationInfo ValidationInfo = FSceneOutlinerDragValidationInfo::Invalid();
+
+		FReply Reply = HandleDrop(SceneOutlinerPtr, DragDropEvent, *ItemPtr, ValidationInfo, false);
+		if (Reply.IsEventHandled())
+		{
+			UpdateOperationDecorator(DragDropEvent, ValidationInfo);
+		}
+	}
+}
+
+void SSceneOutlinerPinnedTreeRow::OnDragLeave(const FDragDropEvent& DragDropEvent)
+{
+	ResetOperationDecorator(DragDropEvent);
+}
+
+FReply SSceneOutlinerPinnedTreeRow::OnDragOver(const FGeometry& MyGeometry, const FDragDropEvent& DragDropEvent)
+{
+	TSharedPtr<SSceneOutliner> SceneOutlinerPtr = SceneOutlinerWeak.Pin();
+	if (SSceneOutliner* SceneOutliner = SceneOutlinerPtr.Get())
+	{
+		if (const ISceneOutlinerTreeItem* ItemPtr = Item.Pin().Get())
+		{
+			return SceneOutliner->OnDragOverItem(DragDropEvent, *ItemPtr);
+		}
+		return FReply::Unhandled();
+	}
+
+	return FReply::Handled();
+}
+
 #undef LOCTEXT_NAMESPACE
diff --git a/Engine/Source/Editor/SceneOutliner/Private/SSceneOutliner.cpp b/Engine/Source/Editor/SceneOutliner/Private/SSceneOutliner.cpp
--- a/Engine/Source/Editor/SceneOutliner/Private/SSceneOutliner.cpp
+++ b/Engine/Source/Editor/SceneOutliner/Private/SSceneOutliner.cpp
@@ -88,6 +88,14 @@ void SSceneOutliner::Construct(const FArguments& InArgs, const FSceneOutlinerIni
 	UOutlinerConfig* OutlinerConfig = GetMutableDefault<UOutlinerConfig>();
 	OutlinerConfig->LoadEditorConfig();
 
+	const FSceneOutlinerConfig* SceneOutlinerConfig = GetConstConfig();
+
+	// Load the pinned items visibility from the config file
+	if (SceneOutlinerConfig)
+	{
+		bShouldStackHierarchyHeaders = SceneOutlinerConfig->bShouldStackHierarchyHeaders;
+	}
+
 	// @todo outliner: Should probably save this in layout!
 	// @todo outliner: Should save spacing for list view in layout
 
@@ -232,6 +240,9 @@ void SSceneOutliner::Construct(const FArguments& InArgs, const FSceneOutlinerIni
 			// Generates the actual widget for a tree item
 			.OnGenerateRow( this, &SSceneOutliner::OnGenerateRowForOutlinerTree ) 
 
+			// Generates the actual widget for a pinned tree item
+			.OnGeneratePinnedRow(this, &SSceneOutliner::OnGeneratePinnedRowForOutlinerTree)
+
 			// Use the level viewport context menu as the right click menu for tree items
 			.OnContextMenuOpening(this, &SSceneOutliner::OnOpenContextMenu)
 
@@ -243,6 +254,9 @@ void SSceneOutliner::Construct(const FArguments& InArgs, const FSceneOutlinerIni
 
 			// Make it easier to see hierarchies when there are a lot of items
 			.HighlightParentNodesForSelection(true)
+
+			// Show the Hierarchy of actors pinned at the top of the tree view
+			.ShouldStackHierarchyHeaders(this, &SSceneOutliner::ShouldStackHierarchyHeaders)
 		]
 	];
 
@@ -299,6 +313,30 @@ void SSceneOutliner::HandleHiddenColumnsChanged()
 	}
 }
 
+void SSceneOutliner::GetSortedColumnIDs(TArray<FName>& OutColumnIDs) const
+{
+	FSceneOutlinerModule& SceneOutlinerModule = FModuleManager::LoadModuleChecked<FSceneOutlinerModule>("SceneOutliner");
+
+	TMap<FName, FSceneOutlinerColumnInfo> FilteredColumnMap;
+	for(auto It(SharedData->ColumnMap.CreateIterator()); It; ++It)
+	{
+		if (SceneOutlinerModule.GetColumnPermissionList()->PassesFilter(It.Key()))
+		{
+			FilteredColumnMap.Add(It.Key(), It.Value());
+		}
+	}
+
+	// Get a list of sorted columns IDs to create
+	OutColumnIDs.Empty();
+
+	OutColumnIDs.Reserve(FilteredColumnMap.Num());
+	FilteredColumnMap.GenerateKeyArray(OutColumnIDs);
+
+	OutColumnIDs.Sort([&](const FName& A, const FName& B) {
+		return FilteredColumnMap[A].PriorityIndex < FilteredColumnMap[B].PriorityIndex;
+		});
+}
+
 void SSceneOutliner::SetupColumns(SHeaderRow& HeaderRow)
 {
 	FSceneOutlinerModule& SceneOutlinerModule = FModuleManager::LoadModuleChecked<FSceneOutlinerModule>("SceneOutliner");
@@ -309,7 +347,7 @@ void SSceneOutliner::SetupColumns(SHeaderRow& HeaderRow)
 	}
 
 	TMap<FName, FSceneOutlinerColumnInfo> FilteredColumnMap;
-	for(auto It(SharedData->ColumnMap.CreateIterator()); It; ++It)
+	for (auto It(SharedData->ColumnMap.CreateIterator()); It; ++It)
 	{
 		if (SceneOutlinerModule.GetColumnPermissionList()->PassesFilter(It.Key()))
 		{
@@ -320,14 +358,8 @@ void SSceneOutliner::SetupColumns(SHeaderRow& HeaderRow)
 	Columns.Empty(FilteredColumnMap.Num());
 	HeaderRow.ClearColumns();
 
-	// Get a list of sorted columns IDs to create
 	TArray<FName> SortedIDs;
-	SortedIDs.Reserve(FilteredColumnMap.Num());
-	FilteredColumnMap.GenerateKeyArray(SortedIDs);
-
-	SortedIDs.Sort([&](const FName& A, const FName& B){
-		return FilteredColumnMap[A].PriorityIndex < FilteredColumnMap[B].PriorityIndex;
-	});
+	GetSortedColumnIDs(SortedIDs);
 
 	TMap<FName, bool> ColumnVisibilities;
 	const FSceneOutlinerConfig* OutlinerConfig = GetConstConfig();
@@ -338,7 +370,6 @@ void SSceneOutliner::SetupColumns(SHeaderRow& HeaderRow)
 		ColumnVisibilities = OutlinerConfig->ColumnVisibilities;
 	}
 
-
 	for (const FName& ID : SortedIDs)
 	{
 		bool bIsVisible = true;
@@ -472,6 +503,19 @@ TSharedRef<SWidget> SSceneOutliner::GetViewButtonContent(bool bShowFilters)
 			LOCTEXT("CollapseAllToolTip", "Collapse All Items in the Hierarchy"),
 			FSlateIcon(),
 			FUIAction(FExecuteAction::CreateSP(this, &SSceneOutliner::CollapseAll)));
+
+		MenuBuilder.AddMenuEntry(
+			LOCTEXT("ShowHierarchy", "Stack Hierarchy Headers"),
+			LOCTEXT("ShowHierarchyToolTip", "Toggle pinning of the hierarchy of items at the top of the outliner"),
+			FSlateIcon(),
+			FUIAction(
+				FExecuteAction::CreateRaw(this, &SSceneOutliner::ToggleStackHierarchyHeaders),
+				FCanExecuteAction(),
+				FIsActionChecked::CreateRaw(this, &SSceneOutliner::ShouldStackHierarchyHeaders)
+			),
+			NAME_None,
+			EUserInterfaceActionType::ToggleButton
+		);
 	}
 	MenuBuilder.EndSection();
 
@@ -1693,6 +1737,11 @@ TSharedRef< ITableRow > SSceneOutliner::OnGenerateRowForOutlinerTree( FSceneOutl
 	return SNew( SSceneOutlinerTreeRow, OutlinerTreeView.ToSharedRef(), SharedThis(this) ).Item( Item );
 }
 
+TSharedRef< ITableRow > SSceneOutliner::OnGeneratePinnedRowForOutlinerTree(FSceneOutlinerTreeItemPtr Item, const TSharedRef< STableViewBase >& OwnerTable)
+{
+	return SNew(SSceneOutlinerPinnedTreeRow, OutlinerTreeView.ToSharedRef(), SharedThis(this)).Item(Item);
+}
+
 void SSceneOutliner::OnGetChildrenForOutlinerTree( FSceneOutlinerTreeItemPtr InParent, TArray< FSceneOutlinerTreeItemPtr >& OutChildren )
 {
 	if( SharedData->bShowParentTree )
@@ -2237,6 +2286,25 @@ FSceneOutlinerTreeItemPtr SSceneOutliner::FindParent(const ISceneOutlinerTreeIte
 	return Parent;
 }
 
+void SSceneOutliner::ToggleStackHierarchyHeaders()
+{
+	bShouldStackHierarchyHeaders = !bShouldStackHierarchyHeaders;
+
+	FSceneOutlinerConfig* OutlinerConfig = GetMutableConfig();
+
+	if (OutlinerConfig != nullptr)
+	{
+		OutlinerConfig->bShouldStackHierarchyHeaders = bShouldStackHierarchyHeaders;
+		SaveConfig();
+	}
+
+	FullRefresh();
+}
+
+bool SSceneOutliner::ShouldStackHierarchyHeaders() const
+{
+	return bShouldStackHierarchyHeaders;
+}
 
 struct FSceneOutlinerConfig* SSceneOutliner::GetMutableConfig()
 {
diff --git a/Engine/Source/Editor/SceneOutliner/Private/SceneOutlinerConfig.h b/Engine/Source/Editor/SceneOutliner/Private/SceneOutlinerConfig.h
--- a/Engine/Source/Editor/SceneOutliner/Private/SceneOutlinerConfig.h
+++ b/Engine/Source/Editor/SceneOutliner/Private/SceneOutlinerConfig.h
@@ -18,6 +18,10 @@ public:
 	/** Map to store the visibility of each column */
 	UPROPERTY()
 	TMap<FName, bool> ColumnVisibilities;
+
+	/** Whether the hierarchy is pinned at the top of the outliner*/
+	UPROPERTY()
+	bool bShouldStackHierarchyHeaders = false;
 };
 
 
diff --git a/Engine/Source/Editor/SceneOutliner/Public/SOutlinerTreeView.h b/Engine/Source/Editor/SceneOutliner/Public/SOutlinerTreeView.h
--- a/Engine/Source/Editor/SceneOutliner/Public/SOutlinerTreeView.h
+++ b/Engine/Source/Editor/SceneOutliner/Public/SOutlinerTreeView.h
@@ -113,3 +113,54 @@ private:
 	double LastHighlightInteractionTime;
 };
 
+/** Widget that represents a pinned row in the outliner's tree control.*/
+class SSceneOutlinerPinnedTreeRow
+	: public SMultiColumnTableRow< FSceneOutlinerTreeItemPtr >
+{
+
+public:
+
+	SLATE_BEGIN_ARGS(SSceneOutlinerPinnedTreeRow) {}
+
+		/** The list item for this row */
+		SLATE_ARGUMENT(FSceneOutlinerTreeItemPtr, Item)
+
+	SLATE_END_ARGS()
+
+
+	/** Construct function for this widget */
+	void Construct(const FArguments& InArgs, const TSharedRef<SSceneOutlinerTreeView>& OutlinerTreeView, TSharedRef<SSceneOutliner> SceneOutliner);
+
+	/** Overridden from SMultiColumnTableRow.  Generates a widget for this column of the tree row. */
+	virtual TSharedRef<SWidget> GenerateWidgetForColumn(const FName& ColumnName) override;
+
+protected:
+
+	virtual void ConstructChildren(ETableViewMode::Type InOwnerTableMode, const TAttribute<FMargin>& InPadding, const TSharedRef<SWidget>& InContent) override
+	{
+		STableRow<FSceneOutlinerTreeItemPtr>::Content = InContent;
+
+		// We don't want the expander by default.
+		this->ChildSlot
+			.Padding(InPadding)
+			[
+				InContent
+			];
+	}
+
+	virtual void OnDragEnter(const FGeometry& MyGeometry, const FDragDropEvent& DragDropEvent) override;
+	virtual void OnDragLeave(const FDragDropEvent& DragDropEvent) override;
+	virtual FReply OnDragOver(const FGeometry& MyGeometry, const FDragDropEvent& DragDropEvent) override;
+	virtual FReply OnDrop(const FGeometry& MyGeometry, const FDragDropEvent& DragDropEvent) override;
+
+private:
+
+	/** Weak reference to the outliner widget that owns our list */
+	TWeakPtr< SSceneOutliner > SceneOutlinerWeak;
+
+	/** Weak reference to the tree view that contains our list */
+	TWeakPtr< SSceneOutlinerTreeView > OutlinerTreeViewWeak;
+
+	/** The item associated with this row of data */
+	TWeakPtr<ISceneOutlinerTreeItem> Item;
+};
\ No newline at end of file
diff --git a/Engine/Source/Editor/SceneOutliner/Public/SSceneOutliner.h b/Engine/Source/Editor/SceneOutliner/Public/SSceneOutliner.h
--- a/Engine/Source/Editor/SceneOutliner/Public/SSceneOutliner.h
+++ b/Engine/Source/Editor/SceneOutliner/Public/SSceneOutliner.h
@@ -538,6 +538,8 @@ private:
 		return Columns;
 	}
 
+	void GetSortedColumnIDs(TArray<FName>& OutColumnIDs) const;
+
 	bool PassesFilters(const ISceneOutlinerTreeItem& Item) const
 	{
 		return Filters->PassesAllFilters(Item);
@@ -604,6 +606,9 @@ private:
 	/** Called by STreeView to generate a table row for the specified item */
 	TSharedRef< ITableRow > OnGenerateRowForOutlinerTree( FSceneOutlinerTreeItemPtr Item, const TSharedRef< STableViewBase >& OwnerTable );
 
+	/** Called by STreeView to generate a table row for the specified item is pinned*/
+	TSharedRef< ITableRow > OnGeneratePinnedRowForOutlinerTree(FSceneOutlinerTreeItemPtr Item, const TSharedRef< STableViewBase >& OwnerTable);
+
 	/** Called by STreeView to get child items for the specified parent item */
 	void OnGetChildrenForOutlinerTree( FSceneOutlinerTreeItemPtr InParent, TArray< FSceneOutlinerTreeItemPtr >& OutChildren );
 
@@ -857,9 +862,16 @@ private:
 	/** Currently selected sorting mode */
 	EColumnSortMode::Type SortMode;
 
-	/** Identifier for this outliner (Set through FSceneOutlinerInitializationOptions)*/
+	/** Identifier for this outliner (Set through FSceneOutlinerInitializationOptions) */
 	FName OutlinerIdentifier;
 
+	/** true if the hierarchy of items is pinned at the top of the outliner */
+	bool bShouldStackHierarchyHeaders;
+
+	void ToggleStackHierarchyHeaders();
+
+	bool ShouldStackHierarchyHeaders() const;
+
 	/** Handles column sorting mode change */
 	void OnColumnSortModeChanged( const EColumnSortPriority::Type SortPriority, const FName& ColumnId, const EColumnSortMode::Type InSortMode );
 
diff --git a/Engine/Source/Runtime/Slate/Private/Widgets/Views/STableViewBase.cpp b/Engine/Source/Runtime/Slate/Private/Widgets/Views/STableViewBase.cpp
--- a/Engine/Source/Runtime/Slate/Private/Widgets/Views/STableViewBase.cpp
+++ b/Engine/Source/Runtime/Slate/Private/Widgets/Views/STableViewBase.cpp
@@ -69,6 +69,15 @@ void STableViewBase::ConstructChildren( const TAttribute<float>& InItemWidth, co
 		.ItemAlignment(InItemAlignment)
 		.ListOrientation(Orientation);
 
+	PinnedItemsPanel = SNew(SListPanel)
+		.Clipping(GetClipping())
+		.ItemWidth(InItemWidth)
+		.ItemHeight(InItemHeight)
+		.NumDesiredItems(this, &STableViewBase::GetNumPinnedItems)
+		.ItemAlignment(InItemAlignment)
+		.ListOrientation(Orientation)
+		.Visibility(this, &STableViewBase::GetPinnedItemsVisiblity);
+
 	TSharedPtr<SWidget> ListAndScrollbar;
 	if (InScrollBar)
 	{
@@ -93,7 +102,17 @@ void STableViewBase::ConstructChildren( const TAttribute<float>& InItemWidth, co
 				+SHorizontalBox::Slot()
 				.FillWidth(1)
 				[
-					ItemsPanel.ToSharedRef()
+					SNew(SVerticalBox)
+					+SVerticalBox::Slot()
+					.AutoHeight()
+					[
+						PinnedItemsPanel.ToSharedRef()
+					]
+					+SVerticalBox::Slot()
+					.FillHeight(1)
+					[
+						ItemsPanel.ToSharedRef()
+					]
 				]
 				+SHorizontalBox::Slot()
 				.AutoWidth()
@@ -111,7 +130,17 @@ void STableViewBase::ConstructChildren( const TAttribute<float>& InItemWidth, co
 				+SVerticalBox::Slot()
 				.FillHeight(1)
 				[
-					ItemsPanel.ToSharedRef()
+					SNew(SVerticalBox)
+					+SVerticalBox::Slot()
+					.AutoHeight()
+					[
+						PinnedItemsPanel.ToSharedRef()
+					]
+					+SVerticalBox::Slot()
+					.FillHeight(1)
+					[
+						ItemsPanel.ToSharedRef()
+					]
 				]
 				+SVerticalBox::Slot()
 				.AutoHeight()
@@ -912,6 +941,27 @@ void STableViewBase::ClearWidgets()
 	ItemsPanel->ClearItems();
 }
 
+void STableViewBase::InsertPinnedWidget( const TSharedRef<SWidget> & WidgetToInset )
+{
+	PinnedItemsPanel->AddSlot(0)
+	[
+		WidgetToInset
+	];
+}
+
+void STableViewBase::AppendPinnedWidget( const TSharedRef<SWidget>& WidgetToAppend )
+{
+	PinnedItemsPanel->AddSlot()
+	[
+		WidgetToAppend
+	];
+}
+
+void STableViewBase::ClearPinnedWidgets()
+{
+	PinnedItemsPanel->ClearItems();
+}
+
 float STableViewBase::GetItemWidth() const
 {
 	return GetItemSize().X;
@@ -1058,9 +1108,19 @@ bool STableViewBase::CanUseInertialScroll( float ScrollAmount ) const
 	return CurrentOverscroll == 0.f || FMath::Sign(CurrentOverscroll) != FMath::Sign(ScrollAmount);
 }
 
+int32 STableViewBase::GetNumPinnedItems() const
+{
+	return PinnedItemsPanel->GetChildren()->Num();
+}
+
+EVisibility STableViewBase::GetPinnedItemsVisiblity() const
+{
+	return PinnedItemsPanel->GetChildren()->Num() != 0 ? EVisibility::Visible : EVisibility::Collapsed;
+}
+
 static const TBitArray<> EmptyBitArray = TBitArray<>();
 
 const TBitArray<>& TableViewHelpers::GetEmptyBitArray()
 {
 	return EmptyBitArray;
-}
+}
\ No newline at end of file
diff --git a/Engine/Source/Runtime/Slate/Public/Widgets/Views/SListView.h b/Engine/Source/Runtime/Slate/Public/Widgets/Views/SListView.h
--- a/Engine/Source/Runtime/Slate/Public/Widgets/Views/SListView.h
+++ b/Engine/Source/Runtime/Slate/Public/Widgets/Views/SListView.h
@@ -20,6 +20,8 @@
 #include "Types/SlateConstants.h"
 #include "Widgets/Layout/SScrollBar.h"
 #include "Framework/Layout/Overscroll.h"
+#include "Widgets/Images/SImage.h"
+#include "Widgets/SOverlay.h"
 #if WITH_ACCESSIBILITY
 #include "Application/SlateApplicationBase.h"
 #include "GenericPlatform/Accessibility/GenericAccessibleInterfaces.h"
@@ -81,10 +83,12 @@ public:
 	SLATE_BEGIN_ARGS(SListView<ItemType>)
 		: _ListViewStyle(&FAppStyle::Get().GetWidgetStyle<FTableViewStyle>("ListView"))
 		, _OnGenerateRow()
+		, _OnGeneratePinnedRow()
 		, _OnEntryInitialized()
 		, _OnRowReleased()
 		, _ListItemsSource()
 		, _ItemHeight(16)
+		, _MaxPinnedItems(6)
 		, _OnContextMenuOpening()
 		, _OnMouseButtonClick()
 		, _OnMouseButtonDoubleClick()
@@ -115,6 +119,8 @@ public:
 		SLATE_STYLE_ARGUMENT( FTableViewStyle, ListViewStyle )
 
 		SLATE_EVENT( FOnGenerateRow, OnGenerateRow )
+
+		SLATE_EVENT( FOnGenerateRow, OnGeneratePinnedRow )
 		
 		SLATE_EVENT( FOnEntryInitialized, OnEntryInitialized )
 
@@ -130,6 +136,8 @@ public:
 
 		SLATE_ATTRIBUTE( float, ItemHeight )
 
+		SLATE_ATTRIBUTE(int32, MaxPinnedItems)
+
 		SLATE_EVENT( FOnContextMenuOpening, OnContextMenuOpening )
 
 		SLATE_EVENT(FOnMouseButtonClick, OnMouseButtonClick)
@@ -198,6 +206,7 @@ public:
 		this->Clipping = InArgs._Clipping;
 
 		this->OnGenerateRow = InArgs._OnGenerateRow;
+		this->OnGeneratePinnedRow = InArgs._OnGeneratePinnedRow;
 		this->OnEntryInitialized = InArgs._OnEntryInitialized;
 		this->OnRowReleased = InArgs._OnRowReleased;
 		this->OnItemScrolledIntoView = InArgs._OnItemScrolledIntoView;
@@ -238,6 +247,9 @@ public:
 
 		this->SetStyle(InArgs._ListViewStyle);
 
+		this->MaxPinnedItems = InArgs._MaxPinnedItems;
+		this->DefaultMaxPinnedItems = InArgs._MaxPinnedItems;
+
 		// Check for any parameters that the coder forgot to specify.
 		FString ErrorString;
 		{
@@ -279,6 +291,7 @@ public:
 	SListView(ETableViewMode::Type InListMode = ETableViewMode::List)
 		: STableViewBase(InListMode)
 		, WidgetGenerator(this)
+		, PinnedWidgetGenerator(this)
 		, SelectorItem(TListTypeTraits<ItemType>::MakeNullPtr())
 		, RangeSelectionStart(TListTypeTraits<ItemType>::MakeNullPtr())
 		, ItemsSource(nullptr)
@@ -735,7 +748,7 @@ private:
 
 		void ProcessItemCleanUp()
 		{
-			
+
 			for (int32 ItemIndex = 0; ItemIndex < ItemsToBeCleanedUp.Num(); ++ItemIndex)
 			{
 				ItemType ItemToBeCleanedUp = ItemsToBeCleanedUp[ItemIndex];
@@ -969,7 +982,7 @@ public:
 	virtual const ItemType* Private_ItemFromWidget( const ITableRow* TheWidget ) const override
 	{
 		ItemType const * LookupResult = WidgetGenerator.WidgetMapToItem.Find( TheWidget );
-		return LookupResult == nullptr ? nullptr : LookupResult;
+		return LookupResult == nullptr ? PinnedWidgetGenerator.WidgetMapToItem.Find(TheWidget) : LookupResult;
 	}
 
 	virtual bool Private_UsesSelectorFocus() const override
@@ -1107,6 +1120,147 @@ public:
 		return SharedThis(this);
 	}
 
+private:
+
+	friend class SListViewPinnedRowWidget;
+
+	// Private class that acts as a wrapper around PinnedRows, to allow for customized styling
+	class SListViewPinnedRowWidget : public SCompoundWidget
+	{
+	public:
+		SLATE_BEGIN_ARGS(SListViewPinnedRowWidget) {}
+
+		SLATE_END_ARGS()
+
+		void Construct(const FArguments& InArgs, TSharedPtr<ITableRow> InPinnedItemRow, TSharedRef<SListView> InOwnerListView, const int32 ItemIndex, const int32 NumPinnedItems)
+		{
+			PinnedItemRow = InPinnedItemRow;
+			OwnerListView = InOwnerListView;
+
+			TSharedPtr<STableRow<ItemType>> PinnedRow = StaticCastSharedPtr<STableRow<ItemType>>(InPinnedItemRow);
+
+			// If the PinnedRow inherits from STableRow (i.e has a custom border already), remove it since we will be adding our own border
+			if (PinnedRow.IsValid())
+			{
+				PinnedRow->SetBorderImage(FAppStyle::Get().GetBrush("NoBrush"));
+			}
+
+			TSharedRef<SWidget> InPinnedItemRowWidget = InPinnedItemRow->AsWidget();
+			InPinnedItemRowWidget->SetVisibility(EVisibility::HitTestInvisible);
+
+			ChildSlot
+				[
+					SNew(SOverlay)
+					.Visibility(TAttribute<EVisibility>::CreateSP(this, &SListViewPinnedRowWidget::SetPinnedItemVisibility, ItemIndex, NumPinnedItems))
+
+					+ SOverlay::Slot()
+					.Padding(FMargin(0.0f, 0.0f, 0.0f, 0.0f))
+					[
+						SNew(SBorder)
+						.BorderImage_Lambda([this]()
+							{
+								return this->IsHovered() ? FAppStyle::Get().GetBrush("Brushes.Hover") : FAppStyle::Get().GetBrush("Brushes.Header");
+							})
+						.Padding(0.f)
+						.Content()
+						[
+							SNew(SHorizontalBox)
+							+ SHorizontalBox::Slot()
+							.AutoWidth()
+							[
+								InPinnedItemRowWidget
+							]
+							+ SHorizontalBox::Slot()
+							.FillWidth(1.0f)
+							.Padding(2.0f, 2.0f, 0.0f, 0.0f)
+							[
+								// Text Block for ellipses, shows up when some items in the pinned list are collapsed when the number of items > MaxPinnedItems
+								SNew(STextBlock).Text(NSLOCTEXT("SListView", "Ellipses", "..."))
+								.Visibility(this, &SListViewPinnedRowWidget::SetPinnedItemEllipsesVisibility, ItemIndex)
+							]
+						]
+					]
+					+ SOverlay::Slot()
+					.HAlign(HAlign_Fill)
+					.VAlign(VAlign_Top)
+					[
+						// A shadow to indicate parent/child relationship
+						SNew(SImage)
+						.Visibility(EVisibility::HitTestInvisible)
+						.Image(FAppStyle::Get().GetBrush("ListView.PinnedItemShadow"))
+					]
+					
+				];
+
+		}
+
+	protected:
+
+		virtual FReply OnMouseButtonDown(const FGeometry& MyGeometry, const FPointerEvent& MouseEvent) override
+		{
+			if (MouseEvent.GetEffectingButton() == EKeys::LeftMouseButton)
+			{
+				const ItemType* PinnedItem = OwnerListView->ItemFromWidget(PinnedItemRow.Get());
+
+				if (PinnedItem)
+				{
+					// Navigate to the pinned item on click
+					OwnerListView->RequestNavigateToItem(*PinnedItem);
+					return FReply::Handled();
+				}
+			}
+
+			return FReply::Unhandled();
+		}
+
+	private:
+
+		EVisibility SetPinnedItemVisibility(const int32 IndexInList, const int32 NumPinnedItems) const
+		{
+			// If the hierarchy is not collapsed (i.e all items are visible)
+			if (!OwnerListView->bIsHierarchyCollapsed)
+			{
+				return EVisibility::Visible;
+			}
+
+			int32 CurrentMaxPinnedItems = OwnerListView->MaxPinnedItems.Get();
+
+			// If this is the last item, it is visible
+			if (IndexInList == NumPinnedItems - 1)
+			{
+				return EVisibility::Visible;
+			}
+			// Only show a limited number of items depending on MaxPinnedItems
+			else if (IndexInList < CurrentMaxPinnedItems - 1)
+			{
+				return EVisibility::Visible;
+			}
+
+			return EVisibility::Collapsed;
+		}
+
+		EVisibility SetPinnedItemEllipsesVisibility(const int32 IndexInList) const
+		{
+			// If all items are visible, the ...'s are never visible
+			if (!OwnerListView->bIsHierarchyCollapsed)
+			{
+				return EVisibility::Collapsed;
+			}
+
+			int32 CurrentMaxPinnedItems = OwnerListView->MaxPinnedItems.Get();
+
+			// If the hierarchy is collapsed, the last item before the collapsed items gets the ellipses
+			return IndexInList == CurrentMaxPinnedItems - 2 ? EVisibility::Visible : EVisibility::Collapsed;
+		}
+
+	private:
+		// A pointer to the ListView that owns this widget
+		TSharedPtr<SListView> OwnerListView;
+
+		// A pointer to the row contained by this widget
+ 		TSharedPtr<ITableRow> PinnedItemRow;
+	};
+
 public:	
 
 	void SetStyle(const FTableViewStyle* InStyle)
@@ -1350,12 +1504,112 @@ public:
 		return GeneratedWidgetDimensions.ScrollAxis;
 	}
 
+	void ReGeneratePinnedItems(const TArray<ItemType>& InItems, const FGeometry& MyGeometry, int32 MaxPinnedItemsOverride = -1)
+	{
+		const float LayoutScaleMultiplier = MyGeometry.GetAccumulatedLayoutTransform().GetScale();
+
+		ClearPinnedWidgets();
+
+		// Ensure that we always begin and clean up a generation pass.
+		FGenerationPassGuard GenerationPassGuard(PinnedWidgetGenerator);
+
+		// Check if the User provided an override for MaxPinnedItems, valid until the next time ReGeneratePinnedItems is called
+		if (MaxPinnedItemsOverride != -1)
+		{
+			MaxPinnedItems.Set(MaxPinnedItemsOverride);
+		}
+		else
+		{
+			// Reset it back to the default value if there is no override
+			MaxPinnedItems = DefaultMaxPinnedItems;
+		}
+
+		int32 CurrentMaxPinnedItems = MaxPinnedItems.Get();
+		// There are more items than what we allow to show
+		if (InItems.Num() > CurrentMaxPinnedItems)
+		{
+			bIsHierarchyCollapsed = true;
+		}
+		else
+		{
+			bIsHierarchyCollapsed = false;
+		}
+		
+
+		const TArray<ItemType>& ItemsSourceRef = (*this->ItemsSource);
+
+		for (int32 ItemIndex = 0; ItemIndex < InItems.Num(); ++ItemIndex)
+		{
+			GeneratePinnedWidgetForItem(InItems[ItemIndex], ItemIndex, InItems.Num(), LayoutScaleMultiplier);
+
+			// Deselect any pinned items that were previously selected, since pinned items can only be navigated to on click and not selected
+			if (TListTypeTraits<ItemType>::IsPtrValid(SelectorItem))
+			{
+				if (InItems[ItemIndex] == SelectorItem)
+				{
+					TListTypeTraits<ItemType>::ResetPtr(SelectorItem);
+				}
+			}
+
+		}
+		
+	}
+
+	void GeneratePinnedWidgetForItem(const ItemType& CurItem, int32 ItemIndex, int32 NumPinnedItems, float LayoutScaleMultiplier)
+	{
+		ensure(TListTypeTraits<ItemType>::IsPtrValid(CurItem));
+		// Find a previously generated Widget for this item, if one exists.
+		TSharedPtr<ITableRow> WidgetForItem = PinnedWidgetGenerator.GetWidgetForItem(CurItem);
+		if (!WidgetForItem.IsValid())
+		{
+			// We couldn't find an existing widgets, meaning that this data item was not visible before.
+			// Make a new widget for it.
+			WidgetForItem = this->GenerateNewPinnedWidget(CurItem, ItemIndex, NumPinnedItems);
+		}
+
+		// It is useful to know the item's index that the widget was generated from.
+		// Helps with even/odd coloring
+		WidgetForItem->SetIndexInList(ItemIndex);
+
+		// Let the item generator know that we encountered the current Item and associated Widget.
+		PinnedWidgetGenerator.OnItemSeen(CurItem, WidgetForItem.ToSharedRef());
+
+		// We wrap the row widget around an SListViewPinnedRowWidget for custom styling
+		TSharedRef< SWidget > NewListItemWidget = SNew(SListViewPinnedRowWidget, WidgetForItem, SharedThis(this), ItemIndex, NumPinnedItems);
+		NewListItemWidget->MarkPrepassAsDirty();
+		NewListItemWidget->SlatePrepass(LayoutScaleMultiplier);
+
+		// We have a widget for this item; add it to the panel so that it is part of the UI.
+		this->AppendPinnedWidget(NewListItemWidget);
+	}
+
 	/** @return how many items there are in the TArray being observed */
 	virtual int32 GetNumItemsBeingObserved() const override
 	{
 		return ItemsSource == nullptr ? 0 : ItemsSource->Num();
 	}
 
+	virtual TSharedRef<ITableRow> GenerateNewPinnedWidget(ItemType InItem, const int32 ItemIndex, const int32 NumPinnedItems)
+	{
+		if (OnGeneratePinnedRow.IsBound())
+		{
+			return OnGeneratePinnedRow.Execute(InItem, SharedThis(this));
+		}
+		else
+		{
+			// The programmer did not provide an OnGeneratePinnedRow() handler; let them know.
+			TSharedRef< STableRow<ItemType> > NewListItemWidget =
+				SNew(STableRow<ItemType>, SharedThis(this))
+				.Content()
+				[
+					SNew(STextBlock).Text(NSLOCTEXT("SListView", "BrokenUIMessage", "OnGeneratePinnedRow() not assigned."))
+				];
+
+			return NewListItemWidget;
+		}
+
+	}
+
 	/**
 	 * Given an data item, generate a widget corresponding to it.
 	 *
@@ -1524,6 +1778,7 @@ public:
 	virtual void RebuildList() override
 	{
 		WidgetGenerator.Clear();
+		PinnedWidgetGenerator.Clear();
 		RequestListRefresh();
 	}
 
@@ -1689,7 +1944,9 @@ public:
 	*/
 	virtual TSharedPtr<ITableRow> WidgetFromItem( const ItemType& InItem ) const override
 	{
-		return WidgetGenerator.GetWidgetForItem(InItem);
+		TSharedPtr<ITableRow> ItemWidget = WidgetGenerator.GetWidgetForItem(InItem);
+
+		return ItemWidget != nullptr ? ItemWidget : PinnedWidgetGenerator.GetWidgetForItem(InItem);
 	}
 
 	/**
@@ -2143,12 +2400,18 @@ protected:
 	/** A widget generator component */
 	FWidgetGenerator WidgetGenerator;
 
+	/** A widget generator component used for pinned items in the list */
+	FWidgetGenerator PinnedWidgetGenerator;
+
 	/** Invoked after initializing an entry being generated, before it may be added to the actual widget hierarchy. */
 	FOnEntryInitialized OnEntryInitialized;
 
 	/** Delegate to be invoked when the list needs to generate a new widget from a data item. */
 	FOnGenerateRow OnGenerateRow;
 
+	/** Delegate to be invoked when the list needs to generate a new pinned widget from a data item. */
+	FOnGenerateRow OnGeneratePinnedRow;
+
 	/** Assign this to get more diagnostics from the list view. */
 	FOnItemToString_Debug OnItemToString_Debug;
 
@@ -2227,6 +2490,15 @@ protected:
 	/** Style resource for the list */
 	const FTableViewStyle* Style;
 
+	/** The maximum number of pinned items allowed */
+	TAttribute<int32> MaxPinnedItems;
+
+	/** The initial value of MaxPinnedItems (used to restore it back if overriden) */
+	TAttribute<int32> DefaultMaxPinnedItems;
+	
+	/** If true, number of pinned items > MaxPinnedItems so some items are collapsed in the hierarchy */
+	bool bIsHierarchyCollapsed = false;
+
 private:
 	struct FGenerationPassGuard
 	{
diff --git a/Engine/Source/Runtime/Slate/Public/Widgets/Views/STableViewBase.h b/Engine/Source/Runtime/Slate/Public/Widgets/Views/STableViewBase.h
--- a/Engine/Source/Runtime/Slate/Public/Widgets/Views/STableViewBase.h
+++ b/Engine/Source/Runtime/Slate/Public/Widgets/Views/STableViewBase.h
@@ -252,6 +252,17 @@ protected:
 	 */
 	void ClearWidgets();
 
+	/** Insert WidgetToInsert at the top of the pinned view. */
+	void InsertPinnedWidget(const TSharedRef<SWidget>& WidgetToInset);
+
+	/** Add a WidgetToAppend to the bottom of the pinned view. */
+	void AppendPinnedWidget(const TSharedRef<SWidget>& WidgetToAppend);
+
+	/**
+	 * Remove all the pinned widgets from the view.
+	 */
+	void ClearPinnedWidgets();
+
 	/**
 	 * Get the uniform item width.
 	 */
@@ -335,6 +346,11 @@ protected:
 	/** @return how many items there are in the TArray being observed */
 	virtual int32 GetNumItemsBeingObserved() const = 0;
 
+	/** @return how many pinned items are in the table */
+	int32 GetNumPinnedItems() const;
+
+	EVisibility GetPinnedItemsVisiblity() const;
+
 	enum class EScrollIntoViewResult
 	{
 		/** The function scrolled an item (if set) into view (or the item was already in view) */
@@ -368,6 +384,9 @@ protected:
 	/** The panel which holds the visible widgets in this list */
 	TSharedPtr< SListPanel > ItemsPanel;
 
+	/** The panel which holds the pinned widgets in this list */
+	TSharedPtr< SListPanel > PinnedItemsPanel;
+
 	/** The scroll bar widget */
 	TSharedPtr< SScrollBar > ScrollBar;
 
diff --git a/Engine/Source/Runtime/Slate/Public/Widgets/Views/STreeView.h b/Engine/Source/Runtime/Slate/Public/Widgets/Views/STreeView.h
--- a/Engine/Source/Runtime/Slate/Public/Widgets/Views/STreeView.h
+++ b/Engine/Source/Runtime/Slate/Public/Widgets/Views/STreeView.h
@@ -14,6 +14,8 @@
 #include "Widgets/Text/STextBlock.h"
 #include "Framework/Layout/Overscroll.h"
 #include "Widgets/Views/SListView.h"
+#include "Algo/Reverse.h"
+#include "Math/NumericLimits.h"
 
 /** Info needed by a (relatively) small fraction of the tree items; some of them may not be visible. */
 struct FSparseItemInfo
@@ -114,10 +116,12 @@ public:
 	SLATE_BEGIN_ARGS( STreeView<ItemType> )
 		: _TreeViewStyle(&FAppStyle::Get().GetWidgetStyle<FTableViewStyle>("TreeView"))
 		, _OnGenerateRow()
+		, _OnGeneratePinnedRow()
 		, _OnGetChildren()
 		, _OnSetExpansionRecursive()
 		, _TreeItemsSource( static_cast< const TArray<ItemType>* >(nullptr) ) //@todo Slate Syntax: Initializing from nullptr without a cast
 		, _ItemHeight(16)
+		, _MaxPinnedItems(6) // Having more than the max amount of items leads to the extra items in the middle being collapsed into ellipses, and the last item is fully shown
 		, _OnContextMenuOpening()
 		, _OnMouseButtonClick()
 		, _OnMouseButtonDoubleClick()
@@ -140,6 +144,7 @@ public:
 		, _AllowInvisibleItemSelection(false)
 		, _HighlightParentNodesForSelection(false)
 		, _ReturnFocusToSelection()
+		, _ShouldStackHierarchyHeaders(false)
 		{
 			this->_Clipping = EWidgetClipping::ClipToBounds;
 		}
@@ -148,6 +153,8 @@ public:
 
 		SLATE_EVENT( FOnGenerateRow, OnGenerateRow )
 
+		SLATE_EVENT( FOnGenerateRow, OnGeneratePinnedRow )
+
 		SLATE_EVENT( FOnWidgetToBeRemoved, OnRowReleased )
 
 		SLATE_EVENT( FOnTableViewScrolled, OnTreeViewScrolled )
@@ -158,13 +165,15 @@ public:
 
 		SLATE_EVENT( FOnSetExpansionRecursive, OnSetExpansionRecursive )
 
-		SLATE_ARGUMENT( const TArray<ItemType>* , TreeItemsSource )
+		SLATE_ARGUMENT( const TArray<ItemType>*, TreeItemsSource )
 
 		SLATE_ATTRIBUTE( float, ItemHeight )
 
+		SLATE_ATTRIBUTE( int32, MaxPinnedItems );
+
 		SLATE_EVENT( FOnContextMenuOpening, OnContextMenuOpening )
 
-		SLATE_EVENT(FOnMouseButtonClick, OnMouseButtonClick)
+		SLATE_EVENT( FOnMouseButtonClick, OnMouseButtonClick)
 
 		SLATE_EVENT( FOnMouseButtonDoubleClick, OnMouseButtonDoubleClick )
 
@@ -172,7 +181,7 @@ public:
 
 		SLATE_EVENT( FOnExpansionChanged, OnExpansionChanged )
 
-		SLATE_EVENT(FIsSelectableOrNavigable, OnIsSelectableOrNavigable)
+		SLATE_EVENT( FIsSelectableOrNavigable, OnIsSelectableOrNavigable)
 
 		SLATE_ATTRIBUTE( ESelectionMode::Type, SelectionMode )
 
@@ -211,6 +220,9 @@ public:
 
 		SLATE_ARGUMENT(bool, ReturnFocusToSelection)
 
+		/** If true, Show the current hierarchy of items pinned at the top of the Tree View */
+		SLATE_ATTRIBUTE(bool, ShouldStackHierarchyHeaders)
+
 	SLATE_END_ARGS()
 
 		
@@ -224,6 +236,7 @@ public:
 		this->Clipping = InArgs._Clipping;
 
 		this->OnGenerateRow = InArgs._OnGenerateRow;
+		this->OnGeneratePinnedRow = InArgs._OnGeneratePinnedRow;
 		this->OnRowReleased = InArgs._OnRowReleased;
 		this->OnItemScrolledIntoView = InArgs._OnItemScrolledIntoView;
 		this->OnGetChildren = InArgs._OnGetChildren;
@@ -259,8 +272,13 @@ public:
 
 		this->bReturnFocusToSelection = InArgs._ReturnFocusToSelection;
 
+		this->bShouldStackHierarchyHeaders = InArgs._ShouldStackHierarchyHeaders;
+
 		this->SetStyle(InArgs._TreeViewStyle);
 
+		this->MaxPinnedItems = InArgs._MaxPinnedItems;
+		this->DefaultMaxPinnedItems = InArgs._MaxPinnedItems;
+
 		// Check for any parameters that the coder forgot to specify.
 		FString ErrorString;
 		{
@@ -673,6 +691,107 @@ public:
 
 		return bSawExpandedItems;
 	}
+
+	int32 PopulatePinnedItems(const TArray<ItemType>& InItemsSource, TArray< ItemType >& InPinnedItems, const STableViewBase::FReGenerateResults& Results)
+	{
+		// The value we return, to signify if we want the hierarchy to be collapsed even if it doesn't reach the max amount
+		int32 MaxPinnedItemsOverride = -1;
+
+		if (InItemsSource.IsEmpty())
+		{
+			return MaxPinnedItemsOverride;
+		}
+
+		// Calculate the index of the first item in view
+		int32 StartIndex = FMath::Clamp((int32)(FMath::FloorToDouble(Results.NewScrollOffset)), 0, InItemsSource.Num() - 1);
+		int32 CurrentItemIndex = StartIndex;
+
+		auto GetNonVisibleParents = [this, &InItemsSource, StartIndex](TArray<ItemType>& OutParents, int32 ItemIndex) {
+
+			int32 ParentIndex = ItemIndex;
+
+			// Walk through the list of parents of the current item until you reach the root
+			do
+			{
+				ParentIndex = DenseItemInfos[ItemIndex].ParentIndex;
+
+				// If the current item has a parent, and the parent is not visible, add the parent to the list of pinned items
+				if (ParentIndex != INDEX_NONE && ParentIndex < StartIndex)
+				{
+					OutParents.Add(InItemsSource[ParentIndex]);
+				}
+
+				ItemIndex = ParentIndex;
+
+			} while (ParentIndex != INDEX_NONE);
+		};
+
+		/* Special Case for if we are at the end of the list. When there is no space to scroll down in a list, changing the pinned hierarchy could also change the first visible item
+		 * which is used to calculate the pinned hierarchy. This leads to an infinite loop, so we solve this by finding a first visible item that has a hierachy large enough to hide
+		 * itself, and then collapse the hierarchy until the item remains the first visible item (so there are no infinite loops since the first visible item doesn't change)
+		 *  
+		 */ 
+		if (Results.bGeneratedPastLastItem && Results.NewScrollOffset > 0)
+		{
+			int32 LastItem = InItemsSource.Num() - 1;
+			int32 CurrentMaxPinnedItems = this->MaxPinnedItems.Get();
+
+			// Could be different than reported by STableViewBase if some items are collapsed
+			int32 NumPinnedItems = (this->GetNumPinnedItems() < CurrentMaxPinnedItems) ? this->GetNumPinnedItems() : CurrentMaxPinnedItems;
+
+			// This is the first item that would be visible, if there were no pinned rows
+			int32 FirstItem = Results.NewScrollOffset - NumPinnedItems;
+
+			// We find items that have a hierarchy big enough to cover themselves, but select the smallest among them
+			int32 MinSpaceOccupied = TNumericLimits<int32>::Max();
+
+			// The index of the item we select
+			int32 MinIndex = -1;
+
+			for (int32 ItemIndex = FirstItem; ItemIndex <= LastItem; ItemIndex++)
+			{
+				// Get all parents of the current item that are not visible, to calculate the number of items in its hierarchy
+				TArray<ItemType> NonVisibleParents;
+				GetNonVisibleParents(NonVisibleParents, ItemIndex);
+
+				int32 NumParents = NonVisibleParents.Num();
+
+				// How many items would be required in the hierarchy to cover the item itself
+				int32 IndexOffset = ItemIndex - FirstItem;
+
+				// If the hierarchy is too small, ignore it
+				if (NumParents < IndexOffset)
+				{
+					continue;
+				}
+
+				// If the hierarchy is big enough, AND the smallest we have found so far
+				if (NumParents - IndexOffset < MinSpaceOccupied)
+				{
+					MinSpaceOccupied = NumParents - IndexOffset;
+					MinIndex = ItemIndex;
+				}
+				
+			}
+
+			// If we found no such items, we are in the middle of generating the list so pinned rows are not required
+			if (MinIndex == -1)
+			{
+				return MaxPinnedItemsOverride;
+			}
+
+			CurrentItemIndex = MinIndex;
+			MaxPinnedItemsOverride = MinIndex - FirstItem;
+		}
+
+		// Get all the parents of the item that are not visible, which is the hierarchy to stack
+		GetNonVisibleParents(InPinnedItems, CurrentItemIndex);
+
+		// Reverse the list so the root is at the front
+		Algo::Reverse(InPinnedItems);
+
+		return MaxPinnedItemsOverride;
+	}
 		
 	/**
 	 * Given a TreeItem, create a Widget to represent it in the tree view.
@@ -795,6 +914,27 @@ public:
 		RequestTreeRefresh();
 	}
 
+	virtual STableViewBase::FReGenerateResults ReGenerateItems(const FGeometry& MyGeometry) override
+	{
+		// We need to call the parent function first to know if we reached the end of the list
+		STableViewBase::FReGenerateResults Results = SListView<ItemType>::ReGenerateItems(MyGeometry);
+
+		if (bShouldStackHierarchyHeaders.Get())
+		{
+			TArray<ItemType> PinnedItems;
+
+			// If we reached the end of the list and there is space, a special case requires the hierarchy to be collapsed forcefully
+			int32 MaxPinnedItemsOverride = PopulatePinnedItems(LinearizedItems, PinnedItems, Results);
+			this->ReGeneratePinnedItems(PinnedItems, MyGeometry, MaxPinnedItemsOverride);
+		}
+		else
+		{
+			this->ClearPinnedWidgets();
+		}
+
+		return Results;
+	}
+
 protected:
 	
 	/** The delegate that is invoked whenever we need to gather an item's children. */
@@ -834,4 +974,7 @@ private:
 
 	/** true if we should highlight all parents for each of the currently selected items */
 	bool bHighlightParentNodesForSelection = false;
+
+	/** true if we want to show the hierarchy of items pinned at the top */
+	TAttribute<bool> bShouldStackHierarchyHeaders = false;
 };
