Slate: Added InitialSortMode attribute to FColumn and STableColumnHeader. Defaults to EColumnSortMode::Ascending as before.

diff --git a/Engine/Source/Runtime/Slate/Private/Widgets/Views/SHeaderRow.cpp b/Engine/Source/Runtime/Slate/Private/Widgets/Views/SHeaderRow.cpp
--- a/Engine/Source/Runtime/Slate/Private/Widgets/Views/SHeaderRow.cpp
+++ b/Engine/Source/Runtime/Slate/Private/Widgets/Views/SHeaderRow.cpp
@@ -32,7 +32,8 @@ public:
 	SLATE_END_ARGS()
 
 	STableColumnHeader()
-		: SortMode( EColumnSortMode::None )
+		: InitialSortMode( EColumnSortMode::Ascending )
+		, SortMode( EColumnSortMode::None )
 		, SortPriority( EColumnSortPriority::Primary )
 		, OnSortModeChanged()
 		, ContextMenuContent( SNullWidget::NullWidget )
@@ -53,6 +54,7 @@ public:
 
 		Style = InArgs._Style;
 		ColumnId = Column.ColumnId;
+		InitialSortMode = Column.InitialSortMode;
 		SortMode = Column.SortMode;
 		SortPriority = Column.SortPriority;
 
@@ -216,6 +218,18 @@ public:
 		];
 	}
 
+	/** Gets initial sorting mode */
+	EColumnSortMode::Type GetInitialSortMode() const
+	{
+		return InitialSortMode.Get();
+	}
+
+	/** Sets initial sorting mode */
+	void SetInitialSortMode(EColumnSortMode::Type NewMode)
+	{
+		InitialSortMode = NewMode;
+	}
+
 	/** Gets sorting mode */
 	EColumnSortMode::Type GetSortMode() const
 	{
@@ -378,7 +392,7 @@ private:
 					ColumnSortPriority = EColumnSortPriority::Primary;
 				}
 
-				ColumnSortMode = EColumnSortMode::Ascending;
+				ColumnSortMode = InitialSortMode.Get();
 			}
 			else
 			{
@@ -418,6 +432,9 @@ private:
 
 private:
 
+	/** Initial sorting mode */
+	TAttribute< EColumnSortMode::Type > InitialSortMode;
+
 	/** Current sorting mode */
 	TAttribute< EColumnSortMode::Type > SortMode;
 
diff --git a/Engine/Source/Runtime/Slate/Public/Widgets/Views/SHeaderRow.h b/Engine/Source/Runtime/Slate/Public/Widgets/Views/SHeaderRow.h
--- a/Engine/Source/Runtime/Slate/Public/Widgets/Views/SHeaderRow.h
+++ b/Engine/Source/Runtime/Slate/Public/Widgets/Views/SHeaderRow.h
@@ -116,6 +116,7 @@ public:
 			, _MenuContent()
 			, _HAlignCell( HAlign_Fill )
 			, _VAlignCell( VAlign_Fill )
+			, _InitialSortMode( EColumnSortMode::Ascending )
 			, _SortMode( EColumnSortMode::None )
 			, _OnSort()
 			{}
@@ -146,6 +147,7 @@ public:
 			SLATE_ARGUMENT( EHorizontalAlignment, HAlignCell )
 			SLATE_ARGUMENT( EVerticalAlignment, VAlignCell )
 
+			SLATE_ATTRIBUTE( EColumnSortMode::Type, InitialSortMode )
 			SLATE_ATTRIBUTE( EColumnSortMode::Type, SortMode )
 			SLATE_ATTRIBUTE( EColumnSortPriority::Type, SortPriority )
 			SLATE_EVENT( FOnSortModeChanged, OnSort )
@@ -171,6 +173,7 @@ public:
 			, HeaderComboVisibility (InArgs._HeaderComboVisibility )
 			, CellHAlignment( InArgs._HAlignCell )
 			, CellVAlignment( InArgs._VAlignCell )
+			, InitialSortMode( InArgs._InitialSortMode )
 			, SortMode( InArgs._SortMode )
 			, SortPriority( InArgs._SortPriority )
 			, OnSortModeChanged( InArgs._OnSort )
@@ -255,6 +258,7 @@ public:
 		EHorizontalAlignment CellHAlignment;
 		EVerticalAlignment CellVAlignment;
 
+		TAttribute< EColumnSortMode::Type > InitialSortMode;
 		TAttribute< EColumnSortMode::Type > SortMode;
 		TAttribute< EColumnSortPriority::Type > SortPriority;
 		FOnSortModeChanged OnSortModeChanged;
