Add Toolstack support to tool palettes.

diff --git a/Engine/Source/Editor/Experimental/EditorInteractiveToolsFramework/Private/InteractiveToolStack.cpp b/Engine/Source/Editor/Experimental/EditorInteractiveToolsFramework/Private/InteractiveToolStack.cpp
new file mode 100644
--- /dev/null
+++ b/Engine/Source/Editor/Experimental/EditorInteractiveToolsFramework/Private/InteractiveToolStack.cpp
@@ -0,0 +1,302 @@
+// Copyright Epic Games, Inc. All Rights Reserved.
+
+
+#include "InteractiveToolStack.h"
+#include "InteractiveToolManager.h"
+#include "ContextObjectStore.h"
+#include "UObject/Class.h"
+#include "EdModeInteractiveToolsContext.h"
+#include "ToolContexts/ToolStackContext.h"
+#include "Toolkits/BaseToolkit.h"
+#include "Tools/UEdMode.h"
+#include "Framework/MultiBox/MultiBoxBuilder.h"
+#include "Framework/Commands/Contexts/UIIdentifierContext.h"
+#include "Framework/Commands/Contexts/UIContentContext.h"
+#include "Textures/SlateIcon.h"
+
+#define LOCTEXT_NAMESPACE "FInteractiveToolStack"
+
+void FInteractiveToolStack::RegisterToolStack(UEdMode* EdMode, UEditorInteractiveToolsContext* UseToolsContext, TSharedPtr<FUICommandInfo> UICommand, FString ToolStackIdentifier)
+{
+	// Note: The stack context is only freed when the tool context is, so should be fine to capture ptr
+	UToolStackContext* ToolStackContext = nullptr;
+	if (ensure(UseToolsContext && UseToolsContext->ContextObjectStore))
+	{
+		ToolStackContext = UseToolsContext->ContextObjectStore->FindContext<UToolStackContext>();
+		if (!ensure(ToolStackContext))
+		{
+			return;
+		}
+	}
+
+	auto ExecuteLastActiveStackTool = [ToolStackIdentifier, UseToolsContext, ToolStackContext]()
+	{
+		if (FInteractiveToolStack* ToolStack = ToolStackContext->AccessToolStack(ToolStackIdentifier))
+		{
+			const FString& Tool = ToolStack->GetLastActiveTool();
+			if (!UseToolsContext->IsToolActive(EToolSide::Mouse, Tool))
+			{
+				ToolStack->NotifyToolActivated(Tool);
+				UseToolsContext->StartTool(Tool);
+			}
+		}
+	};
+
+	auto CanExecuteLastActiveStackTool = [EdMode, ToolStackIdentifier, UseToolsContext, ToolStackContext]()
+	{
+		if (FInteractiveToolStack* ToolStack = ToolStackContext->AccessToolStack(ToolStackIdentifier))
+		{
+			const FString& Tool = ToolStack->GetLastActiveTool();
+			return EdMode->ShouldToolStartBeAllowed(Tool) && UseToolsContext->ToolManager->CanActivateTool(EToolSide::Mouse, Tool);
+		}
+
+		return false;
+	};
+
+	auto IsLastStackToolActive = [ToolStackIdentifier, UseToolsContext, ToolStackContext]()
+	{
+		// Note: This is called per frame to check checkbox status
+		if (FInteractiveToolStack* ToolStack = ToolStackContext->AccessToolStack(ToolStackIdentifier))
+		{
+			const FString& Tool = ToolStack->GetLastActiveTool();
+			return UseToolsContext->IsToolActive(EToolSide::Mouse, Tool);
+		}
+
+		return false;
+	};
+
+	auto GetStackMenuContent = [ToolStackIdentifier, UseToolsContext, ToolStackContext](TSharedRef<FUICommandList> InCommandList)
+	{
+		TSharedRef<SWidget> StackMenu = SNullWidget::NullWidget;
+
+		if (FInteractiveToolStack* ToolStack = ToolStackContext->AccessToolStack(ToolStackIdentifier))
+		{
+			const bool bShouldCloseWindowAfterMenuSelection = true;
+			FMenuBuilder MenuBuilder(bShouldCloseWindowAfterMenuSelection, InCommandList);
+			MenuBuilder.BeginSection(FName(ToolStackIdentifier), FText::FromString(ToolStackIdentifier));
+
+			for (const TPair<FString, TSharedPtr<FUICommandInfo>>& Tool : ToolStack->GetStack())
+			{
+				const FString& ToolName = Tool.Key;
+				const TSharedPtr<FUICommandInfo>& ToolCommand = Tool.Value;
+
+				auto ExecuteStackTool = [ToolName, ToolStackIdentifier, UseToolsContext, ToolStackContext]()
+				{
+					if (FInteractiveToolStack* ToolStack = ToolStackContext->AccessToolStack(ToolStackIdentifier))
+					{
+						if (!UseToolsContext->IsToolActive(EToolSide::Mouse, ToolName))
+						{
+							ToolStack->NotifyToolActivated(ToolName);
+							UseToolsContext->StartTool(ToolName);
+						}
+					}
+				};
+
+				FMenuEntryParams MenuEntryParams;
+				MenuEntryParams.LabelOverride = ToolCommand->GetLabel();
+				MenuEntryParams.ToolTipOverride = ToolCommand->GetDescription();
+				MenuEntryParams.IconOverride = ToolCommand->GetIcon();
+				MenuEntryParams.InputBindingOverride = ToolCommand->GetDefaultChord(EMultipleKeyBindingIndex::Primary).GetInputText();
+				MenuEntryParams.DirectActions = { FExecuteAction::CreateWeakLambda(UseToolsContext, ExecuteStackTool) };
+				MenuBuilder.AddMenuEntry(MenuEntryParams);
+			}
+
+			return MenuBuilder.MakeWidget();
+		}
+
+		return StackMenu;
+	};
+
+	auto GetStackIcon = [ToolStackIdentifier, UseToolsContext, UICommand, ToolStackContext]()
+	{
+		if (FInteractiveToolStack* ToolStack = ToolStackContext->AccessToolStack(ToolStackIdentifier))
+		{
+			return ToolStack->GetLastActiveToolIcon();
+		}
+
+		return UICommand->GetIcon();
+	};
+
+	auto GetStackLabel = [ToolStackIdentifier, UseToolsContext, UICommand, ToolStackContext]()
+	{
+		if (FInteractiveToolStack* ToolStack = ToolStackContext->AccessToolStack(ToolStackIdentifier))
+		{
+			return ToolStack->GetLastActiveToolLabel();
+		}
+
+		return UICommand->GetLabel();
+	};
+
+	auto GetStackDescription = [ToolStackIdentifier, UseToolsContext, UICommand, ToolStackContext]()
+	{
+		if (FInteractiveToolStack* ToolStack = ToolStackContext->AccessToolStack(ToolStackIdentifier))
+		{
+			return ToolStack->GetLastActiveToolDescription();
+		}
+
+		return UICommand->GetDescription();
+	};
+	
+	TSharedPtr<FModeToolkit> Toolkit = EdMode->GetToolkit().Pin();
+	if (ensure(Toolkit))
+	{
+		const TSharedRef<FUICommandList>& CommandList = Toolkit->GetToolkitCommands();
+
+		FUIAction StackAction(
+			FExecuteAction::CreateWeakLambda(UseToolsContext, ExecuteLastActiveStackTool),
+			FCanExecuteAction::CreateWeakLambda(UseToolsContext, CanExecuteLastActiveStackTool),
+			FIsActionChecked::CreateWeakLambda(UseToolsContext, IsLastStackToolActive),
+			EUIActionRepeatMode::RepeatDisabled
+		);
+
+		FUIActionContext StackActionContext;
+
+		FUIContentContext StackContent;
+		StackContent.OnGetContent = FOnGetContent::CreateWeakLambda(UseToolsContext, GetStackMenuContent, CommandList);
+		StackActionContext.AddContext(MakeShared<FUIContentContext>(StackContent));
+
+		FUIIdentifierContext StackIdentifier;
+		StackIdentifier.OnGetContextIcon = FOnGetContextIcon::CreateWeakLambda(UseToolsContext, GetStackIcon);
+		StackIdentifier.OnGetContextabel = FOnGetContextText::CreateWeakLambda(UseToolsContext, GetStackLabel);
+		StackIdentifier.OnGetContextDescription = FOnGetContextText::CreateWeakLambda(UseToolsContext, GetStackDescription);
+		StackActionContext.AddContext(MakeShared<FUIIdentifierContext>(StackIdentifier));
+
+		CommandList->MapAction(UICommand, StackAction, StackActionContext);
+	}
+}
+
+void FInteractiveToolStack::AddToolToStack(UEdMode* EdMode, UEditorInteractiveToolsContext* UseToolsContext, TSharedPtr<FUICommandInfo> UICommand, FString ToolStackIdentifier, FString ToolIdentifier, UInteractiveToolBuilder* Builder)
+{
+	// Note: The stack context is only freed when the tool context is, so should be fine to capture ptr
+	UToolStackContext* ToolStackContext = nullptr;
+	if (ensure(UseToolsContext && UseToolsContext->ContextObjectStore))
+	{
+		ToolStackContext = UseToolsContext->ContextObjectStore->FindContext<UToolStackContext>();
+		if (!ensure(ToolStackContext))
+		{
+			return;
+		}
+	}
+
+	if (FInteractiveToolStack* ToolStack = ToolStackContext->AccessToolStack(ToolStackIdentifier))
+	{
+		ToolStack->AddTool(ToolIdentifier, UICommand);
+		UseToolsContext->ToolManager->RegisterToolType(ToolIdentifier, Builder);
+
+		auto ExecuteLastToolForAction = [EdMode, ToolStackIdentifier, UseToolsContext, UICommand, ToolStackContext]()
+		{
+			if (FInteractiveToolStack* ToolStack = ToolStackContext->AccessToolStack(ToolStackIdentifier))
+			{
+				const FInputChord& InputChord = UICommand->GetDefaultChord(EMultipleKeyBindingIndex::Primary);
+				const FString& Tool = ToolStack->GetLastActiveToolByChord(InputChord);
+
+				if (!UseToolsContext->IsToolActive(EToolSide::Mouse, Tool))
+				{
+					ToolStack->NotifyToolActivated(Tool);
+					UseToolsContext->StartTool(Tool);
+				}
+			}
+		};
+
+		auto CanExecuteLastToolForAction = [EdMode, ToolStackIdentifier, UseToolsContext, UICommand, ToolStackContext]()
+		{
+			if (FInteractiveToolStack* ToolStack = ToolStackContext->AccessToolStack(ToolStackIdentifier))
+			{
+				// @TODO: This leaves us in a possible "deadlock" where the last active tool for a chord can't be used,
+				// but other tools for a chord are valid. Seemingly unlikely scenario so address later if needed.
+
+				const FInputChord& InputChord = UICommand->GetDefaultChord(EMultipleKeyBindingIndex::Primary);
+				const FString& Tool = ToolStack->GetLastActiveToolByChord(InputChord);
+				return EdMode->ShouldToolStartBeAllowed(Tool) && UseToolsContext->ToolManager->CanActivateTool(EToolSide::Mouse, Tool);
+			}
+
+			return false;
+		};
+
+		TSharedPtr<FModeToolkit> Toolkit = EdMode->GetToolkit().Pin();
+		if (ensure(Toolkit))
+		{
+			const TSharedRef<FUICommandList>& CommandList = Toolkit->GetToolkitCommands();
+			CommandList->MapAction(UICommand,
+				FExecuteAction::CreateWeakLambda(UseToolsContext, ExecuteLastToolForAction),
+				FCanExecuteAction::CreateWeakLambda(UseToolsContext, CanExecuteLastToolForAction),
+				FIsActionChecked(),
+				EUIActionRepeatMode::RepeatDisabled);
+		}
+	}
+}
+
+int32 FInteractiveToolStack::AddTool(const FString& InToolName, const TSharedPtr<FUICommandInfo>& InUICommand)
+{
+	// Store original stack order for future display
+	OriginalStack.Add(TPair<FString, TSharedPtr<FUICommandInfo>>{InToolName, InUICommand});
+
+	// Insert in reverse order, so later on we can last element for most recently used
+	// Note: Not using queue since we have to arbitrarily remove tool 'n'
+	return SortedToolStack.Insert(TPair<FString, TSharedPtr<FUICommandInfo>>{InToolName, InUICommand}, 0);
+}
+
+void FInteractiveToolStack::NotifyToolActivated(const FString& InToolName)
+{
+	auto FindByToolName = [&](const TPair<FString, TSharedPtr<FUICommandInfo>>& Element)
+	{
+		return Element.Key == InToolName;
+	};
+
+	TPair<FString, TSharedPtr<FUICommandInfo>> ActivatedTool = *SortedToolStack.FindByPredicate(FindByToolName);
+	SortedToolStack.Remove(ActivatedTool);
+	SortedToolStack.Add(ActivatedTool);
+}
+
+const FString FInteractiveToolStack::GetLastActiveToolByChord(const FInputChord& InInputChord) const
+{
+	// Note: Pass by value intentional, this value is volatile don't want to pass ref.
+	
+	// Most recently used tools at end, traverse array in reverse order
+	auto StackIter = SortedToolStack.end();
+	while (StackIter != SortedToolStack.begin())
+	{
+		--StackIter;
+
+		const TPair<FString, TSharedPtr<FUICommandInfo>>& Tool = *StackIter;
+		const FString& ToolName = Tool.Key;
+		const TSharedPtr<FUICommandInfo>& UICommand = Tool.Value;
+
+		// Test raw chord values, not references
+		if (UICommand.IsValid() && UICommand->GetDefaultChord(EMultipleKeyBindingIndex::Primary) == InInputChord)
+		{
+			return ToolName;
+		}
+	}
+
+	ensureMsgf(false, TEXT("Input chord has no corresponding tool: %s"), *InInputChord.GetInputText().ToString());
+	return GetLastActiveTool();
+}
+
+const FString FInteractiveToolStack::GetLastActiveTool() const
+{
+	// Note: Pass by value intentional, this value is volatile don't want to pass ref.
+	return SortedToolStack.Last().Key;
+}
+
+const FSlateIcon& FInteractiveToolStack::GetLastActiveToolIcon() const
+{
+	return SortedToolStack.Last().Value->GetIcon();
+}
+
+const FText& FInteractiveToolStack::GetLastActiveToolLabel() const
+{
+	return SortedToolStack.Last().Value->GetLabel();
+}
+
+const FText& FInteractiveToolStack::GetLastActiveToolDescription() const
+{
+	return SortedToolStack.Last().Value->GetDescription();
+}
+
+const TArray<TPair<FString, TSharedPtr<FUICommandInfo>>>& FInteractiveToolStack::GetStack() const
+{
+	return OriginalStack;
+}
+
+#undef LOCTEXT_NAMESPACE
diff --git a/Engine/Source/Editor/Experimental/EditorInteractiveToolsFramework/Private/ToolContexts/ToolStackContext.cpp b/Engine/Source/Editor/Experimental/EditorInteractiveToolsFramework/Private/ToolContexts/ToolStackContext.cpp
new file mode 100644
--- /dev/null
+++ b/Engine/Source/Editor/Experimental/EditorInteractiveToolsFramework/Private/ToolContexts/ToolStackContext.cpp
@@ -0,0 +1,100 @@
+// Copyright Epic Games, Inc. All Rights Reserved.
+
+#include "ToolContexts/ToolStackContext.h"
+#include "Framework/Commands/UICommandList.h"
+#include "Toolkits/BaseToolkit.h"
+#include "EdModeInteractiveToolsContext.h"
+
+UToolStackContext::~UToolStackContext()
+{
+	if (UEdMode* EdmodePtr = EdMode.Get())
+	{
+		if (TSharedPtr<FModeToolkit> Toolkit = EdmodePtr->GetToolkit().Pin())
+		{
+			const TSharedRef<FUICommandList>& CommandList = Toolkit->GetToolkitCommands();
+			for (auto& RegisteredTool : RegisteredEditorTools)
+			{
+				CommandList->UnmapAction(RegisteredTool.Key);
+				EdmodePtr->GetInteractiveToolsContext(ToolContextScope)->ToolManager->UnregisterToolType(RegisteredTool.Value);
+			}
+		}
+	}
+}
+
+void UToolStackContext::RegisterToolStack(TSharedPtr<FUICommandInfo> UICommand, const FString& ToolStackIdentifier, EToolsContextScope ToolScope)
+{
+	ensureMsgf(UICommand && !UICommand->GetDefaultChord(EMultipleKeyBindingIndex::Primary).IsValidChord(),
+		TEXT("Tool stack command for %s should be invalid, a stack cannot have a meaningful input chord binding"),
+		*ToolStackIdentifier);
+
+	if (UEdMode* EdmodePtr = EdMode.Get())
+	{
+		if (!EdmodePtr->GetToolkit().IsValid())
+		{
+			return;
+		}
+
+		if (ToolScope == EToolsContextScope::Default)
+		{
+			ToolScope = EdmodePtr->GetDefaultToolScope();
+		}
+
+		UEditorInteractiveToolsContext* UseToolsContext = EdmodePtr->GetInteractiveToolsContext(ToolScope);
+		if (ensure(UseToolsContext != nullptr) == false)
+		{
+			return;
+		}
+
+		if (ensure(!ToolStacks.Contains(ToolStackIdentifier)))
+		{
+			ToolStacks.Add(ToolStackIdentifier, {});
+		}
+
+		FInteractiveToolStack::RegisterToolStack(EdmodePtr, UseToolsContext, UICommand, ToolStackIdentifier);
+		if (ToolScope == EToolsContextScope::Editor)
+		{
+			RegisteredEditorToolStacks.Emplace(UICommand, ToolStackIdentifier);
+		}
+	}
+}
+
+void UToolStackContext::AddToolToStack(TSharedPtr<FUICommandInfo> UICommand, const FString& ToolStackIdentifier, const FString& ToolIdentifier, UInteractiveToolBuilder* Builder, EToolsContextScope ToolScope)
+{
+	// Note: We don't enforce the stack context to match the tool context
+	// this is since the tool manager is not aware contexts exist
+
+	if (UEdMode* EdmodePtr = EdMode.Get())
+	{
+		if (!EdmodePtr->GetToolkit().IsValid())
+		{
+			return;
+		}
+
+		if (ToolScope == EToolsContextScope::Default)
+		{
+			ToolScope = EdmodePtr->GetDefaultToolScope();
+		}
+
+		UEditorInteractiveToolsContext* UseToolsContext = EdmodePtr->GetInteractiveToolsContext(ToolScope);
+		if (ensure(UseToolsContext != nullptr) == false)
+		{
+			return;
+		}
+
+		FInteractiveToolStack::AddToolToStack(EdmodePtr, UseToolsContext, UICommand, ToolStackIdentifier, ToolIdentifier, Builder);
+		if (ToolScope == EToolsContextScope::Editor)
+		{
+			RegisteredEditorTools.Emplace(UICommand, ToolStackIdentifier);
+		}
+	}
+}
+
+void UToolStackContext::UnregisterToolStack(const FString& Identifier)
+{
+	ToolStacks.Remove(Identifier);
+}
+
+FInteractiveToolStack* UToolStackContext::AccessToolStack(const FString& Identifier)
+{
+	return ToolStacks.Find(Identifier);
+}
\ No newline at end of file
diff --git a/Engine/Source/Editor/Experimental/EditorInteractiveToolsFramework/Public/InteractiveToolStack.h b/Engine/Source/Editor/Experimental/EditorInteractiveToolsFramework/Public/InteractiveToolStack.h
new file mode 100644
--- /dev/null
+++ b/Engine/Source/Editor/Experimental/EditorInteractiveToolsFramework/Public/InteractiveToolStack.h
@@ -0,0 +1,91 @@
+// Copyright Epic Games, Inc. All Rights Reserved.
+
+#pragma once
+
+#include "CoreMinimal.h"
+#include "InputBehaviorSet.h"
+#include "InteractiveToolActionSet.h"
+#include "ToolContextInterfaces.h"
+#include "UObject/UObjectGlobals.h"
+#include "InteractiveTool.h"
+#include "Framework/Commands/UIAction.h"
+#include "Framework/Commands/UICommandInfo.h"
+
+#include "InteractiveToolStack.generated.h"
+
+class UEdMode;
+class UEditorInteractiveToolsContext;
+class UInteractiveToolBuilder;
+
+/**
+ * FInteractiveToolStack represents a bundle of tools that may or may not share the same input chord
+ * Should some tools share a chord the most recently used tool will be activated by the chord.
+ * If no tools have been activated yet, the chord will go by order added to the stack.
+ * 
+ * To use tool stacks, simply add a UToolStackContext to your relevant interactive tools context.
+ * This can be done by calling 'AddContextObject' to the context object store of your tool context.
+ * See: 'UWidgetEditorToolPaletteMode::Enter' as reference.
+ */
+USTRUCT()
+struct EDITORINTERACTIVETOOLSFRAMEWORK_API FInteractiveToolStack
+{
+	GENERATED_BODY()
+
+public:
+
+	/** 
+	 * Helper method to generate tool stacks
+	 * 
+	 * @param EdMode				EdMode registering this tool stack
+	 * @param UseToolsContext		Tool context for this stack 
+	 * @param UICommand				Command to map tool start / stop actions to
+	 * @param ToolStackIdentifier	Unique string identifier for the tool stack
+	 */
+	static void RegisterToolStack(UEdMode* EdMode, UEditorInteractiveToolsContext* UseToolsContext, TSharedPtr<FUICommandInfo> UICommand, FString ToolStackIdentifier);
+
+	/**
+	 * Helper method to add tools to stacks
+	 *
+	 * @param EdMode				EdMode registering tool to this tool stack
+	 * @param UseToolsContext		Tool context for this stack
+	 * @param UICommand				Command to map tool start / stop actions to
+	 * @param ToolStackIdentifier	Unique string identifier for the tool stack
+	 * @param ToolIdentifier		Unique string identifier for the tool, used to check if tool is active
+	 * @param Builder				Builder for tool to be used by actions
+	 */
+	static void AddToolToStack(UEdMode* EdMode, UEditorInteractiveToolsContext* UseToolsContext, TSharedPtr<FUICommandInfo> UICommand, FString ToolStackIdentifier, FString ToolIdentifier, UInteractiveToolBuilder* Builder);
+
+public:
+
+	/** Add the tool to our internal tool array */
+	int32 AddTool(const FString& InToolName, const TSharedPtr<FUICommandInfo>& InUICommand);
+
+	/** Signal that a tool was activated increasing its priority in the toolstack for future chord inputs */
+	void NotifyToolActivated(const FString& InToolName);
+
+	/** Get Tool that should be activated for the given input chord */
+	const FString GetLastActiveToolByChord(const FInputChord& InInputChord) const;
+
+	/** Get most recently activated tool within this stack */
+	const FString GetLastActiveTool() const;
+
+	/** Get most recently activated tool's icon within this stack */
+	const FSlateIcon& GetLastActiveToolIcon() const;
+
+	/** Get most recently activated tool's label within this stack */
+	const FText& GetLastActiveToolLabel() const;
+
+	/** Get most recently activated tool's description within this stack */
+	const FText& GetLastActiveToolDescription() const;
+
+	/** Get the original tool stack */
+	const TArray<TPair<FString, TSharedPtr<FUICommandInfo>>>& GetStack() const;
+
+protected:
+
+	/** Tools contained within this stack along with their input chord, sorted by most recently used */
+	TArray<TPair<FString, TSharedPtr<FUICommandInfo>>> SortedToolStack;
+
+	/** Original stack ordering */
+	TArray<TPair<FString, TSharedPtr<FUICommandInfo>>> OriginalStack;
+};
\ No newline at end of file
diff --git a/Engine/Source/Editor/Experimental/EditorInteractiveToolsFramework/Public/ToolContexts/ToolStackContext.h b/Engine/Source/Editor/Experimental/EditorInteractiveToolsFramework/Public/ToolContexts/ToolStackContext.h
new file mode 100644
--- /dev/null
+++ b/Engine/Source/Editor/Experimental/EditorInteractiveToolsFramework/Public/ToolContexts/ToolStackContext.h
@@ -0,0 +1,75 @@
+// Copyright Epic Games, Inc. All Rights Reserved.
+
+#pragma once
+
+#include "CoreMinimal.h"
+#include "Tools/UEdMode.h"
+#include "InteractiveToolManager.h"
+#include "InteractiveToolStack.h"
+#include "ToolStackContext.generated.h"
+
+/**
+ * Context needed to support tool stacks, add this to your ContextObjectStore if you want to use tool stacks
+ */
+UCLASS(Transient)
+class EDITORINTERACTIVETOOLSFRAMEWORK_API UToolStackContext : public UObject
+{
+	GENERATED_BODY()
+
+	virtual ~UToolStackContext();
+
+public:
+	/**
+	 * Registers and maps the provided UI command to actions that start / stop multiple tools within a tool stack
+	 * Later on this UI command can be referenced to add the tool stack to a toolbar.
+	 *
+	 * @param	UICommand			Command to map tool stack to, should not have any chords due to child tools having hotkeys
+	 * @param	ToolStackIdentifier	Unique string identifier for the tool stack
+	 * @param	ToolScope			Scope to determine appropriate tool context for stack
+	 */
+	virtual void RegisterToolStack(TSharedPtr<FUICommandInfo> UICommand, const FString& ToolStackIdentifier, EToolsContextScope ToolScope = EToolsContextScope::Default);
+
+	/**
+	 * Registers and maps the provided UI command to actions that control a tool within a tool stack
+	 *
+	 * @param	UICommand			Command containing info such as icon, text, & chord for the tool
+	 * @param	ToolStackIdentifier	Unique string identifier for the tool stack to add this tool to
+	 * @param	ToolIdentifier		Unique string identifier for the tool, used to check if tool is active
+	 * @param	Builder				Builder for tool to be used by actions
+	 * @param	ToolScope			Scope to determine lifetime of tool (Editor, Mode, etc)
+	 */
+	virtual void AddToolToStack(TSharedPtr<FUICommandInfo> UICommand, const FString& ToolStackIdentifier, const FString& ToolIdentifier, UInteractiveToolBuilder* Builder, EToolsContextScope ToolScope = EToolsContextScope::Default);
+
+	/**
+	 * Unregister a Tool stack, unregistration of individual tools is handled seperately
+	 * Note: Currently this is never called since we properly free resources on destruction
+	 * @param Identifier string used to identify this stack
+	 */
+	virtual void UnregisterToolStack(const FString& Identifier);
+
+	/**
+	 * Access an existing tool stack
+	 * @param	Identifier string used to identify this stack
+	 * @return	Ptr to stack if it exists, nullptr otherwise
+	 */
+	virtual FInteractiveToolStack* AccessToolStack(const FString& Identifier);
+
+public:
+	/** EdMode owning the tools referenced by this stack */
+	UPROPERTY()
+	TWeakObjectPtr<UEdMode> EdMode;
+
+	/** Scope for this context object */
+	EToolsContextScope ToolContextScope;
+
+protected:
+	/** Current set of named ToolStacks */
+	UPROPERTY()
+	TMap<FString, FInteractiveToolStack> ToolStacks;
+
+	/** List of Tools this Mode has registered in the EditorToolsContext, to be unregistered on Mode shutdown */
+	TArray<TPair<TSharedPtr<FUICommandInfo>, FString>> RegisteredEditorTools;
+
+	/** List of Tools stacks this Mode has registered in the EditorToolsContext, to be unregistered on Mode shutdown */
+	TArray<TPair<TSharedPtr<FUICommandInfo>, FString>> RegisteredEditorToolStacks;
+};
diff --git a/Engine/Source/Editor/UnrealEd/Public/Tools/UEdMode.h b/Engine/Source/Editor/UnrealEd/Public/Tools/UEdMode.h
--- a/Engine/Source/Editor/UnrealEd/Public/Tools/UEdMode.h
+++ b/Engine/Source/Editor/UnrealEd/Public/Tools/UEdMode.h
@@ -131,6 +131,15 @@ public:
 
 	virtual void Enter();
 
+	/**
+	 * Registers and maps the provided UI command to actions that start / stop a given tool.
+	 * Later on this UI command can be referenced to add the tool to a toolbar.
+	 * 
+	 * @param	UICommand		Command to map tool start / stop actions to
+	 * @param	ToolIdentifier	Unique string identifier for the tool, used to check if tool is active
+	 * @param	Builder			Builder for tool to be used by actions
+	 * @param	ToolScope		Scope to determine lifetime of tool (Editor, Mode, etc)
+	 */
 	virtual void RegisterTool(TSharedPtr<FUICommandInfo> UICommand, FString ToolIdentifier, UInteractiveToolBuilder* Builder, EToolsContextScope ToolScope = EToolsContextScope::Default);
 
 	/** 
@@ -231,7 +240,7 @@ public:
 
 	/**
 	 * Default Scope for InteractiveToolsContext API functions, eg RegisterTool(), GetToolManager(), GetInteractiveToolsContext().
-	 * See EToolsContextScope for details. Defaults to Editor scope.
+	 * See EToolsContextScope for details. Defaults to Ed Mode scope. 
 	 */
 	EToolsContextScope GetDefaultToolScope() const { return EToolsContextScope::EdMode; }
 
diff --git a/Engine/Source/Runtime/Slate/Private/Framework/Commands/UIAction.cpp b/Engine/Source/Runtime/Slate/Private/Framework/Commands/UIAction.cpp
--- a/Engine/Source/Runtime/Slate/Private/Framework/Commands/UIAction.cpp
+++ b/Engine/Source/Runtime/Slate/Private/Framework/Commands/UIAction.cpp
@@ -66,3 +66,24 @@ FUIAction::FUIAction(FExecuteAction InitExecuteAction, FCanExecuteAction InitCan
 	, RepeatMode(InitRepeatMode)
 { 
 }
+
+/////////////////////////////////////////////////////
+// FUIActionContext
+
+TSharedPtr<IUIActionContextBase> FUIActionContext::FindContext(const FName InName) const
+{
+	for (const TSharedPtr<IUIActionContextBase>& Context : Contexts)
+	{
+		if (Context && Context->GetContextName().IsEqual(InName))
+		{
+			return Context;
+		}
+	}
+
+	return nullptr;
+}
+
+void FUIActionContext::AddContext(const TSharedPtr<IUIActionContextBase>& InContext)
+{
+	Contexts.Add(InContext);
+}
diff --git a/Engine/Source/Runtime/Slate/Private/Framework/Commands/UICommandList.cpp b/Engine/Source/Runtime/Slate/Private/Framework/Commands/UICommandList.cpp
--- a/Engine/Source/Runtime/Slate/Private/Framework/Commands/UICommandList.cpp
+++ b/Engine/Source/Runtime/Slate/Private/Framework/Commands/UICommandList.cpp
@@ -62,6 +62,17 @@ void FUICommandList::MapAction( const TSharedPtr< const FUICommandInfo > InUICom
 	UICommandBindingMap.Add( InUICommandInfo, InUIAction );
 }
 
+void FUICommandList::MapAction(const TSharedPtr< const FUICommandInfo > InUICommandInfo, const FUIAction& InUIAction, const FUIActionContext& InUIActionContext)
+{
+	if (!ensure(InUICommandInfo.IsValid()))
+	{
+		return;
+	}
+
+	ContextsInList.Add(InUICommandInfo->GetBindingContext());
+	UICommandBindingMap.Add(InUICommandInfo, InUIAction);
+	UICommandContextMap.Add(InUICommandInfo, InUIActionContext);
+}
 
 void FUICommandList::Append( const TSharedRef<FUICommandList>& InCommandsToAppend )
 {
@@ -79,6 +90,7 @@ void FUICommandList::Append( const TSharedRef<FUICommandList>& InCommandsToAppen
 void FUICommandList::UnmapAction( const TSharedPtr< const FUICommandInfo > InUICommandInfo )
 {
 	UICommandBindingMap.Remove( InUICommandInfo );
+	UICommandContextMap.Remove( InUICommandInfo );
 }
 
 bool FUICommandList::IsActionMapped(const TSharedPtr< const FUICommandInfo > InUICommandInfo) const
@@ -258,6 +270,31 @@ const FUIAction* FUICommandList::GetActionForCommand(TSharedPtr<const FUICommand
 }
 
 
+const FUIActionContext* FUICommandList::GetContextForCommand(TSharedPtr<const FUICommandInfo> Command) const
+{
+	// Make sure the command is valid
+	if (!ensure(Command.IsValid()))
+	{
+		return NULL;
+	}
+
+	// Check in my own binding map. This should not be prevented by CanProduceActionForCommand.
+	// Any action directly requested from a command list should be returned if it actually exists in the list.
+	const FUIActionContext* Context = UICommandContextMap.Find(Command);
+
+	if (!Context)
+	{
+		// We did not find the action in our own list. Recursively attempt to find the command in children and parents.
+		const bool bIncludeChildren = true;
+		const bool bIncludeParents = true;
+		TSet<TSharedRef<const FUICommandList>> VisitedLists;
+		Context = GetContextForCommandRecursively(Command.ToSharedRef(), bIncludeChildren, bIncludeParents, VisitedLists);
+	}
+
+	return Context;
+}
+
+
 const FUIAction* FUICommandList::GetActionForCommandRecursively(const TSharedRef<const FUICommandInfo>& Command, bool bIncludeChildren, bool bIncludeParents, TSet<TSharedRef<const FUICommandList>>& InOutVisitedLists) const
 {
 	// Detect cycles in the graph
@@ -329,6 +366,77 @@ const FUIAction* FUICommandList::GetActionForCommandRecursively(const TSharedRef
 }
 
 
+const FUIActionContext* FUICommandList::GetContextForCommandRecursively(const TSharedRef<const FUICommandInfo>& Command, bool bIncludeChildren, bool bIncludeParents, TSet<TSharedRef<const FUICommandList>>& InOutVisitedLists) const
+{
+	// Detect cycles in the graph
+	{
+		const TSharedRef<const FUICommandList>& ListAsShared = AsShared();
+		if ( InOutVisitedLists.Contains(ListAsShared) )
+		{
+			// This node was already visited. End recursion.
+			return NULL;
+		}
+
+		InOutVisitedLists.Add( ListAsShared );
+	}
+
+	const FUIActionContext* Context = nullptr;
+
+	// Make sure I am capable of processing this command
+	bool bCapableOfCommand = true;
+	if ( CanProduceActionForCommand.IsBound() )
+	{
+		bCapableOfCommand = CanProduceActionForCommand.Execute(Command);
+	}
+
+	if ( bCapableOfCommand )
+	{
+		// Check in my own binding map
+		Context = UICommandContextMap.Find(Command);
+	
+		// If the context was not found, check in my children binding maps
+		if ( !Context && bIncludeChildren )
+		{
+			for ( auto ChildIt = ChildUICommandLists.CreateConstIterator(); ChildIt; ++ChildIt )
+			{
+				TWeakPtr<FUICommandList> Child = *ChildIt;
+				if ( Child.IsValid() )
+				{
+					const bool bShouldIncludeChildrenOfChild = true;
+					const bool bShouldIncludeParentsOfChild = false;
+					Context = Child.Pin()->GetContextForCommandRecursively(Command, bShouldIncludeChildrenOfChild, bShouldIncludeParentsOfChild, InOutVisitedLists);
+					if (Context)
+					{
+						break;
+					}
+				}
+			}
+		}
+	}
+
+	// If the context was not found, check in my parent binding maps
+	if ( !Context && bIncludeParents )
+	{
+		for ( auto ParentIt = ParentUICommandLists.CreateConstIterator(); ParentIt; ++ParentIt )
+		{
+			TWeakPtr<FUICommandList> Parent = *ParentIt;
+			if ( Parent.IsValid() )
+			{
+				const bool bShouldIncludeChildrenOfParent = false;
+				const bool bShouldIncludeParentsOfParent = true;
+				Context = Parent.Pin()->GetContextForCommandRecursively(Command, bShouldIncludeChildrenOfParent, bShouldIncludeParentsOfParent, InOutVisitedLists);
+				if (Context)
+				{
+					break;
+				}
+			}
+		}
+	}
+
+	return Context;
+}
+
+
 void FUICommandList::GatherContextsForList(TSet<FName>& OutAllContexts) const
 {
 	TSet<TSharedRef<const FUICommandList>> VisitedLists;
diff --git a/Engine/Source/Runtime/Slate/Private/Framework/MultiBox/MultiBoxBuilder.cpp b/Engine/Source/Runtime/Slate/Private/Framework/MultiBox/MultiBoxBuilder.cpp
--- a/Engine/Source/Runtime/Slate/Private/Framework/MultiBox/MultiBoxBuilder.cpp
+++ b/Engine/Source/Runtime/Slate/Private/Framework/MultiBox/MultiBoxBuilder.cpp
@@ -12,6 +12,7 @@
 #include "Framework/MultiBox/SToolBarSeparatorBlock.h"
 #include "Framework/MultiBox/SToolBarButtonBlock.h"
 #include "Framework/MultiBox/SToolBarComboButtonBlock.h"
+#include "Framework/MultiBox/SToolBarStackButtonBlock.h"
 #include "Framework/MultiBox/SEditableTextBlock.h"
 #include "Framework/MultiBox/SButtonRowBlock.h"
 #include "Framework/MultiBox/SWidgetBlock.h"
@@ -490,6 +491,24 @@ void FToolBarBuilder::AddComboButton( const FUIAction& InAction, const FOnGetCon
 	MultiBox->AddMultiBlock( NewToolBarComboButtonBlock );
 }
 
+void FToolBarBuilder::AddToolbarStackButton(const TSharedPtr< const FUICommandInfo > InCommand, FName InTutorialHighlightName)
+{
+	ApplySectionBeginning();
+
+	TSharedRef< FToolBarStackButtonBlock > NewToolBarStackButtonBlock(new FToolBarStackButtonBlock(InCommand.ToSharedRef(), CommandListStack.Last()));
+
+	if (LabelVisibility.IsSet())
+	{
+		NewToolBarStackButtonBlock->SetLabelVisibility(LabelVisibility.GetValue());
+	}
+
+	NewToolBarStackButtonBlock->SetForceSmallIcons(bForceSmallIcons);
+	NewToolBarStackButtonBlock->SetTutorialHighlightName(GenerateTutorialIdentfierName(TutorialHighlightName, InTutorialHighlightName, InCommand, MultiBox->GetBlocks().Num()));
+	NewToolBarStackButtonBlock->SetStyleNameOverride(CurrentStyleOverride);
+
+	MultiBox->AddMultiBlock(NewToolBarStackButtonBlock);
+}
+
 void FToolBarBuilder::AddToolBarWidget( TSharedRef<SWidget> InWidget, const TAttribute<FText>& InLabel, FName InTutorialHighlightName, bool bSearchable )
 {
 	ApplySectionBeginning();
diff --git a/Engine/Source/Runtime/Slate/Private/Framework/MultiBox/SToolBarStackButtonBlock.cpp b/Engine/Source/Runtime/Slate/Private/Framework/MultiBox/SToolBarStackButtonBlock.cpp
new file mode 100644
--- /dev/null
+++ b/Engine/Source/Runtime/Slate/Private/Framework/MultiBox/SToolBarStackButtonBlock.cpp
@@ -0,0 +1,597 @@
+// Copyright Epic Games, Inc. All Rights Reserved.
+
+#include "Framework/MultiBox/SToolBarStackButtonBlock.h"
+
+#include "Framework/Application/SlateApplication.h"
+#include "Framework/MultiBox/MultiBoxBuilder.h"
+#include "Widgets/SBoxPanel.h"
+#include "Widgets/SOverlay.h"
+#include "Widgets/Images/SImage.h"
+#include "Widgets/Images/SLayeredImage.h"
+#include "Widgets/Input/SComboButton.h"
+#include "Widgets/Input/SCheckBox.h"
+#include "Widgets/Input/SButton.h"
+#include "Widgets/Layout/SBorder.h"
+#include "Widgets/Layout/SBox.h"
+#include "Styling/ToolBarStyle.h"
+
+#include <type_traits>
+
+/////////////////////////////////////////////////////
+// SInternalStack
+
+/** CheckBox where right & long-press clicks spawn menu. */
+class SCheckBoxStack : public SCheckBox
+{
+
+public: 
+
+	using Super = SCheckBox;
+
+	//~ Begin SWidget interface
+	virtual FReply OnMouseButtonDown(const FGeometry& MyGeometry, const FPointerEvent& MouseEvent) override;
+	virtual FReply OnMouseButtonUp(const FGeometry& MyGeometry, const FPointerEvent& MouseEvent) override;
+	virtual void OnMouseCaptureLost(const FCaptureLostEvent& CaptureLostEvent) override;
+	//~ End SWidget interface
+
+protected:
+
+	/** Starts countdown for spawning of menu via hold press */
+	void BeginHoldPress();
+
+	/** Cancels any existing hold presses */
+	void EndHoldPress();
+
+	/** Callback for after hold press duration has ended */
+	bool OnHandleHoldPress(float DeltaTime);
+
+protected:
+
+	/** Ticker to handle hold actions. */
+	float HoldPressDuration = 0.25f;
+
+	/** Ticker to handle hold actions. */
+	FTSTicker::FDelegateHandle HoldPressHandle;
+};
+
+FReply SCheckBoxStack::OnMouseButtonDown( const FGeometry& MyGeometry, const FPointerEvent& MouseEvent )
+{
+	if ( MouseEvent.GetEffectingButton() == EKeys::LeftMouseButton )
+	{
+		bIsPressed = true;
+
+		EButtonClickMethod::Type InputClickMethod = GetClickMethodFromInputType(MouseEvent);
+
+		BeginHoldPress();
+
+		if(InputClickMethod == EButtonClickMethod::MouseDown )
+		{
+			ToggleCheckedState();
+			const ECheckBoxState State = IsCheckboxChecked.Get();
+			if(State == ECheckBoxState::Checked)
+			{
+				PlayCheckedSound();
+			}
+			else if(State == ECheckBoxState::Unchecked)
+			{
+				PlayUncheckedSound();
+			}
+
+			// Clear our pressed state immediately so that we appear checked
+			bIsPressed = false;
+
+			// Set focus to this button, but don't capture the mouse
+			return FReply::Handled().SetUserFocus(AsShared(), EFocusCause::Mouse);
+		}
+		else
+		{
+			// Capture the mouse, and also set focus to this button
+			return FReply::Handled().CaptureMouse(AsShared()).SetUserFocus(AsShared(), EFocusCause::Mouse);
+		}
+	}
+	else if ( MouseEvent.GetEffectingButton() == EKeys::RightMouseButton && OnGetMenuContent.IsBound() )
+	{
+		EndHoldPress(); 
+		OnHandleHoldPress(0.0f);
+
+		return FReply::Handled();
+	}
+	else
+	{
+		return FReply::Unhandled();
+	}
+}
+
+FReply SCheckBoxStack::OnMouseButtonUp( const FGeometry& MyGeometry, const FPointerEvent& MouseEvent )
+{
+	EndHoldPress();
+
+	return Super::OnMouseButtonUp(MyGeometry, MouseEvent);
+}
+
+void SCheckBoxStack::OnMouseCaptureLost(const FCaptureLostEvent& CaptureLostEvent)
+{
+	EndHoldPress();
+
+	Super::OnMouseCaptureLost(CaptureLostEvent);
+}
+
+void SCheckBoxStack::BeginHoldPress()
+{
+	HoldPressHandle = FTSTicker::GetCoreTicker().AddTicker(FTickerDelegate::CreateSP(this, &SCheckBoxStack::OnHandleHoldPress), HoldPressDuration);
+}
+
+void SCheckBoxStack::EndHoldPress()
+{
+	if (HoldPressHandle.IsValid())
+	{
+		FTSTicker::GetCoreTicker().RemoveTicker(HoldPressHandle);
+		HoldPressHandle.Reset();
+	}
+}
+
+bool SCheckBoxStack::OnHandleHoldPress(float DeltaTime)
+{
+	// For now, always spawns menu top right
+	FSlateApplication::Get().PushMenu(
+		AsShared(),
+		FWidgetPath(),
+		OnGetMenuContent.Execute(),
+		GetTickSpaceGeometry().GetRenderBoundingRect().GetTopRight(),
+		FPopupTransitionEffect(FPopupTransitionEffect::ContextMenu)
+	);
+
+	// Do not repeat ticker
+	return false;
+}
+
+/////////////////////////////////////////////////////
+// FToolBarStackButtonBlock
+
+FToolBarStackButtonBlock::FToolBarStackButtonBlock(const TSharedRef< const FUICommandInfo > InCommand, TSharedPtr< const FUICommandList > InCommandList, bool bInSimpleStackBox)
+	: FMultiBlock(InCommand, InCommandList)
+	, bSimpleComboBox(bInSimpleStackBox)
+{
+}
+
+void FToolBarStackButtonBlock::CreateMenuEntry(class FMenuBuilder& MenuBuilder) const
+{
+	MenuBuilder.AddWrapperSubMenu(FText::GetEmpty(), FText::GetEmpty(), GetStackContent()->OnGetContent, FSlateIcon(), GetDirectActions());
+}
+
+TSharedRef< class IMultiBlockBaseWidget > FToolBarStackButtonBlock::ConstructWidget() const
+{
+	return SNew( SToolBarStackButtonBlock )
+		.LabelVisibility( LabelVisibility.IsSet() ? LabelVisibility.GetValue() : TOptional< EVisibility >() )
+		.ForceSmallIcons( bForceSmallIcons )
+		.Cursor( EMouseCursor::Default );
+}
+
+const TSharedPtr<FUIIdentifierContext> FToolBarStackButtonBlock::GetStackIdentifier() const
+{
+	if (CachedStackIdentifier.IsValid())
+	{
+		return CachedStackIdentifier.Pin();
+	}
+
+	if (const FUIActionContext* ActionContext = GetActionList()->GetContextForCommand(GetAction()))
+	{
+		if (TSharedPtr<FUIIdentifierContext> IdentifierContext = StaticCastSharedPtr<FUIIdentifierContext>(ActionContext->FindContext(FUIIdentifierContext::ContextName)))
+		{
+			CachedStackIdentifier = IdentifierContext;
+			return IdentifierContext;
+		}
+	}
+
+	return nullptr;
+}
+
+const TSharedPtr<FUIContentContext> FToolBarStackButtonBlock::GetStackContent() const
+{
+	if (CachedStackContent.IsValid())
+	{
+		return CachedStackContent.Pin();
+	}
+
+	if (const FUIActionContext* ActionContext = GetActionList()->GetContextForCommand(GetAction()))
+	{
+		if (TSharedPtr<FUIContentContext> IdentifierContext = StaticCastSharedPtr<FUIContentContext>(ActionContext->FindContext(FUIContentContext::ContextName)))
+		{
+			CachedStackContent = IdentifierContext;
+			return IdentifierContext;
+		}
+	}
+
+	return nullptr;
+}
+
+/////////////////////////////////////////////////////
+// SToolBarStackButtonBlock
+
+void SToolBarStackButtonBlock::Construct( const FArguments& InArgs )
+{
+	LabelVisibilityOverride = InArgs._LabelVisibility;
+	bForceSmallIcons = InArgs._ForceSmallIcons;
+}
+
+void SToolBarStackButtonBlock::BuildMultiBlockWidget(const ISlateStyle* StyleSet, const FName& StyleName)
+{
+	TSharedRef< const FMultiBox > MultiBox( OwnerMultiBoxWidget.Pin()->GetMultiBox() );
+	
+	TSharedRef< const FToolBarStackButtonBlock > ToolBarStackButtonBlock = StaticCastSharedRef< const FToolBarStackButtonBlock >( MultiBlock.ToSharedRef() );
+
+	TAttribute<FText> Label;
+
+	const FToolBarStyle& ToolBarStyle = StyleSet->GetWidgetStyle<FToolBarStyle>(StyleName);
+
+	// If override is set use that
+	if (LabelVisibilityOverride.IsSet())
+	{
+		LabelVisibility = LabelVisibilityOverride.GetValue();
+	}
+	else if (!ToolBarStyle.bShowLabels)
+	{
+		// Otherwise check the style
+		LabelVisibility = EVisibility::Collapsed;
+	}
+	else
+	{
+		LabelVisibility = TAttribute< EVisibility >::Create(TAttribute< EVisibility >::FGetter::CreateSP(SharedThis(this), &SToolBarStackButtonBlock::GetIconVisibility, false));
+	}
+
+	TSharedRef<SLayeredImage> IconWidget =
+		SNew(SLayeredImage)
+		.ColorAndOpacity(this, &SToolBarStackButtonBlock::GetIconForegroundColor)
+		.Visibility(EVisibility::HitTestInvisible)
+		.Image(this, &SToolBarStackButtonBlock::GetIconBrush);
+
+	IconWidget->AddLayer(TAttribute<const FSlateBrush*>(this, &SToolBarStackButtonBlock::GetNormalIconBrush));
+
+	// Add this widget to the search list of the multibox
+	OwnerMultiBoxWidget.Pin()->AddElement(this->AsWidget(), Label.Get(), MultiBlock->GetSearchable());
+
+	// Setup the string for the metatag
+	FName TagName;
+	if (ToolBarStackButtonBlock->GetTutorialHighlightName().IsNone())
+	{
+		TagName = *FString::Printf(TEXT("Stack Button,%s,0"), *Label.Get().ToString());
+	}
+	else
+	{
+		TagName = ToolBarStackButtonBlock->GetTutorialHighlightName();
+	}
+
+	// Create the content for our button
+	EMultiBlockLocation::Type BlockLocation = GetMultiBlockLocation();
+	FName BlockStyle = EMultiBlockLocation::ToName(ISlateStyle::Join(StyleName, ".Button"), BlockLocation);
+	const FComboButtonStyle* ComboButtonStyle = &FAppStyle::Get().GetWidgetStyle<FComboButtonStyle>("StatusBar.StatusBarComboButton");
+	const FButtonStyle* ButtonStyle = BlockLocation == EMultiBlockLocation::None ? &ToolBarStyle.ButtonStyle : &StyleSet->GetWidgetStyle<FButtonStyle>(BlockStyle);
+	const FCheckBoxStyle* CheckStyle = BlockLocation == EMultiBlockLocation::None ? &ToolBarStyle.ToggleButton : &StyleSet->GetWidgetStyle<FCheckBoxStyle>(BlockStyle);
+
+	if (ToolBarStackButtonBlock->bSimpleComboBox)
+	{
+		ComboButtonStyle = &ToolBarStyle.SettingsComboButton;
+		ButtonStyle = &ComboButtonStyle->ButtonStyle;
+	}
+
+	const bool bHasDownArrowShadow = !ComboButtonStyle->ShadowOffset.IsZero();
+
+	FMargin ExtraDropDownPadding = FMargin(0, 0, 4, 0);
+	FMargin ExtraLabelPadding = FMargin(4, 0, 0, 0);
+
+	TSharedRef<SWidget> ButtonContent = SNullWidget::NullWidget;
+	if (MultiBox->GetType() == EMultiBoxType::SlimHorizontalToolBar)
+	{
+		const FVector2D IconSize = ToolBarStyle.IconSize;
+
+		IconWidget->SetDesiredSizeOverride(IconSize);
+		ButtonContent =
+			SNew(SHorizontalBox)
+			.AddMetaData<FTagMetaData>(FTagMetaData(TagName))
+			+ SHorizontalBox::Slot()
+			.AutoWidth()
+			.VAlign(VAlign_Center)
+			.HAlign(HAlign_Center)
+			.Padding(ExtraLabelPadding)
+			[
+				IconWidget
+			]
+			// Label text
+			+ SHorizontalBox::Slot()
+			.AutoWidth()
+			.Padding(ToolBarStyle.LabelPadding + ExtraLabelPadding)
+			.VAlign(VAlign_Center)
+			[
+				SNew(STextBlock)
+				.Visibility(LabelVisibility)
+				.Text(this, &SToolBarStackButtonBlock::GetLabel)
+				.TextStyle(&ToolBarStyle.LabelStyle)	// Smaller font for tool tip labels
+			]
+
+			// Drop down arrow
+			+SHorizontalBox::Slot()
+			.AutoWidth()
+			.HAlign( HAlign_Right )
+			.VAlign(ComboButtonStyle->DownArrowAlign)
+			.Padding(ComboButtonStyle->DownArrowPadding)
+			[
+				SNew(SOverlay)
+				// drop shadow
+				+ SOverlay::Slot()
+				.VAlign(VAlign_Top)
+				.Padding(FMargin(ComboButtonStyle->ShadowOffset.X, ComboButtonStyle->ShadowOffset.Y, 0, 0))
+				[
+					SNew(SImage)
+					.Visibility( bHasDownArrowShadow ? EVisibility::Visible : EVisibility::Collapsed )
+					.Image( &ComboButtonStyle->DownArrowImage )
+					.ColorAndOpacity( ComboButtonStyle->ShadowColorAndOpacity )
+				]
+				+ SOverlay::Slot()
+				.VAlign(VAlign_Top)
+				[
+					SNew(SImage)
+					.Visibility( EVisibility::Visible  )
+					.Image( &ComboButtonStyle->DownArrowImage )
+					// Inherit tinting from parent
+					.ColorAndOpacity( FSlateColor::UseForeground() )
+				]
+			];
+	}
+	else
+	{
+		ButtonContent =
+			SNew(SHorizontalBox)
+			.AddMetaData<FTagMetaData>(FTagMetaData(TagName))
+			+ SHorizontalBox::Slot()
+			.FillWidth(1)
+			.VAlign(VAlign_Center)
+			.Padding(ExtraLabelPadding)
+			[
+				SNew(SVerticalBox)
+				// Icon image
+				+ SVerticalBox::Slot()
+				.AutoHeight()
+				.HAlign(HAlign_Center)	// Center the icon horizontally, so that large labels don't stretch out the artwork
+				[
+					IconWidget
+				]
+				// Label text
+				+ SVerticalBox::Slot().AutoHeight()
+				.Padding(ToolBarStyle.LabelPadding)
+				.HAlign(HAlign_Center)	// Center the label text horizontally
+				[
+					SNew(STextBlock)
+					.Visibility(LabelVisibility)
+					.Text(this, &SToolBarStackButtonBlock::GetLabel)
+					.TextStyle(&ToolBarStyle.LabelStyle)	// Smaller font for tool tip labels
+				]
+			]
+		
+			// Drop down arrow
+			+SHorizontalBox::Slot()
+			.AutoWidth()
+			.HAlign( HAlign_Right )
+			.VAlign(ComboButtonStyle->DownArrowAlign)
+			.Padding(ComboButtonStyle->DownArrowPadding)
+			[
+				SNew(SOverlay)
+				// drop shadow
+				+ SOverlay::Slot()
+				.VAlign(VAlign_Top)
+				.Padding(FMargin(ComboButtonStyle->ShadowOffset.X, ComboButtonStyle->ShadowOffset.Y, 0, 0))
+				[
+					SNew(SImage)
+					.Visibility( bHasDownArrowShadow ? EVisibility::Visible : EVisibility::Collapsed )
+					.Image( &ComboButtonStyle->DownArrowImage )
+					.ColorAndOpacity( ComboButtonStyle->ShadowColorAndOpacity )
+				]
+				+ SOverlay::Slot()
+				.VAlign(VAlign_Top)
+				.Padding(ExtraDropDownPadding)
+				[
+					SNew(SImage)
+					.Visibility( EVisibility::Visible  )
+					.Image( &ComboButtonStyle->DownArrowImage )
+					// Inherit tinting from parent
+					.ColorAndOpacity( FSlateColor::UseForeground() )
+				]
+			];
+	}
+
+	OpenForegroundColor = ButtonStyle->HoveredForeground;
+
+	ChildSlot
+	[
+		SAssignNew(StackButtonWidget, SCheckBoxStack)
+		.Style(CheckStyle)
+		.ClickMethod(EButtonClickMethod::MouseDown)
+		.ToolTipText(this, &SToolBarStackButtonBlock::GetDescription)
+		.OnCheckStateChanged(this, &SToolBarStackButtonBlock::OnCheckStateChanged)
+		.IsChecked(this, &SToolBarStackButtonBlock::GetCheckState)
+		.IsEnabled(this, &SToolBarStackButtonBlock::IsEnabled)
+		.OnGetMenuContent(this, &SToolBarStackButtonBlock::OnGetMenuContent)
+		[
+			ButtonContent
+		]
+	];
+
+
+	FMargin Padding = ToolBarStyle.ComboButtonPadding;
+	if (ToolBarStackButtonBlock->bSimpleComboBox)
+	{
+		Padding = FMargin(0);
+	}
+
+	ChildSlot.Padding(Padding);
+
+	// Bind our widget's visible state to whether or not the button should be visible
+	SetVisibility( TAttribute<EVisibility>(this, &SToolBarStackButtonBlock::GetVisibility) );
+}
+
+void SToolBarStackButtonBlock::OnCheckStateChanged( const ECheckBoxState NewCheckedState )
+{
+	TSharedPtr< const FUICommandList > ActionList = MultiBlock->GetActionList();
+	TSharedPtr< const FUICommandInfo > Action = MultiBlock->GetAction();
+	const FUIAction& DirectActions = MultiBlock->GetDirectActions();
+
+	if (ActionList.IsValid() && Action.IsValid())
+	{
+		ActionList->ExecuteAction(Action.ToSharedRef());
+	}
+	else
+	{
+		// There is no action list or action associated with this block via a UI command.  Execute any direct action we have
+		MultiBlock->GetDirectActions().Execute();
+	}
+}
+
+ECheckBoxState SToolBarStackButtonBlock::GetCheckState() const
+{
+	TSharedPtr<const FUICommandList> ActionList = MultiBlock->GetActionList();
+	TSharedPtr<const FUICommandInfo> Action = MultiBlock->GetAction();
+	const FUIAction& DirectActions = MultiBlock->GetDirectActions();
+
+	ECheckBoxState CheckState = ECheckBoxState::Unchecked;
+	if( ActionList.IsValid() && Action.IsValid() )
+	{
+		CheckState = ActionList->GetCheckState( Action.ToSharedRef() );
+	}
+	else
+	{
+		// There is no action list or action associated with this block via a UI command.  Execute any direct action we have
+		CheckState = DirectActions.GetCheckState();
+	}
+
+	return CheckState;
+}
+
+TSharedRef<SWidget> SToolBarStackButtonBlock::OnGetMenuContent()
+{
+	TSharedRef< const FToolBarStackButtonBlock > ToolBarButtonComboBlock = StaticCastSharedRef< const FToolBarStackButtonBlock >( MultiBlock.ToSharedRef() );
+	
+	if (const TSharedPtr<FUIContentContext> StackContent = ToolBarButtonComboBlock->GetStackContent())
+	{
+		return StackContent->OnGetContent.Execute();
+	}
+
+	return SNullWidget::NullWidget;
+}
+
+bool SToolBarStackButtonBlock::IsEnabled() const
+{
+	const FUIAction& UIAction = MultiBlock->GetDirectActions();
+	if( UIAction.CanExecuteAction.IsBound() )
+	{
+		return UIAction.CanExecuteAction.Execute();
+	}
+
+	return true;
+}
+
+EVisibility SToolBarStackButtonBlock::GetVisibility() const
+{
+	const FUIAction& UIAction = MultiBlock->GetDirectActions();
+	if (UIAction.IsActionVisibleDelegate.IsBound())
+	{
+		return UIAction.IsActionVisibleDelegate.Execute() ? EVisibility::Visible : EVisibility::Collapsed;
+	}
+
+	return EVisibility::Visible;
+}
+
+const FSlateBrush* SToolBarStackButtonBlock::GetIconBrush() const
+{
+	return bForceSmallIcons || FMultiBoxSettings::UseSmallToolBarIcons.Get() ? GetSmallIconBrush() : GetNormalIconBrush();
+}
+
+FText SToolBarStackButtonBlock::GetLabel() const
+{
+	TSharedRef< const FToolBarStackButtonBlock > ToolBarStackButtonBlock = StaticCastSharedRef< const FToolBarStackButtonBlock >(MultiBlock.ToSharedRef());
+
+	return ToolBarStackButtonBlock->GetStackIdentifier()
+		? ToolBarStackButtonBlock->GetStackIdentifier()->OnGetContextabel.Execute()
+		: ToolBarStackButtonBlock->GetAction()->GetLabel();
+}
+
+FText SToolBarStackButtonBlock::GetDescription() const
+{
+	TSharedRef< const FToolBarStackButtonBlock > ToolBarStackButtonBlock = StaticCastSharedRef< const FToolBarStackButtonBlock >(MultiBlock.ToSharedRef());
+
+	return ToolBarStackButtonBlock->GetStackIdentifier()
+		? ToolBarStackButtonBlock->GetStackIdentifier()->OnGetContextDescription.Execute()
+		: ToolBarStackButtonBlock->GetAction()->GetDescription();
+}
+
+const FSlateBrush* SToolBarStackButtonBlock::GetNormalIconBrush() const
+{
+	TSharedRef< const FToolBarStackButtonBlock > ToolBarStackButtonBlock = StaticCastSharedRef< const FToolBarStackButtonBlock >(MultiBlock.ToSharedRef());
+	
+	const FSlateIcon ActionIcon = ToolBarStackButtonBlock->GetStackIdentifier() ? ToolBarStackButtonBlock->GetStackIdentifier()->OnGetContextIcon.Execute() : FSlateIcon();
+	
+	if( ActionIcon.IsSet() )
+	{
+		return ActionIcon.GetSmallIcon();
+	}
+	else
+	{
+		check( OwnerMultiBoxWidget.IsValid() );
+
+		TSharedPtr<SMultiBoxWidget> MultiBoxWidget = OwnerMultiBoxWidget.Pin();
+		const ISlateStyle* const StyleSet = MultiBoxWidget->GetStyleSet();
+
+		static const FName IconName("MultiBox.GenericToolBarIcon.Small" );
+		return StyleSet->GetBrush(IconName);
+	}
+}
+
+const FSlateBrush* SToolBarStackButtonBlock::GetSmallIconBrush() const
+{
+	TSharedRef< const FToolBarStackButtonBlock > ToolBarStackButtonBlock = StaticCastSharedRef< const FToolBarStackButtonBlock >(MultiBlock.ToSharedRef());
+
+	const FSlateIcon ActionIcon = ToolBarStackButtonBlock->GetStackIdentifier() ? ToolBarStackButtonBlock->GetStackIdentifier()->OnGetContextIcon.Execute() : FSlateIcon();
+	
+	if( ActionIcon.IsSet() )
+	{
+		return ActionIcon.GetSmallIcon();
+	}
+	else
+	{
+		check( OwnerMultiBoxWidget.IsValid() );
+
+		TSharedPtr<SMultiBoxWidget> MultiBoxWidget = OwnerMultiBoxWidget.Pin();
+		const ISlateStyle* const StyleSet = MultiBoxWidget->GetStyleSet();
+
+		static const FName IconName("MultiBox.GenericToolBarIcon.Small" );
+		return StyleSet->GetBrush(IconName);
+	}
+}
+
+EVisibility SToolBarStackButtonBlock::GetIconVisibility(bool bIsASmallIcon) const
+{
+	return ((bForceSmallIcons || FMultiBoxSettings::UseSmallToolBarIcons.Get()) ^ bIsASmallIcon) ? EVisibility::Collapsed : EVisibility::Visible;
+}
+
+FSlateColor SToolBarStackButtonBlock::GetIconForegroundColor() const
+{
+	// If any brush has a tint, don't assume it should be subdued
+	const FSlateBrush* Brush = GetIconBrush();
+	if (Brush && Brush->TintColor != FLinearColor::White)
+	{
+		return FLinearColor::White;
+	}
+
+	return FSlateColor::UseForeground();
+}
+
+const FSlateBrush* SToolBarStackButtonBlock::GetOverlayIconBrush() const
+{
+	const FSlateIcon& ActualIcon = Icon.Get();
+
+	if (ActualIcon.IsSet())
+	{
+		return ActualIcon.GetOverlayIcon();
+	}
+
+	return nullptr;
+}
+
diff --git a/Engine/Source/Runtime/Slate/Public/Framework/Commands/Contexts/UIContentContext.h b/Engine/Source/Runtime/Slate/Public/Framework/Commands/Contexts/UIContentContext.h
new file mode 100644
--- /dev/null
+++ b/Engine/Source/Runtime/Slate/Public/Framework/Commands/Contexts/UIContentContext.h
@@ -0,0 +1,28 @@
+// Copyright Epic Games, Inc. All Rights Reserved.
+
+#pragma once
+
+#include "CoreMinimal.h"
+#include "Layout/Visibility.h"
+#include "Styling/SlateTypes.h"
+#include "Framework/SlateDelegates.h"
+
+/**
+ * Context for generating content from a FUICommandInfo
+ */
+struct SLATE_API FUIContentContext : public IUIActionContextBase
+{
+	virtual ~FUIContentContext() {}
+
+	/** Holds a delegate that is executed to generate menu actions from the stack. */
+	FOnGetContent OnGetContent;
+
+	//~ Begin IUIActionContextBase Interface
+	virtual FName GetContextName() const override
+	{
+		return ContextName;
+	}
+	//~ End IUIActionContextBase Interface
+
+	inline static const FName ContextName = FName("ContentContext");
+};
\ No newline at end of file
diff --git a/Engine/Source/Runtime/Slate/Public/Framework/Commands/Contexts/UIIdentifierContext.h b/Engine/Source/Runtime/Slate/Public/Framework/Commands/Contexts/UIIdentifierContext.h
new file mode 100644
--- /dev/null
+++ b/Engine/Source/Runtime/Slate/Public/Framework/Commands/Contexts/UIIdentifierContext.h
@@ -0,0 +1,41 @@
+// Copyright Epic Games, Inc. All Rights Reserved.
+
+#pragma once
+
+#include "CoreMinimal.h"
+#include "Layout/Visibility.h"
+#include "Styling/SlateTypes.h"
+#include "Framework/SlateDelegates.h"
+#include "Textures/SlateIcon.h"
+
+/** Delegate for dynamic generation of icons */
+DECLARE_DELEGATE_RetVal(const FSlateIcon, FOnGetContextIcon)
+
+/** Delegate for dynamic generation of text */
+DECLARE_DELEGATE_RetVal(const FText, FOnGetContextText)
+
+/**
+ * Context for generating identifiers such as icons, labels, and descriptions from a FUICommandInfo
+ */
+struct SLATE_API FUIIdentifierContext : public IUIActionContextBase
+{
+	virtual ~FUIIdentifierContext() {}
+
+	/** Holds a delegate that is executed to generate context icon. */
+	FOnGetContextIcon OnGetContextIcon;
+
+	/** Holds a delegate that is executed to generate context label. */
+	FOnGetContextText OnGetContextabel;
+
+	/** Holds a delegate that is executed to generate context description. */
+	FOnGetContextText OnGetContextDescription;
+
+	//~ Begin IUIActionContextBase Interface
+	virtual FName GetContextName() const override
+	{
+		return ContextName;
+	}
+	//~ End IUIActionContextBase Interface
+
+	inline static const FName ContextName = FName("IdentifierContext");
+};
\ No newline at end of file
diff --git a/Engine/Source/Runtime/Slate/Public/Framework/Commands/UIAction.h b/Engine/Source/Runtime/Slate/Public/Framework/Commands/UIAction.h
--- a/Engine/Source/Runtime/Slate/Public/Framework/Commands/UIAction.h
+++ b/Engine/Source/Runtime/Slate/Public/Framework/Commands/UIAction.h
@@ -207,3 +207,33 @@ public:
 		return ECheckBoxState::Unchecked;
 	}
 };
+
+/**
+ * Base type for UI Action context object. 
+ * 
+ * Note: not using UObjects to avoid overhead.
+ */
+struct SLATE_API IUIActionContextBase
+{
+
+public:
+	virtual FName GetContextName() const = 0;
+};
+
+/**
+ * Container of multiple 'IUIActionContextBase' context objects for FUIActions.
+ * Use if you need or want to associate a FUICommandInfo with additional data.
+ * 
+ * Ex: An FOnGetContent delegate that can be called to generate a widget.
+ */
+struct SLATE_API FUIActionContext
+{
+public:
+
+	TSharedPtr<IUIActionContextBase> FindContext(const FName InName) const;
+
+	void AddContext(const TSharedPtr<IUIActionContextBase>& InContext);
+
+private:
+	TArray<TSharedPtr<IUIActionContextBase>> Contexts;
+};
\ No newline at end of file
diff --git a/Engine/Source/Runtime/Slate/Public/Framework/Commands/UICommandList.h b/Engine/Source/Runtime/Slate/Public/Framework/Commands/UICommandList.h
--- a/Engine/Source/Runtime/Slate/Public/Framework/Commands/UICommandList.h
+++ b/Engine/Source/Runtime/Slate/Public/Framework/Commands/UICommandList.h
@@ -95,6 +95,16 @@ public:
 	 */
 	virtual void MapAction( const TSharedPtr< const FUICommandInfo > InUICommandInfo, const FUIAction& InUIAction );
 
+	/**
+	 * Maps a command info to a series of delegates that are executed by a multibox or mouse/keyboard input
+	 * Additionally, add a container of context objects to also associate with this command
+	 *
+	 * @param InUICommandInfo	The command info to map
+	 * @param InUIAction		Action to map to this command
+	 * @param InUIActionContext	Action context to map to this command
+	 */
+	virtual void MapAction(const TSharedPtr< const FUICommandInfo > InUICommandInfo, const FUIAction& InUIAction, const FUIActionContext& InUIActionContext);
+
 	/**
 	 * Append commands in InCommandsToAppend to this command list.
 	 */
@@ -187,6 +197,13 @@ public:
 	  */
 	const FUIAction* GetActionForCommand(TSharedPtr<const FUICommandInfo> Command) const;
 
+	/** 
+	  * Attempts to find an context for the specified command in the current UICommandList. This is a wrapper for GetContextForCommandRecursively.
+	  *
+	  * @param Command				The UI command for which you are discovering an action
+	  */
+	const FUIActionContext* GetContextForCommand(TSharedPtr<const FUICommandInfo> Command) const;
+
 protected:
 
 	/**
@@ -212,6 +229,17 @@ protected:
 	  */
 	const FUIAction* GetActionForCommandRecursively(const TSharedRef<const FUICommandInfo>& Command, bool bIncludeChildren, bool bIncludeParents, TSet<TSharedRef<const FUICommandList>>& InOutVisitedLists) const;
 
+	/**
+	  * Attempts to find an context for the specified command in the current UICommandList. If it is not found, the context for the
+	  * specified command is discovered in the children recursively then the parents recursively.
+	  *
+	  * @param Command				The UI command for which you are discovering an action
+	  * @param bIncludeChildren		If true, children of this command list will be searched in the event that the action is not found
+	  * @param bIncludeParents		If true, parents of this command list will be searched in the event that the action is not found
+	  * @param OutVisitedLists		The set of visited lists during recursion. This is used to prevent cycles.
+	  */
+	const FUIActionContext* GetContextForCommandRecursively(const TSharedRef<const FUICommandInfo>& Command, bool bIncludeChildren, bool bIncludeParents, TSet<TSharedRef<const FUICommandList>>& InOutVisitedLists) const;
+
 	/** Returns all contexts associated with this list. This is a wrapper for GatherContextsForListRecursively */
 	void GatherContextsForList(TSet<FName>& OutAllContexts) const;
 
@@ -220,7 +248,8 @@ protected:
 
 private:
 
-	typedef TMap< const TSharedPtr< const FUICommandInfo >, FUIAction > FUIBindingMap;
+	using FUIBindingMap = TMap<const TSharedPtr<const FUICommandInfo>, FUIAction>;
+	using FUIContextMap = TMap<const TSharedPtr<const FUICommandInfo>, FUIActionContext>;
 	
 	/** Known contexts in this list.  Each context must be known so we can quickly look up commands from bindings */
 	TSet<FName> ContextsInList;
@@ -228,10 +257,16 @@ private:
 	/** Mapping of command to action */
 	FUIBindingMap UICommandBindingMap;
 
+	/** Mapping of command to action context, used for special cases and thus much more sparse */
+	FUIContextMap UICommandContextMap;
+
 	/** The list of parent and children UICommandLists */
 	TArray<TWeakPtr<FUICommandList>> ParentUICommandLists;
 	TArray<TWeakPtr<FUICommandList>> ChildUICommandLists;
 
 	/** Determines if this UICommandList is capable of producing an action for the supplied command */
 	FCanProduceActionForCommand CanProduceActionForCommand;
+
+	/** Helper struct for list recursion */
+	friend struct RecursiveCommandDataGetter;
 };
diff --git a/Engine/Source/Runtime/Slate/Public/Framework/MultiBox/MultiBoxBuilder.h b/Engine/Source/Runtime/Slate/Public/Framework/MultiBox/MultiBoxBuilder.h
--- a/Engine/Source/Runtime/Slate/Public/Framework/MultiBox/MultiBoxBuilder.h
+++ b/Engine/Source/Runtime/Slate/Public/Framework/MultiBox/MultiBoxBuilder.h
@@ -535,6 +535,13 @@ public:
 	 */
 	void AddComboButton( const FUIAction& InAction, const FOnGetContent& InMenuContentGenerator, const TAttribute<FText>& InLabelOverride = TAttribute<FText>(), const TAttribute<FText>& InToolTipOverride = TAttribute<FText>(), const TAttribute<FSlateIcon>& InIconOverride = TAttribute<FSlateIcon>(), bool bInSimpleComboBox = false, FName InTutorialHighlightName = NAME_None );
 
+	/**
+	 * Adds a tool bar stack button
+	 *
+	 * @param	InCommand				The command associated with this tool bar button
+	 * @param	InTutorialHighlightName	Name to identify this widget and highlight during tutorials
+	 */
+	void AddToolbarStackButton(const TSharedPtr< const FUICommandInfo > InCommand, FName InTutorialHighlightName = NAME_None);
 
 	/**
 	 * Adds any widget to the toolbar
diff --git a/Engine/Source/Runtime/Slate/Public/Framework/MultiBox/SToolBarStackButtonBlock.h b/Engine/Source/Runtime/Slate/Public/Framework/MultiBox/SToolBarStackButtonBlock.h
new file mode 100644
--- /dev/null
+++ b/Engine/Source/Runtime/Slate/Public/Framework/MultiBox/SToolBarStackButtonBlock.h
@@ -0,0 +1,199 @@
+// Copyright Epic Games, Inc. All Rights Reserved.
+
+#pragma once
+
+#include "CoreMinimal.h"
+#include "Misc/Attribute.h"
+#include "Layout/Visibility.h"
+#include "Widgets/SWidget.h"
+#include "Textures/SlateIcon.h"
+#include "Widgets/DeclarativeSyntaxSupport.h"
+#include "Framework/MultiBox/MultiBox.h"
+#include "Framework/Commands/Contexts/UIIdentifierContext.h"
+#include "Framework/Commands/Contexts/UIContentContext.h"
+
+class SCheckBoxStack;
+
+/**
+ * Tool bar combo button MultiBlock, but toggleable. Right & long-press clicks spawn menu.
+ * Designed to hold a stack of tools in one block.
+ */
+class FToolBarStackButtonBlock
+	: public FMultiBlock
+{
+
+public:
+
+	/**
+	 * Constructor
+	 *
+	 * @param	InCommand			The command associated with this tool bar button
+	 * @param	InCommandList		The list of commands that are mapped to delegates so that we know what to execute for this button
+	 * @param	bInSimpleStackBox	If true, the icon and label won't be displayed
+	 */
+	FToolBarStackButtonBlock(const TSharedRef< const FUICommandInfo > InCommand, TSharedPtr< const FUICommandList > InCommandList, bool bInSimpleStackBox = false);
+
+	/** FMultiBlock interface */
+	virtual void CreateMenuEntry(class FMenuBuilder& MenuBuilder) const override;
+
+	/** 
+	 * Sets the visibility of the blocks label
+	 *
+	 * @param InLabelVisibility		Visibility setting to use for the label
+	 */
+	void SetLabelVisibility( EVisibility InLabelVisibility ) { LabelVisibility = InLabelVisibility; }
+
+	/** Set whether this toolbar should always use small icons, regardless of the current settings */
+	void SetForceSmallIcons( const bool InForceSmallIcons ) { bForceSmallIcons = InForceSmallIcons; }
+protected:
+
+	/**
+	 * Allocates a widget for this type of MultiBlock.  Override this in derived classes.
+	 *
+	 * @return  MultiBlock widget object
+	 */
+	virtual TSharedRef< class IMultiBlockBaseWidget > ConstructWidget() const override;
+
+	/** Provides dynamic, icons, labels, & descriptions for this stack, caches on get */
+	const TSharedPtr<FUIIdentifierContext> GetStackIdentifier() const;
+
+	/** Provides context menu content stack, caches on get */
+	const TSharedPtr<FUIContentContext> GetStackContent() const;
+
+protected:
+
+	// Friend our corresponding widget class
+	friend class SToolBarStackButtonBlock;
+
+	/** Delegate that generates a widget for this combo button's menu content.  Called when the menu is summoned. */
+	FOnGetContent MenuContentGenerator;
+
+	/** Controls the Labels visibility, defaults to GetIconVisibility if no override is provided */
+	TOptional< EVisibility > LabelVisibility;
+
+	/** If true, the icon and label won't be displayed */
+	bool bSimpleComboBox;
+
+	/** Whether this toolbar should always use small icons, regardless of the current settings */
+	bool bForceSmallIcons;
+
+	/** To avoid re-searching the FUICommandList, we cache the last identifier context */
+	mutable TWeakPtr<FUIIdentifierContext> CachedStackIdentifier;
+
+	/** To avoid re-searching the FUICommandList, we cache the last content context */
+	mutable TWeakPtr<FUIContentContext> CachedStackContent;
+};
+
+
+
+/**
+ * Tool bar button MultiBlock widget
+ */
+class SLATE_API SToolBarStackButtonBlock
+	: public SMultiBlockBaseWidget
+{
+
+public:
+
+	SLATE_BEGIN_ARGS( SToolBarStackButtonBlock )
+		: _ForceSmallIcons( false )
+	{}
+
+		/** Controls the visibility of the blocks label */
+		SLATE_ARGUMENT( TOptional< EVisibility >, LabelVisibility )
+
+		/** Whether this toolbar should always use small icons, regardless of the current settings */
+		SLATE_ARGUMENT( bool, ForceSmallIcons )
+
+	SLATE_END_ARGS()
+
+
+	/**
+	 * Builds this MultiBlock widget up from the MultiBlock associated with it
+	 */
+	virtual void BuildMultiBlockWidget(const ISlateStyle* StyleSet, const FName& StyleName) override;
+
+
+	/**
+	 * Construct this widget
+	 *
+	 * @param	InArgs	The declaration data for this widget
+	 */
+	void Construct( const FArguments& InArgs );
+
+protected:
+
+	/**
+	 * Called by Slate when this tool bar check box button is toggled
+	 */
+	void OnCheckStateChanged(const ECheckBoxState NewCheckedState);
+
+	/**
+	 * Called by slate to determine if this button should appear checked
+	 *
+	 * @return ECheckBoxState::Checked if it should be checked, ECheckBoxState::Unchecked if not.
+	 */
+	ECheckBoxState GetCheckState() const;
+
+	/**
+	 * Called by Slate when content for this button's menu needs to be generated
+	 *
+	 * @return	The widget to use for the menu content
+	 */
+	TSharedRef<SWidget> OnGetMenuContent();
+
+	/**
+	 * Called by Slate to determine if this button is enabled
+	 * 
+	 * @return True if the menu entry is enabled, false otherwise
+	 */
+	bool IsEnabled() const;
+
+	/**
+	 * Called by Slate to determine if this button is visible
+	 *
+	 * @return EVisibility::Visible or EVisibility::Collapsed, depending on if the button should be displayed
+	 */
+	EVisibility GetVisibility() const;
+protected:
+
+	/** Gets the icon brush for the toolbar block widget */
+	const FSlateBrush* GetIconBrush() const;
+
+	/** Gets the label for the toolbar block widget */
+	FText GetLabel() const;
+
+	/** Gets the description for the toolbar block widget */
+	FText GetDescription() const;
+
+	/** @return The icon for the toolbar button; may be dynamic, so check HasDynamicIcon */
+	const FSlateBrush* GetNormalIconBrush() const;
+
+	/** @return The small icon for the toolbar button; may be dynamic, so check HasDynamicIcon */
+	const FSlateBrush* GetSmallIconBrush() const;
+
+	/** Called by Slate to determine whether icons/labels are visible */
+	EVisibility GetIconVisibility(bool bIsASmallIcon) const;
+
+	FSlateColor GetIconForegroundColor() const;
+
+	const FSlateBrush* GetOverlayIconBrush() const;
+
+protected:
+	/** Overrides the visibility of the of label. This is used to set up the LabelVisibility attribute */
+	TOptional<EVisibility> LabelVisibilityOverride;
+
+	/** Controls the visibility of the of label, defaults to GetIconVisibility */
+	TAttribute< EVisibility > LabelVisibility;
+
+	/** Optional overridden icon for this tool bar button.  IF not set, then the action's icon will be used instead. */
+	TAttribute<FSlateIcon> Icon;
+
+	TSharedPtr<SCheckBoxStack> StackButtonWidget;
+
+	/** The foreground color for button when the combo button is open */
+	FSlateColor OpenForegroundColor;
+
+	/** Whether this toolbar should always use small icons, regardless of the current settings */
+	bool bForceSmallIcons;
+};
diff --git a/Engine/Source/Runtime/Slate/Public/SlateFwd.h b/Engine/Source/Runtime/Slate/Public/SlateFwd.h
--- a/Engine/Source/Runtime/Slate/Public/SlateFwd.h
+++ b/Engine/Source/Runtime/Slate/Public/SlateFwd.h
@@ -14,6 +14,8 @@ class FToolBarButtonBlock;
 class SToolBarButtonBlock;
 class FToolBarComboButtonBlock;
 class SToolBarComboButtonBlock;
+class FToolBarStackButtonBlock;
+class SToolBarStackButtonBlock;
 class SHyperlink;
 class SRichTextHyperlink;
 class SThrobber;
