Slate: Fixed a visual issue with search box size changed on clicked. This was because the editabletextbox always append an error text widget to the horizontalbox on text changed if there isn't one already.

The new behavior ensure the error widget is created during construct so that the initial size of the textbox is correct

diff --git a/Engine/Source/Runtime/Slate/Private/Widgets/Input/SEditableTextBox.cpp b/Engine/Source/Runtime/Slate/Private/Widgets/Input/SEditableTextBox.cpp
--- a/Engine/Source/Runtime/Slate/Private/Widgets/Input/SEditableTextBox.cpp
+++ b/Engine/Source/Runtime/Slate/Private/Widgets/Input/SEditableTextBox.cpp
@@ -98,6 +98,15 @@ void SEditableTextBox::Construct( const FArguments& InArgs )
 			ErrorReporting->AsWidget()
 		];
 	}
+	else
+	{
+		// this also creates a default widget
+		// if we don't create the widget in Construct() 
+		// it will get created in OnEditableTextChanged()
+		// create it now so that the default size of the textbox
+		// won't grow after user use it once
+		SetError(FText::GetEmpty());
+	}
 }
 
 void SEditableTextBox::SetStyle(const FEditableTextBoxStyle* InStyle)
