Update the pinned icon on sidebars

diff --git a/Engine/Source/Runtime/Slate/Private/Framework/Docking/STabSidebar.cpp b/Engine/Source/Runtime/Slate/Private/Framework/Docking/STabSidebar.cpp
--- a/Engine/Source/Runtime/Slate/Private/Framework/Docking/STabSidebar.cpp
+++ b/Engine/Source/Runtime/Slate/Private/Framework/Docking/STabSidebar.cpp
@@ -221,8 +221,7 @@ public:
 							[
 								SNew(SImage)
 								.ColorAndOpacity(FSlateColor::UseForeground())
-								.Image(FAppStyle::Get().GetBrush("Docking.Sidebar.PinTabIcon"))
-								.DesiredSizeOverride(FVector2D(12, 12))
+								.Image(this, &STabDrawerButton::GetPinButtonImage)
 							]
 						]
 					]
@@ -370,6 +369,18 @@ private:
 		}
 	}
 
+	const FSlateBrush* GetPinButtonImage() const
+	{
+		if (IsTabPinned(Tab.ToSharedRef()))
+		{
+			return FAppStyle::Get().GetBrush("Icons.Pinned");
+		}
+		else
+		{
+			return FAppStyle::Get().GetBrush("Icons.Unpinned");
+		}
+	}
+
 	void OnPinButtonCheckStateChanged(ECheckBoxState State)
 	{
 		OnDrawerPinButtonToggled.ExecuteIfBound(Tab.ToSharedRef(), State == ECheckBoxState::Checked);
