Fix to flickering on ComboBox menu when the combobox size changes.

diff --git a/Engine/Source/Runtime/Slate/Private/Widgets/Input/SMenuAnchor.cpp b/Engine/Source/Runtime/Slate/Private/Widgets/Input/SMenuAnchor.cpp
--- a/Engine/Source/Runtime/Slate/Private/Widgets/Input/SMenuAnchor.cpp
+++ b/Engine/Source/Runtime/Slate/Private/Widgets/Input/SMenuAnchor.cpp
@@ -145,6 +145,16 @@ void SMenuAnchor::Tick( const FGeometry& AllottedGeometry, const double InCurren
 		// passed down in FGeometry, along with the window DPI Scale, as one extra value code can take into account if it needs to.
 		const FVector2D NewSize = PopupGeometry.GetDrawSize() / ( AllottedGeometry.GetAccumulatedLayoutTransform().GetScale() / PopupWindow->GetLocalToWindowTransform().GetScale() );
 
+		// When in a ComboBox, we want the submenu to be at minimum the size of the combo box, so we update the Popupwindow minimum size.
+		const EMenuPlacement PlacementMode = Placement.Get();
+		if (PlacementMode == MenuPlacement_ComboBox)
+		{
+			// Set the new size limits
+			FWindowSizeLimits SizeLimits = PopupWindow->GetSizeLimits();
+			SizeLimits.SetMinWidth(NewSize.X);
+			PopupWindow->SetSizeLimits(SizeLimits);
+		}
+
 		// We made a window for showing the popup.
 		// Update the window's position!
 		PopupWindow->ReshapeWindow(NewPosition, NewSize);
diff --git a/Engine/Source/Runtime/SlateCore/Private/Widgets/SWindow.cpp b/Engine/Source/Runtime/SlateCore/Private/Widgets/SWindow.cpp
--- a/Engine/Source/Runtime/SlateCore/Private/Widgets/SWindow.cpp
+++ b/Engine/Source/Runtime/SlateCore/Private/Widgets/SWindow.cpp
@@ -709,6 +709,11 @@ FWindowSizeLimits SWindow::GetSizeLimits() const
 	return SizeLimits;
 }
 
+void SWindow::SetSizeLimits(const FWindowSizeLimits& InSizeLimits)
+{
+	SizeLimits = InSizeLimits;
+}
+
 void SWindow::SetAllowFastUpdate(bool bInAllowFastUpdate)
 {
 	if (bAllowFastUpdate != bInAllowFastUpdate)
diff --git a/Engine/Source/Runtime/SlateCore/Public/Widgets/SWindow.h b/Engine/Source/Runtime/SlateCore/Public/Widgets/SWindow.h
--- a/Engine/Source/Runtime/SlateCore/Public/Widgets/SWindow.h
+++ b/Engine/Source/Runtime/SlateCore/Public/Widgets/SWindow.h
@@ -1058,6 +1058,9 @@ public:
 	/** Optional constraints on min and max sizes that this window can be. */
 	FWindowSizeLimits GetSizeLimits() const;
 
+	/** Set optional constraints on min and max sizes that this window can be. */
+	void SetSizeLimits(const FWindowSizeLimits& InSizeLimits);
+
 	void SetAllowFastUpdate(bool bInAllowFastUpdate);
 public:
 
