Slate: Added more customization option to the search box, which allows for setting the alignment of search result buttons and the glass image separately

diff --git a/Engine/Source/Runtime/Slate/Private/Widgets/Input/SSearchBox.cpp b/Engine/Source/Runtime/Slate/Private/Widgets/Input/SSearchBox.cpp
--- a/Engine/Source/Runtime/Slate/Private/Widgets/Input/SSearchBox.cpp
+++ b/Engine/Source/Runtime/Slate/Private/Widgets/Input/SSearchBox.cpp
@@ -6,6 +6,7 @@
 #include "Widgets/Input/SButton.h"
 #include "Widgets/Text/STextBlock.h"
 #include "Widgets/Images/SThrobber.h"
+#include "Widgets/Layout/SSpacer.h"
 
 void SSearchBox::Construct( const FArguments& InArgs )
 {
@@ -36,8 +37,8 @@ void SSearchBox::Construct( const FArguments& InArgs )
 		.OnKeyDownHandler( InArgs._OnKeyDownHandler )
 	);
 
-	// If we want to have the buttons appear to the left of the text box we have to insert the slots instead of add them
-	int32 SlotIndex = InArgs._Style->bLeftAlignButtons ? 0 : Box->NumSlots();
+	// If we want to have the search result buttons appear to the left of the text box we have to insert the slots instead of add them
+	int32 SlotIndex = InArgs._Style->bLeftAlignSearchResultButtons ? 0 : Box->NumSlots();
 
 	// Add a throbber to show if there is a search running.
 	Box->InsertSlot(SlotIndex++)
@@ -56,7 +57,7 @@ void SSearchBox::Construct( const FArguments& InArgs )
 	// If a search delegate was bound, add a previous and next button
 	if (OnSearchDelegate.IsBound())
 	{
-		FMargin SearchResultPadding = InArgs._Style->bLeftAlignButtons ? FMargin(5, 0, 2, 0) : FMargin(2, 0, 2, 0);
+		FMargin SearchResultPadding = InArgs._Style->bLeftAlignSearchResultButtons ? FMargin(5, 0, 2, 0) : FMargin(2, 0, 2, 0);
 
 		// Search result data text
 		Box->InsertSlot(SlotIndex++)
@@ -118,7 +119,17 @@ void SSearchBox::Construct( const FArguments& InArgs )
 	}
 
 	// Add a search glass image so that the user knows this text box is for searching
-	Box->InsertSlot(SlotIndex++)
+	int GlassImageSlotIndex = InArgs._Style->bLeftAlignGlassImageAndClearButton ? 0 : Box->NumSlots();
+
+	if (InArgs._Style->bLeftAlignSearchResultButtons == InArgs._Style->bLeftAlignGlassImageAndClearButton)
+	{
+		// if they are on the same side, Glass Image follows the search result buttons
+		GlassImageSlotIndex = SlotIndex;
+	}
+
+	const int32 ClearButtonIndex = GlassImageSlotIndex + 1;
+	
+	Box->InsertSlot(GlassImageSlotIndex)
 		.AutoWidth()
 		.Padding(InArgs._Style->ImagePadding)
 		.HAlign(HAlign_Center)
@@ -131,8 +142,8 @@ void SSearchBox::Construct( const FArguments& InArgs )
 		];
 
 	// Add an X to clear the search whenever there is some text typed into it
-	FMargin ClearButtonMargin = InArgs._Style->bLeftAlignButtons ? FMargin(2.f, 0.f) : FMargin(2.f, 0.f, 4.f, 0.f);
-	Box->InsertSlot(SlotIndex++)
+	FMargin ClearButtonMargin = FMargin(2.f, 0.f);
+	Box->InsertSlot(ClearButtonIndex)
 	.AutoWidth()
 	.HAlign(HAlign_Center)
 	.VAlign(VAlign_Center)
@@ -153,6 +164,15 @@ void SSearchBox::Construct( const FArguments& InArgs )
 			.ColorAndOpacity(FSlateColor::UseForeground())
 		]
 	];
+
+	// Padding at the end so that the text box grows long enough to host all buttons
+	// the text box needs extra space because it is rounded on both ends
+	Box->AddSlot()
+	.AutoWidth()
+	[
+		SNew(SSpacer)
+		.Size(FVector2d(7,0))
+	];
 }
 
 EActiveTimerReturnType SSearchBox::TriggerOnTextChanged( double, float, FText NewText )
diff --git a/Engine/Source/Runtime/SlateCore/Private/Styling/SlateTypes.cpp b/Engine/Source/Runtime/SlateCore/Private/Styling/SlateTypes.cpp
--- a/Engine/Source/Runtime/SlateCore/Private/Styling/SlateTypes.cpp
+++ b/Engine/Source/Runtime/SlateCore/Private/Styling/SlateTypes.cpp
@@ -377,7 +377,9 @@ const FExpandableAreaStyle& FExpandableAreaStyle::GetDefault()
 
 
 FSearchBoxStyle::FSearchBoxStyle()
-	: bLeftAlignButtons(false)
+	: bLeftAlignButtons_DEPRECATED(false)
+	, bLeftAlignSearchResultButtons(false)
+	, bLeftAlignGlassImageAndClearButton(false)
 {
 }
 
@@ -408,6 +410,16 @@ const FSearchBoxStyle& FSearchBoxStyle::GetDefault()
 	return Default;
 }
 
+FSearchBoxStyle& FSearchBoxStyle::SetLeftAlignButtons(bool bInLeftAlignButtons)
+{
+	bLeftAlignButtons_DEPRECATED = bInLeftAlignButtons;
+	
+	bLeftAlignSearchResultButtons = bInLeftAlignButtons;
+	bLeftAlignGlassImageAndClearButton= bInLeftAlignButtons;
+
+	return *this;
+}
+
 FSliderStyle::FSliderStyle()
 	: BarThickness(2.0f)
 {
diff --git a/Engine/Source/Runtime/SlateCore/Private/Styling/StarshipCoreStyle.cpp b/Engine/Source/Runtime/SlateCore/Private/Styling/StarshipCoreStyle.cpp
--- a/Engine/Source/Runtime/SlateCore/Private/Styling/StarshipCoreStyle.cpp
+++ b/Engine/Source/Runtime/SlateCore/Private/Styling/StarshipCoreStyle.cpp
@@ -1091,7 +1091,8 @@ void FStarshipCoreStyle::SetupTextStyles(TSharedRef<FStyle>& Style)
 			.SetGlassImage(IMAGE_BRUSH_SVG("Starship/Common/search", Icon16x16))
 			.SetClearImage(IMAGE_BRUSH_SVG("Starship/Common/close", Icon16x16))
 			.SetImagePadding(FMargin(3.f, 0.f, 0.f, 0.f))
-			.SetLeftAlignButtons(true)
+			.SetLeftAlignSearchResultButtons(true)
+			.SetLeftAlignGlassImageAndClearButton(true)
 		);
 	}
 
diff --git a/Engine/Source/Runtime/SlateCore/Public/Styling/SlateTypes.h b/Engine/Source/Runtime/SlateCore/Public/Styling/SlateTypes.h
--- a/Engine/Source/Runtime/SlateCore/Public/Styling/SlateTypes.h
+++ b/Engine/Source/Runtime/SlateCore/Public/Styling/SlateTypes.h
@@ -1207,9 +1207,19 @@ struct SLATECORE_API FSearchBoxStyle : public FSlateWidgetStyle
 	FSearchBoxStyle& SetImagePadding(const FMargin& InImagePadding){ ImagePadding = InImagePadding; return *this; }
 
 	/** If true, buttons appear to the left of the search text */
+	UPROPERTY(meta = (DeprecatedProperty, DeprecationMessage="Use LeftAlignSearchResultButtons and LeftAlignGlassImageAndClearButton instead"))
+	bool bLeftAlignButtons_DEPRECATED;
+	FSearchBoxStyle& SetLeftAlignButtons(bool bInLeftAlignButtons);
+
+	/** If true, search result buttons appear to the left of the search text */
+	UPROPERTY(EditAnywhere, Category = Appearance)
+	bool bLeftAlignSearchResultButtons;
+	FSearchBoxStyle& SetLeftAlignSearchResultButtons(bool bInLeftAlignSearchResultButtons){ bLeftAlignSearchResultButtons = bInLeftAlignSearchResultButtons; return *this; }
+	
+	/** If true, glass image and clear button appear to the left of the search text */
 	UPROPERTY(EditAnywhere, Category = Appearance)
-	bool bLeftAlignButtons;
-	FSearchBoxStyle& SetLeftAlignButtons(bool bInLeftAlignButtons){ bLeftAlignButtons = bInLeftAlignButtons; return *this; }
+	bool bLeftAlignGlassImageAndClearButton;
+	FSearchBoxStyle& SetLeftAlignGlassImageAndClearButton(bool bInLeftAlignGlassImageAndClearButton){ bLeftAlignGlassImageAndClearButton = bInLeftAlignGlassImageAndClearButton; return *this; }
 };
 
 
