SlateEditableTextLayout: Minor bug fixes + Improvements to:
+Fixed a bug where AdvanceSearch would stop advancing if there are multiple matches in a single line and only in that line.
+Fixed a bug where AdvanceSearch would stop advancing if there is a match at the end of the line
+Fixed a bug where selected text would be deleted when an invalid char is entered. The new behavior ensures that invalid char are treated as nop
+Improved scrolling behavior when advancing search, it now attempts to show as much of the matched text in the view as possible, while keeping the cursor in the view as well
+Added APIs to query search results
+Added API to query text selection

diff --git a/Engine/Source/Runtime/Slate/Private/Widgets/Text/SMultiLineEditableText.cpp b/Engine/Source/Runtime/Slate/Private/Widgets/Text/SMultiLineEditableText.cpp
--- a/Engine/Source/Runtime/Slate/Private/Widgets/Text/SMultiLineEditableText.cpp
+++ b/Engine/Source/Runtime/Slate/Private/Widgets/Text/SMultiLineEditableText.cpp
@@ -139,6 +139,16 @@ FText SMultiLineEditableText::GetSearchText() const
 	return EditableTextLayout->GetSearchText();
 }
 
+int32 SMultiLineEditableText::GetSearchResultIndex() const
+{
+	return EditableTextLayout->GetSearchResultIndex();
+}
+
+int32 SMultiLineEditableText::GetNumSearchResults() const
+{
+	return EditableTextLayout->GetNumSearchResults();
+}
+
 void SMultiLineEditableText::SetTextStyle(const FTextBlockStyle* InTextStyle)
 {
 	if (InTextStyle)
@@ -470,6 +480,11 @@ FText SMultiLineEditableText::GetSelectedText() const
 	return EditableTextLayout->GetSelectedText();
 }
 
+FTextSelection SMultiLineEditableText::GetSelection() const
+{
+	return EditableTextLayout->GetSelection();
+}
+
 void SMultiLineEditableText::InsertTextAtCursor(const FText& InText)
 {
 	EditableTextLayout->InsertTextAtCursor(InText.ToString());
diff --git a/Engine/Source/Runtime/Slate/Private/Widgets/Text/SlateEditableTextLayout.cpp b/Engine/Source/Runtime/Slate/Private/Widgets/Text/SlateEditableTextLayout.cpp
--- a/Engine/Source/Runtime/Slate/Private/Widgets/Text/SlateEditableTextLayout.cpp
+++ b/Engine/Source/Runtime/Slate/Private/Widgets/Text/SlateEditableTextLayout.cpp
@@ -307,6 +307,16 @@ FText FSlateEditableTextLayout::GetSearchText() const
 	return SearchText;
 }
 
+int32 FSlateEditableTextLayout::GetSearchResultIndex() const
+{
+	return CurrentSearchResultIndex;
+}
+
+int32 FSlateEditableTextLayout::GetNumSearchResults() const
+{
+	return SearchResultToIndexMap.Num();
+}
+
 void FSlateEditableTextLayout::SetTextStyle(const FTextBlockStyle& InTextStyle)
 {
 	TextStyle = InTextStyle;
@@ -434,6 +444,15 @@ FText FSlateEditableTextLayout::GetSelectedText() const
 	return FText::GetEmpty();
 }
 
+FTextSelection FSlateEditableTextLayout::GetSelection() const
+{
+	const FTextLocation CursorInteractionPosition = CursorInfo.GetCursorInteractionLocation();
+	const FTextLocation SelectionLocation = SelectionStart.Get(CursorInteractionPosition);
+	const FTextSelection Selection(SelectionLocation, CursorInteractionPosition);
+	
+	return Selection;
+}
+
 void FSlateEditableTextLayout::SetTextShapingMethod(const TOptional<ETextShapingMethod>& InTextShapingMethod)
 {
 	TextLayout->SetTextShapingMethod((InTextShapingMethod.IsSet()) ? InTextShapingMethod.GetValue() : GetDefaultTextShapingMethod());
@@ -604,9 +623,12 @@ void FSlateEditableTextLayout::BeginSearch(const FText& InSearchText, const ESea
 
 void FSlateEditableTextLayout::AdvanceSearch(const bool InReverse)
 {
+	//FirstMatchedLocation used as a key to find the index of the first matched string among all matches
+	FTextLocation FirstMatchedLocation(INDEX_NONE, INDEX_NONE);
+	const FTextLocation CursorInteractionPosition = CursorInfo.GetCursorInteractionLocation();
+	
 	if (!SearchText.IsEmpty())
 	{
-		const FTextLocation CursorInteractionPosition = CursorInfo.GetCursorInteractionLocation();
 		const FTextLocation SelectionLocation = SelectionStart.Get(CursorInteractionPosition);
 		const FTextSelection Selection(SelectionLocation, CursorInteractionPosition);
 
@@ -618,17 +640,43 @@ void FSlateEditableTextLayout::AdvanceSearch(const bool InReverse)
 
 		int32 CurrentLineIndex = SearchStartLocation.GetLineIndex();
 		int32 CurrentLineOffset = SearchStartLocation.GetOffset();
+
+		int32 NumLinesSearched = 0;
+		
 		do
 		{
 			const FTextLayout::FLineModel& Line = Lines[CurrentLineIndex];
 
-			// Do we have a match on this line?
-			const int32 CurrentSearchBegin = Line.Text->Find(SearchTextString, SearchCase, InReverse ? ESearchDir::FromEnd : ESearchDir::FromStart, CurrentLineOffset);
-			if (CurrentSearchBegin != INDEX_NONE)
+			bool bShouldSearchLine = true;
+			
+			if (!InReverse && CurrentLineOffset >= Line.Text->Len() )
 			{
-				SelectionStart = FTextLocation(CurrentLineIndex, CurrentSearchBegin);
-				CursorInfo.SetCursorLocationAndCalculateAlignment(*TextLayout, FTextLocation(CurrentLineIndex, CurrentSearchBegin + SearchTextLength));
-				break;
+				// CurrentLineOffset needs to be less than len(),
+				// otherwise, Find() clamps it to len() - 1 (see FString::Find()),
+				// and if there is a match at len() - 1, the search gets stuck
+				// 
+				// for example, for text "[cursor]abcd", len() = 4, len() - 1 = 3
+				// if you search for 'd', after the first advance,
+				// we get "abcd[cursor]", where CurrentLineOffset = 4
+				// if we advance one more time with Find('d', offset = 4)
+				// internally it becomes a search starting from len() - 1 = 3,
+				// since at index 3 there is a 'd' match, the CurrentSearchBegin is 3
+				// as a result, cursor position is set to 3 + len('d') = 4, again
+				bShouldSearchLine = false;
+			}
+
+			if (bShouldSearchLine)
+			{
+				// Do we have a match on this line?
+				const int32 CurrentSearchBegin = Line.Text->Find(SearchTextString, SearchCase, InReverse ? ESearchDir::FromEnd : ESearchDir::FromStart, CurrentLineOffset);
+				if (CurrentSearchBegin != INDEX_NONE)
+				{
+					SelectionStart = FTextLocation(CurrentLineIndex, CurrentSearchBegin);
+					CursorInfo.SetCursorLocationAndCalculateAlignment(*TextLayout, FTextLocation(CurrentLineIndex, CurrentSearchBegin + SearchTextLength));
+
+					FirstMatchedLocation = SelectionStart.GetValue();
+					break;
+				}
 			}
 
 			if (InReverse)
@@ -651,11 +699,43 @@ void FSlateEditableTextLayout::AdvanceSearch(const bool InReverse)
 				}
 				CurrentLineOffset = 0;
 			}
-		}
-		while (CurrentLineIndex != SearchStartLocation.GetLineIndex());
+
+			NumLinesSearched++;
+
+		}while(NumLinesSearched <= Lines.Num());
+		// use "<=" because if we start a search from the middle of a line
+		// the search should wrap around and search from the beginning of the line
+		// so loop twice even if there is only a single line
 	}
 
 	UpdateCursorHighlight();
+
+	// UpdateCursorHighlight() ensures SearchResultToIndexMap is up to date
+	CurrentSearchResultIndex = 0;
+	if (FirstMatchedLocation.IsValid())
+	{
+		int32* Index = SearchResultToIndexMap.Find(FirstMatchedLocation);
+		if (ensure(Index))
+		{
+			CurrentSearchResultIndex = *Index;
+		}
+		
+		// PositionToScrollIntoView is set to cursor position in UpdateCursorHighlight();
+		// Scrolling to cursor position directly does not always produce a good result
+		// because you can have only the last letter of the matched text in the view
+		// and most of the text out of view. The following code addresses this problem
+		const SlateEditableTextTypes::FScrollInfo SelectionScrollInfo(GetSelection().GetBeginning(), SlateEditableTextTypes::ECursorAlignment::Left);
+		const FVector2D LocalSelectionBeginLocation = TextLayout->GetLocationAt(SelectionScrollInfo.Position, false) / TextLayout->GetScale();
+
+		// if we are moving towards left side, reset the scroll offset
+		// such that when the cursor is scrolled into the view during tick
+		// , we move as much of the selection into the view as possible
+		if (LocalSelectionBeginLocation.X < 0.0f)
+		{
+			ScrollOffset.X = 0;
+		}
+	}
+
 }
 
 FVector2D FSlateEditableTextLayout::SetHorizontalScrollFraction(const float InScrollOffsetFraction)
@@ -1479,16 +1559,16 @@ bool FSlateEditableTextLayout::HandleTypeChar(const TCHAR InChar)
 		return false;
 	}
 
-	if (AnyTextSelected())
-	{
-		// Delete selected text
-		DeleteSelectedText();
-	}
-
 	// Certain characters are not allowed
 	const bool bIsCharAllowed = IsCharAllowed(InChar);
 	if (bIsCharAllowed)
 	{
+		if (AnyTextSelected())
+		{
+			// Delete selected text only if an allowed char is received
+			DeleteSelectedText();
+		}
+		
 		const FTextLocation CursorInteractionPosition = CursorInfo.GetCursorInteractionLocation();
 		const TArray< FTextLayout::FLineModel >& Lines = TextLayout->GetLineModels();
 		const FTextLayout::FLineModel& Line = Lines[CursorInteractionPosition.GetLineIndex()];
@@ -2539,6 +2619,7 @@ void FSlateEditableTextLayout::UpdateCursorHighlight()
 		const FString& SearchTextString = SearchText.ToString();
 		const int32 SearchTextLength = SearchTextString.Len();
 
+		int32 SearchResultIndex = 0;
 		const TArray< FTextLayout::FLineModel >& Lines = TextLayout->GetLineModels();
 		for (int32 LineIndex = 0; LineIndex < Lines.Num(); ++LineIndex)
 		{
@@ -2551,6 +2632,12 @@ void FSlateEditableTextLayout::UpdateCursorHighlight()
 			{
 				FindBegin = CurrentSearchBegin + SearchTextLength;
 				ActiveLineHighlights.Add(FTextLineHighlight(LineIndex, FTextRange(CurrentSearchBegin, FindBegin), SearchHighlightZOrder, SearchSelectionHighlighter.ToSharedRef()));
+
+				// SearchResultIndex starts from 1
+				// for example, if it is used to display stats about search results
+				// it would appear as "1 of 5" for the first match among five matches.
+				SearchResultIndex++;
+				SearchResultToIndexMap.Add(FTextLocation(LineIndex, CurrentSearchBegin), SearchResultIndex);
 			}
 		}
 
@@ -2658,6 +2745,7 @@ void FSlateEditableTextLayout::RemoveCursorHighlight()
 	}
 
 	ActiveLineHighlights.Empty();
+	SearchResultToIndexMap.Reset();
 }
 
 void FSlateEditableTextLayout::UpdatePreferredCursorScreenOffsetInLine()
diff --git a/Engine/Source/Runtime/Slate/Public/Framework/Text/TextLayout.h b/Engine/Source/Runtime/Slate/Public/Framework/Text/TextLayout.h
--- a/Engine/Source/Runtime/Slate/Public/Framework/Text/TextLayout.h
+++ b/Engine/Source/Runtime/Slate/Public/Framework/Text/TextLayout.h
@@ -123,6 +123,11 @@ public:
 	int32 GetOffset() const { return Offset; }
 	bool IsValid() const { return LineIndex != INDEX_NONE && Offset != INDEX_NONE; }
 
+	friend FORCEINLINE uint32 GetTypeHash(const FTextLocation& InSubject)
+	{
+		return HashCombine(InSubject.LineIndex, InSubject.Offset);
+	}
+
 private:
 	int32 LineIndex;
 	int32 Offset;
diff --git a/Engine/Source/Runtime/Slate/Public/Widgets/Text/SMultiLineEditableText.h b/Engine/Source/Runtime/Slate/Public/Widgets/Text/SMultiLineEditableText.h
--- a/Engine/Source/Runtime/Slate/Public/Widgets/Text/SMultiLineEditableText.h
+++ b/Engine/Source/Runtime/Slate/Public/Widgets/Text/SMultiLineEditableText.h
@@ -253,6 +253,12 @@ public:
 	/** Get the text that is currently being searched for (if any) */
 	FText GetSearchText() const;
 
+	/** Get the index of the search result (0 if none) */
+	int32 GetSearchResultIndex() const;
+
+	/** Get the total number of search results (0 if none) */
+	int32 GetNumSearchResults() const;
+
 	/** See attribute TextStyle */
 	void SetTextStyle(const FTextBlockStyle* InTextStyle);
 
@@ -345,6 +351,9 @@ public:
 	/** Get the currently selected text */
 	FText GetSelectedText() const;
 
+	/** Get the current selection */
+	FTextSelection GetSelection() const;
+
 	/** Insert the given text at the current cursor position, correctly taking into account new line characters */
 	void InsertTextAtCursor(const FText& InText);
 	void InsertTextAtCursor(const FString& InString);
diff --git a/Engine/Source/Runtime/Slate/Public/Widgets/Text/SlateEditableTextLayout.h b/Engine/Source/Runtime/Slate/Public/Widgets/Text/SlateEditableTextLayout.h
--- a/Engine/Source/Runtime/Slate/Public/Widgets/Text/SlateEditableTextLayout.h
+++ b/Engine/Source/Runtime/Slate/Public/Widgets/Text/SlateEditableTextLayout.h
@@ -45,6 +45,12 @@ public:
 
 	void SetSearchText(const TAttribute<FText>& InSearchText);
 	FText GetSearchText() const;
+	
+	/** Get the index of the search result (0 if none) */
+	int32 GetSearchResultIndex() const;
+	
+	/** Get the total number of search results (0 if none) */
+	int32 GetNumSearchResults() const;
 
 	void SetTextStyle(const FTextBlockStyle& InTextStyle);
 	const FTextBlockStyle& GetTextStyle() const;
@@ -81,6 +87,9 @@ public:
 	/** Get the currently selected text */
 	FText GetSelectedText() const;
 
+	/** Get the current selection*/
+	FTextSelection GetSelection() const;
+	
 	/** Set the text shaping method that the internal text layout should use */
 	void SetTextShapingMethod(const TOptional<ETextShapingMethod>& InTextShapingMethod);
 
@@ -540,6 +549,12 @@ private:
 	/** The case-sensitivity of the active search (set from BeginSearch) */
 	ESearchCase::Type SearchCase;
 
+	/** The map to look up the index of each search result (key is the starting location of each matched string)*/
+	TMap<FTextLocation, int32> SearchResultToIndexMap;
+
+	/** The active search result index */ 
+	int32 CurrentSearchResultIndex;
+
 	/** Whether text wraps onto a new line when it's length exceeds this width; if this value is zero or negative, no wrapping occurs. */
 	TAttribute<float> WrapTextAt;
 
