Converted SlateImageRun / SlateWidgetRun over to use a TAttribute<int16> for a baseline value, which allows for lambdas to be passed in. Updated all call sites that are reliant on the font service to request their baseline on demand. Fixes issues where the baseline was being computed before the font was loaded, resulting in an incorrect baseline value being used.

diff --git a/Engine/Plugins/Experimental/CommonUI/Source/CommonUI/Private/CommonRichTextBlock.cpp b/Engine/Plugins/Experimental/CommonUI/Source/CommonUI/Private/CommonRichTextBlock.cpp
--- a/Engine/Plugins/Experimental/CommonUI/Source/CommonUI/Private/CommonRichTextBlock.cpp
+++ b/Engine/Plugins/Experimental/CommonUI/Source/CommonUI/Private/CommonRichTextBlock.cpp
@@ -133,10 +133,17 @@ public:
 			ModelRange.EndIndex = InOutModelText->Len();
 
 			// Calculate the baseline of the text within the owning rich text
-			const TSharedRef<FSlateFontMeasure> FontMeasure = FSlateApplication::Get().GetRenderer()->GetFontMeasureService();
-			int16 WidgetBaseline = FontMeasure->GetBaseline(TextStyle.Font) - FMath::Min(0.0f, TextStyle.ShadowOffset.Y);
+			// Requested on demand as the font may not be loaded right now
+			const FSlateFontInfo Font = TextStyle.Font;
+			const float ShadowOffsetY = FMath::Min(0.0f, TextStyle.ShadowOffset.Y);
 
-			FSlateWidgetRun::FWidgetRunInfo WidgetRunInfo(DecoratorWidget.ToSharedRef(), WidgetBaseline);
+			TAttribute<int16> GetBaseline = TAttribute<int16>::CreateLambda([Font, ShadowOffsetY]()
+			{
+				const TSharedRef<FSlateFontMeasure> FontMeasure = FSlateApplication::Get().GetRenderer()->GetFontMeasureService();
+				return FontMeasure->GetBaseline(Font) - ShadowOffsetY;
+			});
+
+			FSlateWidgetRun::FWidgetRunInfo WidgetRunInfo(DecoratorWidget.ToSharedRef(), GetBaseline);
 			SlateRun = FSlateWidgetRun::Create(TextLayout, RunInfo, InOutModelText, WidgetRunInfo, ModelRange);
 		}
 		else
diff --git a/Engine/Source/Runtime/Slate/Private/Framework/Text/SlateImageRun.cpp b/Engine/Source/Runtime/Slate/Private/Framework/Text/SlateImageRun.cpp
--- a/Engine/Source/Runtime/Slate/Private/Framework/Text/SlateImageRun.cpp
+++ b/Engine/Source/Runtime/Slate/Private/Framework/Text/SlateImageRun.cpp
@@ -10,7 +10,7 @@
 
 #include "Framework/Text/RunUtils.h"
 
-TSharedRef< FSlateImageRun > FSlateImageRun::Create( const FRunInfo& InRunInfo, const TSharedRef< const FString >& InText, const FSlateBrush* InImage, int16 InBaseline )
+TSharedRef< FSlateImageRun > FSlateImageRun::Create( const FRunInfo& InRunInfo, const TSharedRef< const FString >& InText, const FSlateBrush* InImage, const TAttribute<int16>& InBaseline )
 {
 	if ( InImage == nullptr)
 	{
@@ -20,7 +20,7 @@ TSharedRef< FSlateImageRun > FSlateImageRun::Create( const FRunInfo& InRunInfo,
 	return MakeShareable( new FSlateImageRun( InRunInfo, InText, InImage, InBaseline ) );
 }
 
-TSharedRef< FSlateImageRun > FSlateImageRun::Create( const FRunInfo& InRunInfo, const TSharedRef< const FString >& InText, const FSlateBrush* InImage, int16 InBaseline, const FTextRange& InRange )
+TSharedRef< FSlateImageRun > FSlateImageRun::Create( const FRunInfo& InRunInfo, const TSharedRef< const FString >& InText, const FSlateBrush* InImage, const TAttribute<int16>& InBaseline, const FTextRange& InRange )
 {
 	if ( InImage == nullptr)
 	{
@@ -30,17 +30,17 @@ TSharedRef< FSlateImageRun > FSlateImageRun::Create( const FRunInfo& InRunInfo,
 	return MakeShareable( new FSlateImageRun( InRunInfo, InText, InImage, InBaseline, InRange ) );
 }
 
-TSharedRef< FSlateImageRun > FSlateImageRun::Create( const FRunInfo& InRunInfo, const TSharedRef< const FString >& InText, FName InDynamicBrushName, int16 InBaseline )
+TSharedRef< FSlateImageRun > FSlateImageRun::Create( const FRunInfo& InRunInfo, const TSharedRef< const FString >& InText, FName InDynamicBrushName, const TAttribute<int16>& InBaseline )
 {
 	return MakeShareable( new FSlateImageRun( InRunInfo, InText, InDynamicBrushName, InBaseline ) );
 }
 
-TSharedRef< FSlateImageRun > FSlateImageRun::Create( const FRunInfo& InRunInfo, const TSharedRef< const FString >& InText, FName InDynamicBrushName, int16 InBaseline, const FTextRange& InRange )
+TSharedRef< FSlateImageRun > FSlateImageRun::Create( const FRunInfo& InRunInfo, const TSharedRef< const FString >& InText, FName InDynamicBrushName, const TAttribute<int16>& InBaseline, const FTextRange& InRange )
 {
 	return MakeShareable( new FSlateImageRun( InRunInfo, InText, InDynamicBrushName, InBaseline, InRange ) );
 }
 
-FSlateImageRun::FSlateImageRun( const FRunInfo& InRunInfo, const TSharedRef< const FString >& InText, const FSlateBrush* InImage, int16 InBaseline, const FTextRange& InRange ) 
+FSlateImageRun::FSlateImageRun( const FRunInfo& InRunInfo, const TSharedRef< const FString >& InText, const FSlateBrush* InImage, const TAttribute<int16>& InBaseline, const FTextRange& InRange )
 	: RunInfo( InRunInfo )
 	, Text( InText )
 	, Range( InRange )
@@ -50,7 +50,7 @@ FSlateImageRun::FSlateImageRun( const FRunInfo& InRunInfo, const TSharedRef< con
 	check( Image != nullptr);
 }
 
-FSlateImageRun::FSlateImageRun( const FRunInfo& InRunInfo, const TSharedRef< const FString >& InText, const FSlateBrush* InImage, int16 InBaseline ) 
+FSlateImageRun::FSlateImageRun( const FRunInfo& InRunInfo, const TSharedRef< const FString >& InText, const FSlateBrush* InImage, const TAttribute<int16>& InBaseline )
 	: RunInfo( InRunInfo )
 	, Text( InText )
 	, Range( 0, Text->Len() )
@@ -60,7 +60,7 @@ FSlateImageRun::FSlateImageRun( const FRunInfo& InRunInfo, const TSharedRef< con
 	check( Image != nullptr);
 }
 
-FSlateImageRun::FSlateImageRun( const FRunInfo& InRunInfo, const TSharedRef< const FString >& InText, FName InDynamicBrushName, int16 InBaseline ) 
+FSlateImageRun::FSlateImageRun( const FRunInfo& InRunInfo, const TSharedRef< const FString >& InText, FName InDynamicBrushName, const TAttribute<int16>& InBaseline )
 	: RunInfo( InRunInfo )
 	, Text( InText )
 	, Range( 0, Text->Len() )
@@ -72,7 +72,7 @@ FSlateImageRun::FSlateImageRun( const FRunInfo& InRunInfo, const TSharedRef< con
 	Image = DynamicBrush.Get();
 }
 
-FSlateImageRun::FSlateImageRun( const FRunInfo& InRunInfo, const TSharedRef< const FString >& InText, FName InDynamicBrushName, int16 InBaseline, const FTextRange& InRange ) 
+FSlateImageRun::FSlateImageRun( const FRunInfo& InRunInfo, const TSharedRef< const FString >& InText, FName InDynamicBrushName, const TAttribute<int16>& InBaseline, const FTextRange& InRange )
 	: RunInfo( InRunInfo )
 	, Text( InText )
 	, Range( InRange )
@@ -192,7 +192,7 @@ int16 FSlateImageRun::GetMaxHeight( float Scale ) const
 
 int16 FSlateImageRun::GetBaseLine( float Scale ) const 
 {
-	return Baseline * Scale;
+	return Baseline.Get(0) * Scale;
 }
 
 FTextRange FSlateImageRun::GetTextRange() const 
diff --git a/Engine/Source/Runtime/Slate/Private/Framework/Text/SlateWidgetRun.cpp b/Engine/Source/Runtime/Slate/Private/Framework/Text/SlateWidgetRun.cpp
--- a/Engine/Source/Runtime/Slate/Private/Framework/Text/SlateWidgetRun.cpp
+++ b/Engine/Source/Runtime/Slate/Private/Framework/Text/SlateWidgetRun.cpp
@@ -30,7 +30,7 @@ void FSlateWidgetRun::SetTextRange( const FTextRange& Value )
 
 int16 FSlateWidgetRun::GetBaseLine( float Scale ) const 
 {
-	return Info.Baseline * Scale;
+	return Info.Baseline.Get(0) * Scale;
 }
 
 int16 FSlateWidgetRun::GetMaxHeight( float Scale ) const 
diff --git a/Engine/Source/Runtime/Slate/Public/Framework/Text/SlateImageRun.h b/Engine/Source/Runtime/Slate/Public/Framework/Text/SlateImageRun.h
--- a/Engine/Source/Runtime/Slate/Public/Framework/Text/SlateImageRun.h
+++ b/Engine/Source/Runtime/Slate/Public/Framework/Text/SlateImageRun.h
@@ -22,10 +22,10 @@ class SLATE_API FSlateImageRun : public ISlateRun, public TSharedFromThis< FSlat
 {
 public:
 
-	static TSharedRef< FSlateImageRun > Create( const FRunInfo& InRunInfo, const TSharedRef< const FString >& InText, const FSlateBrush* InImage, int16 InBaseline );
-	static TSharedRef< FSlateImageRun > Create( const FRunInfo& InRunInfo, const TSharedRef< const FString >& InText, const FSlateBrush* InImage, int16 InBaseline, const FTextRange& InRange );
-	static TSharedRef< FSlateImageRun > Create( const FRunInfo& InRunInfo, const TSharedRef< const FString >& InText, FName InDynamicBrushName, int16 InBaseline );
-	static TSharedRef< FSlateImageRun > Create( const FRunInfo& InRunInfo, const TSharedRef< const FString >& InText, FName InDynamicBrushName, int16 InBaseline, const FTextRange& InRange );
+	static TSharedRef< FSlateImageRun > Create( const FRunInfo& InRunInfo, const TSharedRef< const FString >& InText, const FSlateBrush* InImage, const TAttribute<int16>& InBaseline );
+	static TSharedRef< FSlateImageRun > Create( const FRunInfo& InRunInfo, const TSharedRef< const FString >& InText, const FSlateBrush* InImage, const TAttribute<int16>& InBaseline , const FTextRange& InRange );
+	static TSharedRef< FSlateImageRun > Create( const FRunInfo& InRunInfo, const TSharedRef< const FString >& InText, FName InDynamicBrushName, const TAttribute<int16>& InBaseline );
+	static TSharedRef< FSlateImageRun > Create( const FRunInfo& InRunInfo, const TSharedRef< const FString >& InText, FName InDynamicBrushName, const TAttribute<int16>& InBaseline, const FTextRange& InRange );
 
 public:
 
@@ -68,13 +68,13 @@ public:
 
 protected:
 
-	FSlateImageRun( const FRunInfo& InRunInfo, const TSharedRef< const FString >& InText, const FSlateBrush* InImage, int16 InBaseline );
+	FSlateImageRun( const FRunInfo& InRunInfo, const TSharedRef< const FString >& InText, const FSlateBrush* InImage, const TAttribute<int16>& InBaseline );
 
-	FSlateImageRun( const FRunInfo& InRunInfo, const TSharedRef< const FString >& InText, const FSlateBrush* InImage, int16 InBaseline, const FTextRange& InRange );
+	FSlateImageRun( const FRunInfo& InRunInfo, const TSharedRef< const FString >& InText, const FSlateBrush* InImage, const TAttribute<int16>& InBaseline, const FTextRange& InRange );
 
-	FSlateImageRun( const FRunInfo& InRunInfo, const TSharedRef< const FString >& InText, FName InDynamicBrushName, int16 InBaseline );
+	FSlateImageRun( const FRunInfo& InRunInfo, const TSharedRef< const FString >& InText, FName InDynamicBrushName, const TAttribute<int16>& InBaseline );
 
-	FSlateImageRun( const FRunInfo& InRunInfo, const TSharedRef< const FString >& InText, FName InDynamicBrushName, int16 InBaseline, const FTextRange& InRange );
+	FSlateImageRun( const FRunInfo& InRunInfo, const TSharedRef< const FString >& InText, FName InDynamicBrushName, const TAttribute<int16>& InBaseline, const FTextRange& InRange );
 
 private:
 
@@ -82,7 +82,7 @@ private:
 	TSharedRef< const FString > Text;
 	FTextRange Range;
 	const FSlateBrush* Image;
-	int16 Baseline;
+	TAttribute< int16 > Baseline;
 	TSharedPtr< FSlateDynamicImageBrush > DynamicBrush;
 };
 
diff --git a/Engine/Source/Runtime/Slate/Public/Framework/Text/SlateWidgetRun.h b/Engine/Source/Runtime/Slate/Public/Framework/Text/SlateWidgetRun.h
--- a/Engine/Source/Runtime/Slate/Public/Framework/Text/SlateWidgetRun.h
+++ b/Engine/Source/Runtime/Slate/Public/Framework/Text/SlateWidgetRun.h
@@ -24,10 +24,10 @@ public:
 	struct FWidgetRunInfo
 	{
 		TSharedRef< SWidget > Widget;
-		int16 Baseline;
+		TAttribute< int16 > Baseline;
 		TOptional< FVector2D > Size;
 
-		FWidgetRunInfo( const TSharedRef< SWidget >& InWidget, int16 InBaseline, TOptional< FVector2D > InSize = TOptional< FVector2D >() )
+		FWidgetRunInfo( const TSharedRef< SWidget >& InWidget, const TAttribute< int16 >& InBaseline, TOptional< FVector2D > InSize = TOptional< FVector2D >() )
 			: Widget( InWidget )
 			, Baseline( InBaseline )
 			, Size( InSize )
diff --git a/Engine/Source/Runtime/UMG/Private/Components/RichTextBlockDecorator.cpp b/Engine/Source/Runtime/UMG/Private/Components/RichTextBlockDecorator.cpp
--- a/Engine/Source/Runtime/UMG/Private/Components/RichTextBlockDecorator.cpp
+++ b/Engine/Source/Runtime/UMG/Private/Components/RichTextBlockDecorator.cpp
@@ -37,10 +37,17 @@ TSharedRef<ISlateRun> FRichTextDecorator::Create(const TSharedRef<class FTextLay
 		ModelRange.EndIndex = InOutModelText->Len();
 
 		// Calculate the baseline of the text within the owning rich text
-		const TSharedRef<FSlateFontMeasure> FontMeasure = FSlateApplication::Get().GetRenderer()->GetFontMeasureService();
-		int16 WidgetBaseline = FontMeasure->GetBaseline(TextStyle.Font) - FMath::Min(0.0f, TextStyle.ShadowOffset.Y);
+		// Requested on demand as the font may not be loaded right now
+		const FSlateFontInfo Font = TextStyle.Font;
+		const float ShadowOffsetY = FMath::Min(0.0f, TextStyle.ShadowOffset.Y);
 
-		FSlateWidgetRun::FWidgetRunInfo WidgetRunInfo(DecoratorWidget.ToSharedRef(), WidgetBaseline);
+		TAttribute<int16> GetBaseline = TAttribute<int16>::CreateLambda([Font, ShadowOffsetY]()
+		{
+			const TSharedRef<FSlateFontMeasure> FontMeasure = FSlateApplication::Get().GetRenderer()->GetFontMeasureService();
+			return FontMeasure->GetBaseline(Font) - ShadowOffsetY;
+		});
+
+		FSlateWidgetRun::FWidgetRunInfo WidgetRunInfo(DecoratorWidget.ToSharedRef(), GetBaseline);
 		SlateRun = FSlateWidgetRun::Create(TextLayout, RunInfo, InOutModelText, WidgetRunInfo, ModelRange);
 	}
 	else
