SlateEditableTextLayout: minor API addition/change

+added API to select a block of text
+added API to get a specific line of text
+changed Transaction API so that it supports merging multiple transactions together for easier undo

diff --git a/Engine/Source/Runtime/Slate/Private/Widgets/Text/SMultiLineEditableText.cpp b/Engine/Source/Runtime/Slate/Private/Widgets/Text/SMultiLineEditableText.cpp
--- a/Engine/Source/Runtime/Slate/Private/Widgets/Text/SMultiLineEditableText.cpp
+++ b/Engine/Source/Runtime/Slate/Private/Widgets/Text/SMultiLineEditableText.cpp
@@ -99,6 +99,11 @@ void SMultiLineEditableText::GetCurrentTextLine(FString& OutTextLine) const
 	EditableTextLayout->GetCurrentTextLine(OutTextLine);
 }
 
+void SMultiLineEditableText::GetTextLine(const int32 InLineIndex, FString& OutTextLine) const
+{
+	EditableTextLayout->GetTextLine(InLineIndex, OutTextLine);
+}
+
 void SMultiLineEditableText::SetText(const TAttribute< FText >& InText)
 {
 	EditableTextLayout->SetText(InText);
@@ -448,6 +453,16 @@ float SMultiLineEditableText::UpdateAndClampVerticalScrollBar(const float InView
 	return EditableTextLayout->GetScrollOffset().Y;
 }
 
+void SMultiLineEditableText::BeginEditTransaction()
+{
+	EditableTextLayout->BeginEditTransation();
+}
+
+void SMultiLineEditableText::EndEditTransaction()
+{
+	EditableTextLayout->EndEditTransaction();
+}
+
 FReply SMultiLineEditableText::OnFocusReceived( const FGeometry& MyGeometry, const FFocusEvent& InFocusEvent )
 {
 	EditableTextLayout->HandleFocusReceived(InFocusEvent);
@@ -470,6 +485,11 @@ void SMultiLineEditableText::SelectAllText()
 	EditableTextLayout->SelectAllText();
 }
 
+void SMultiLineEditableText::SelectText(const FTextLocation& InSelectionStart, const FTextLocation& InCursorLocation)
+{
+	EditableTextLayout->SelectText(InSelectionStart, InCursorLocation);
+}
+
 void SMultiLineEditableText::ClearSelection()
 {
 	EditableTextLayout->ClearSelection();
diff --git a/Engine/Source/Runtime/Slate/Private/Widgets/Text/SlateEditableTextLayout.cpp b/Engine/Source/Runtime/Slate/Private/Widgets/Text/SlateEditableTextLayout.cpp
--- a/Engine/Source/Runtime/Slate/Private/Widgets/Text/SlateEditableTextLayout.cpp
+++ b/Engine/Source/Runtime/Slate/Private/Widgets/Text/SlateEditableTextLayout.cpp
@@ -134,6 +134,7 @@ FSlateEditableTextLayout::FSlateEditableTextLayout(ISlateEditableTextWidget& InO
 	PreferredCursorScreenOffsetInLine = 0.0f;
 	SelectionStart = TOptional<FTextLocation>();
 	CurrentUndoLevel = INDEX_NONE;
+	NumTransactionsOpened = 0;
 
 	bIsDragSelecting = false;
 	bWasFocusedByLastMouseDown = false;
@@ -1826,6 +1827,28 @@ void FSlateEditableTextLayout::SelectWordAt(const FVector2D& InLocalPosition)
 	}
 }
 
+void FSlateEditableTextLayout::SelectText(const FTextLocation& InSelectionStart, const FTextLocation& InCursorLocation)
+{
+	if (TextLayout->IsEmpty())
+	{
+		return;
+	}
+
+	const FTextLocation NewCursorPosition = InCursorLocation;
+	CursorInfo.SetCursorLocationAndCalculateAlignment(*TextLayout, NewCursorPosition);
+
+	if (InSelectionStart != InCursorLocation)
+	{
+		SelectionStart = InSelectionStart;
+	}
+	else
+	{
+		SelectionStart.Reset();
+	}
+	
+	UpdateCursorHighlight();	
+}
+
 void FSlateEditableTextLayout::ClearSelection()
 {
 	SelectionStart = TOptional<FTextLocation>();
@@ -3082,11 +3105,13 @@ bool FSlateEditableTextLayout::HasTextChangedFromOriginal() const
 
 void FSlateEditableTextLayout::BeginEditTransation()
 {
-	if (StateBeforeChangingText.IsSet() || OwnerWidget->IsTextReadOnly())
+	NumTransactionsOpened += 1;
+	
+	if (NumTransactionsOpened > 1 || OwnerWidget->IsTextReadOnly())
 	{
 		// Already within a translation - don't open another
-		//Or never change text on read only controls.
-		//The TextReadOnly is an attribute, the return value may have changed since the last time it was checked.
+		// Or never change text on read only controls.
+		// The TextReadOnly is an attribute, the return value may have changed since the last time it was checked.
 		return;
 	}
 
@@ -3098,9 +3123,15 @@ void FSlateEditableTextLayout::BeginEditTransation()
 
 void FSlateEditableTextLayout::EndEditTransaction()
 {
-	if (!StateBeforeChangingText.IsSet())
+	NumTransactionsOpened -= 1;
+
+	check(NumTransactionsOpened >= 0);
+
+	if (NumTransactionsOpened > 0)
 	{
-		// No transaction to close
+		// Don't close transaction if there are more opened
+		// Caller of the first opened transaction should be
+		// responsible for actually closing the transaction
 		return;
 	}
 
@@ -3664,6 +3695,15 @@ void FSlateEditableTextLayout::GetCurrentTextLine(FString& OutTextLine) const
 	}
 }
 
+void FSlateEditableTextLayout::GetTextLine(const int32 InLineIndex, FString& OutTextLine) const
+{
+	const TArray< FTextLayout::FLineModel >& Lines = TextLayout->GetLineModels();
+	if (Lines.IsValidIndex(InLineIndex))
+	{
+		OutTextLine = *Lines[InLineIndex].Text;
+	}
+}
+
 TSharedRef<FSlateEditableTextLayout::FVirtualKeyboardEntry> FSlateEditableTextLayout::FVirtualKeyboardEntry::Create(FSlateEditableTextLayout& InOwnerLayout)
 {
 	return MakeShareable(new FVirtualKeyboardEntry(InOwnerLayout));
diff --git a/Engine/Source/Runtime/Slate/Public/Widgets/Text/SMultiLineEditableText.h b/Engine/Source/Runtime/Slate/Public/Widgets/Text/SMultiLineEditableText.h
--- a/Engine/Source/Runtime/Slate/Public/Widgets/Text/SMultiLineEditableText.h
+++ b/Engine/Source/Runtime/Slate/Public/Widgets/Text/SMultiLineEditableText.h
@@ -41,6 +41,26 @@ class FSlateEditableTextLayout;
 class SLATE_API SMultiLineEditableText : public SWidget, public ISlateEditableTextWidget
 {
 public:
+
+	/** Used to merge multiple text edit transactions within a scope */
+	struct FScopedEditableTextTransaction
+	{
+	public:
+		FScopedEditableTextTransaction(TSharedPtr<SMultiLineEditableText> InText)
+			: Text(InText)
+		{
+			Text->BeginEditTransaction();
+		}
+
+		~FScopedEditableTextTransaction()
+		{
+			Text->EndEditTransaction();	
+		};
+
+	private:
+		TSharedPtr<SMultiLineEditableText> Text;
+	};
+	
 	/** Called when the cursor is moved within the text area */
 	DECLARE_DELEGATE_OneParam( FOnCursorMoved, const FTextLocation& );
 
@@ -239,6 +259,14 @@ public:
 	 */
 	void GetCurrentTextLine(FString& OutTextLine) const;
 
+	/**
+	 * Fill OutTextLine with the text line at the specified index
+	 *
+	 * @param InLineIndex   Index of the line
+	 * @param OutTextLine   FString of the line
+	 */
+	void GetTextLine(const int32 InLineIndex, FString& OutTextLine) const;
+	
 	/**
 	 * Sets the text that appears when there is no text in the text box
 	 */
@@ -345,6 +373,9 @@ public:
 	/** Select all the text in the document */
 	void SelectAllText();
 
+	/** Select a block of text */
+	void SelectText(const FTextLocation& InSelectionStart, const FTextLocation& InCursorLocation);
+	
 	/** Clear the active text selection */
 	void ClearSelection();
 
@@ -489,6 +520,13 @@ protected:
 	virtual float UpdateAndClampVerticalScrollBar(const float InViewOffset, const float InViewFraction, const EVisibility InVisiblityOverride) override;
 	//~ End ISlateEditableTextWidget Interface
 
+protected:
+	/** Called to begin an undoable editable text transaction, marked as protected for use with FScopedEditableTextTransaction only */
+	void BeginEditTransaction();
+
+	/** Called to end an undoable editable text transaction, marked as protected for use with FScopedEditableTextTransaction only */
+	void EndEditTransaction();
+
 protected:
 	/** The text layout that deals with the editable text */
 	TUniquePtr<FSlateEditableTextLayout> EditableTextLayout;
diff --git a/Engine/Source/Runtime/Slate/Public/Widgets/Text/SlateEditableTextLayout.h b/Engine/Source/Runtime/Slate/Public/Widgets/Text/SlateEditableTextLayout.h
--- a/Engine/Source/Runtime/Slate/Public/Widgets/Text/SlateEditableTextLayout.h
+++ b/Engine/Source/Runtime/Slate/Public/Widgets/Text/SlateEditableTextLayout.h
@@ -228,6 +228,9 @@ public:
 	/** Select the word under the given position (the position is local to the text layout space) */
 	void SelectWordAt(const FVector2D& InLocalPosition);
 
+	/** Select a block of text */
+	void SelectText(const FTextLocation& InSelectionStart, const FTextLocation& InCursorLocation);
+
 	/** Clear the active text selection */
 	void ClearSelection();
 
@@ -398,6 +401,14 @@ public:
 	 */
 	void GetCurrentTextLine(FString& OutTextLine) const;
 
+	/**
+	 * Fill OutTextLine with the text line at the specified index
+	 *
+	 * @param InLineIndex   Index of the line
+	 * @param OutTextLine   FString of the line
+	 */
+	void GetTextLine(const int32 InLineIndex, FString& OutTextLine) const;
+
 private:
 	/** Insert the given text at the current cursor position, correctly taking into account new line characters */
 	void InsertTextAtCursorImpl(const FString& InString);
@@ -627,6 +638,9 @@ private:
 	/** Undo state that will be pushed if text is actually changed between calls to BeginEditTransation() and EndEditTransaction() */
 	TOptional<SlateEditableTextTypes::FUndoState> StateBeforeChangingText;
 
+	/** Track the number transactions that are opened */
+	int32 NumTransactionsOpened;
+
 	/** Original text undo state */
 	SlateEditableTextTypes::FUndoState OriginalText;
 
