DockingArea: Fix a bug where top level menus would draw over tab sometimes due to the menu bar not updating when a tab was removed

diff --git a/Engine/Source/Runtime/Slate/Private/Framework/Docking/SDockingArea.cpp b/Engine/Source/Runtime/Slate/Private/Framework/Docking/SDockingArea.cpp
--- a/Engine/Source/Runtime/Slate/Private/Framework/Docking/SDockingArea.cpp
+++ b/Engine/Source/Runtime/Slate/Private/Framework/Docking/SDockingArea.cpp
@@ -659,6 +659,9 @@ void SDockingArea::UpdateWindowChromeAndSidebar()
 				ParentWindow->GetTitleBar()->SetAllowMenuBar(bAccountForMenuBarPadding);
 			}
 		}
+
+		// Call the delegate for when the window is activated, to update the global menu bar without the user having to click on it
+		OnOwningWindowActivated();
 	}
 
 }
