SlateEditableTextLayout: Improved Scroll Behavior during search
+changed to apply minimal scroll needed to show the entirety of the matched text
+Scroll to the beginning of the line if both the line beginning and the matched text can fit into the view and if we are going from right to left

diff --git a/Engine/Source/Runtime/Slate/Private/Framework/Text/TextLayout.cpp b/Engine/Source/Runtime/Slate/Private/Framework/Text/TextLayout.cpp
--- a/Engine/Source/Runtime/Slate/Private/Framework/Text/TextLayout.cpp
+++ b/Engine/Source/Runtime/Slate/Private/Framework/Text/TextLayout.cpp
@@ -2642,6 +2642,11 @@ FVector2D FTextLayout::GetSize() const
 	return TextLayoutSize.GetDrawSize() * Inverse(Scale);
 }
 
+FVector2D FTextLayout::GetViewSize() const
+{
+	return ViewSize;
+}
+
 FVector2D FTextLayout::GetDrawSize() const
 {
 	return TextLayoutSize.GetDrawSize();
diff --git a/Engine/Source/Runtime/Slate/Private/Widgets/Text/SlateEditableTextLayout.cpp b/Engine/Source/Runtime/Slate/Private/Widgets/Text/SlateEditableTextLayout.cpp
--- a/Engine/Source/Runtime/Slate/Private/Widgets/Text/SlateEditableTextLayout.cpp
+++ b/Engine/Source/Runtime/Slate/Private/Widgets/Text/SlateEditableTextLayout.cpp
@@ -725,18 +725,30 @@ void FSlateEditableTextLayout::AdvanceSearch(const bool InReverse)
 		// Scrolling to cursor position directly does not always produce a good result
 		// because you can have only the last letter of the matched text in the view
 		// and most of the text out of view. The following code addresses this problem
-		const SlateEditableTextTypes::FScrollInfo SelectionScrollInfo(GetSelection().GetBeginning(), SlateEditableTextTypes::ECursorAlignment::Left);
-		const FVector2D LocalSelectionBeginLocation = TextLayout->GetLocationAt(SelectionScrollInfo.Position, false) / TextLayout->GetScale();
+		const FTextLocation LineStart(GetSelection().GetBeginning().GetLineIndex(), 0);
+		const FVector2D LocalLineStartLocation = TextLayout->GetLocationAt(LineStart, false) / TextLayout->GetScale();
 
-		// if we are moving towards left side, reset the scroll offset
-		// such that when the cursor is scrolled into the view during tick
-		// , we move as much of the selection into the view as possible
+		const FVector2D LocalSelectionBeginLocation = TextLayout->GetLocationAt(GetSelection().GetBeginning(), false) / TextLayout->GetScale();
+
+		const FVector2D LocalSelectionEndLocation = TextLayout->GetLocationAt(GetSelection().GetEnd(), false) / TextLayout->GetScale();
+		
+		// Only apply extra scrolling if we are going from right to left 
 		if (LocalSelectionBeginLocation.X < 0.0f)
 		{
-			ScrollOffset.X = 0;
+			const float DistanceFromSelectionEndToLineStart = LocalSelectionEndLocation.X - LocalLineStartLocation.X;
+
+			if (DistanceFromSelectionEndToLineStart < TextLayout->GetViewSize().X)
+			{
+				// Scroll to line start if both the matched text and line start can fit into the view
+				PositionToScrollIntoView = SlateEditableTextTypes::FScrollInfo(LineStart, SlateEditableTextTypes::ECursorAlignment::Left);
+			}
+			else
+			{
+				// Otherwise, just apply minimal scrolling such that the entirety of the matched text is in the view
+				PositionToScrollIntoView = SlateEditableTextTypes::FScrollInfo(GetSelection().GetBeginning(), SlateEditableTextTypes::ECursorAlignment::Left);
+			}
 		}
 	}
-
 }
 
 FVector2D FSlateEditableTextLayout::SetHorizontalScrollFraction(const float InScrollOffsetFraction)
diff --git a/Engine/Source/Runtime/Slate/Public/Framework/Text/TextLayout.h b/Engine/Source/Runtime/Slate/Public/Framework/Text/TextLayout.h
--- a/Engine/Source/Runtime/Slate/Public/Framework/Text/TextLayout.h
+++ b/Engine/Source/Runtime/Slate/Public/Framework/Text/TextLayout.h
@@ -370,6 +370,11 @@ public:
 	 */
 	FVector2D GetSize() const;
 
+	/**
+	 * Get the size of the text layout that can actually be seen from the parent widget
+	 */
+	FVector2D GetViewSize() const;
+
 	/**
 	 * Get the size of the text layout, including any lines which extend beyond the wrapping boundaries (eg, lines with lots of trailing whitespace, or lines with no break candidates)
 	 * @note This value is scaled
