SMultilineEdtiableTexT: Exposed GetCharacterAt API

diff --git a/Engine/Source/Runtime/Slate/Private/Widgets/Text/SMultiLineEditableText.cpp b/Engine/Source/Runtime/Slate/Private/Widgets/Text/SMultiLineEditableText.cpp
--- a/Engine/Source/Runtime/Slate/Private/Widgets/Text/SMultiLineEditableText.cpp
+++ b/Engine/Source/Runtime/Slate/Private/Widgets/Text/SMultiLineEditableText.cpp
@@ -570,6 +570,11 @@ FTextLocation SMultiLineEditableText::GetCursorLocation() const
 	return EditableTextLayout->GetCursorLocation();
 }
 
+TCHAR SMultiLineEditableText::GetCharacterAt(const FTextLocation& Location) const
+{
+	return EditableTextLayout->GetCharacterAt(Location);
+}
+
 TSharedPtr<const SScrollBar> SMultiLineEditableText::GetHScrollBar() const
 {
 	return HScrollBar;
diff --git a/Engine/Source/Runtime/Slate/Public/Widgets/Text/SMultiLineEditableText.h b/Engine/Source/Runtime/Slate/Public/Widgets/Text/SMultiLineEditableText.h
--- a/Engine/Source/Runtime/Slate/Public/Widgets/Text/SMultiLineEditableText.h
+++ b/Engine/Source/Runtime/Slate/Public/Widgets/Text/SMultiLineEditableText.h
@@ -422,6 +422,9 @@ public:
 	/** Get the interaction position of the cursor (where to insert, delete, etc, text from/to) */
 	FTextLocation GetCursorLocation() const;
 
+	/** Get the character at Location */
+	TCHAR GetCharacterAt(const FTextLocation& Location) const;
+
 	/** Get the horizontal scroll bar widget */
 	TSharedPtr<const SScrollBar> GetHScrollBar() const;
 
