TextLayout: Fixed a bug caused by LWC, double to float conversion caused a few conditional checks to have unexpected results

diff --git a/Engine/Source/Runtime/Slate/Private/Framework/Text/SlateTextRun.cpp b/Engine/Source/Runtime/Slate/Private/Framework/Text/SlateTextRun.cpp
--- a/Engine/Source/Runtime/Slate/Private/Framework/Text/SlateTextRun.cpp
+++ b/Engine/Source/Runtime/Slate/Private/Framework/Text/SlateTextRun.cpp
@@ -189,10 +189,10 @@ int32 FSlateTextRun::GetTextIndexAt( const TSharedRef< ILayoutBlock >& Block, co
 	const FVector2D& BlockOffset = Block->GetLocationOffset();
 	const FVector2D& BlockSize = Block->GetSize();
 
-	const float Left = BlockOffset.X;
-	const float Top = BlockOffset.Y;
-	const float Right = BlockOffset.X + BlockSize.X;
-	const float Bottom = BlockOffset.Y + BlockSize.Y;
+	const double Left = BlockOffset.X;
+	const double Top = BlockOffset.Y;
+	const double Right = BlockOffset.X + BlockSize.X;
+	const double Bottom = BlockOffset.Y + BlockSize.Y;
 
 	const bool ContainsPoint = Location.X >= Left && Location.X < Right && Location.Y >= Top && Location.Y < Bottom;
 
