SMultilineEditableText: Exposed DeleteSelectedText APi

diff --git a/Engine/Source/Runtime/Slate/Private/Widgets/Text/SMultiLineEditableText.cpp b/Engine/Source/Runtime/Slate/Private/Widgets/Text/SMultiLineEditableText.cpp
--- a/Engine/Source/Runtime/Slate/Private/Widgets/Text/SMultiLineEditableText.cpp
+++ b/Engine/Source/Runtime/Slate/Private/Widgets/Text/SMultiLineEditableText.cpp
@@ -505,6 +505,11 @@ FTextSelection SMultiLineEditableText::GetSelection() const
 	return EditableTextLayout->GetSelection();
 }
 
+void SMultiLineEditableText::DeleteSelectedText()
+{
+	EditableTextLayout->DeleteSelectedText();
+}
+
 void SMultiLineEditableText::InsertTextAtCursor(const FText& InText)
 {
 	EditableTextLayout->InsertTextAtCursor(InText.ToString());
diff --git a/Engine/Source/Runtime/Slate/Public/Widgets/Text/SMultiLineEditableText.h b/Engine/Source/Runtime/Slate/Public/Widgets/Text/SMultiLineEditableText.h
--- a/Engine/Source/Runtime/Slate/Public/Widgets/Text/SMultiLineEditableText.h
+++ b/Engine/Source/Runtime/Slate/Public/Widgets/Text/SMultiLineEditableText.h
@@ -385,6 +385,9 @@ public:
 	/** Get the current selection */
 	FTextSelection GetSelection() const;
 
+	/** Delete any currently selected text */
+	void DeleteSelectedText();
+
 	/** Insert the given text at the current cursor position, correctly taking into account new line characters */
 	void InsertTextAtCursor(const FText& InText);
 	void InsertTextAtCursor(const FString& InString);
