EditableTextLayout: Making sure insert text at cursor clears selection. In some cases we can have an empty selection active and if we don't clear it during text insertion. It will become an active selection after insertion. It is dealt with in a similar way to what is in HandleTypeChar

diff --git a/Engine/Source/Runtime/Slate/Private/Widgets/Text/SlateEditableTextLayout.cpp b/Engine/Source/Runtime/Slate/Private/Widgets/Text/SlateEditableTextLayout.cpp
--- a/Engine/Source/Runtime/Slate/Private/Widgets/Text/SlateEditableTextLayout.cpp
+++ b/Engine/Source/Runtime/Slate/Private/Widgets/Text/SlateEditableTextLayout.cpp
@@ -2076,6 +2076,7 @@ void FSlateEditableTextLayout::InsertTextAtCursorImpl(const FString& InString)
 			TextLayout->InsertAt(CursorInteractionPosition, NewLineText);
 
 			// Advance caret position
+			ClearSelection();
 			const FTextLocation NewCursorPosition = FTextLocation(CursorInteractionPosition.GetLineIndex(), FMath::Min(CursorInteractionPosition.GetOffset() + NewLineText.Len(), Line.Text->Len()));
 			CursorInfo.SetCursorLocationAndCalculateAlignment(*TextLayout, NewCursorPosition);
 		}
