Temporary fix to open Datatable.

diff --git a/Engine/Source/Runtime/Slate/Private/Widgets/Views/SHeaderRow.cpp b/Engine/Source/Runtime/Slate/Private/Widgets/Views/SHeaderRow.cpp
--- a/Engine/Source/Runtime/Slate/Private/Widgets/Views/SHeaderRow.cpp
+++ b/Engine/Source/Runtime/Slate/Private/Widgets/Views/SHeaderRow.cpp
@@ -814,7 +814,7 @@ void SHeaderRow::RegenerateWidgets()
 					const float GripSize = SplitterHandleSize > 0.0f ? SplitterHandleSize + 4.0f : 5.0f;
 					TSharedRef<SBorder> SizingGrip = SNew(SBorder)
 						.Padding(0.0f)
-						.BorderImage( &Style->ColumnSplitterStyle )
+						.BorderImage( &Style->ColumnSplitterStyle.HandleNormalBrush )
 						.Cursor(EMouseCursor::ResizeLeftRight)
 						.Content()
 						[
