Added SBox overrides for Padding() without specifying FMargin(), as is the case with all slots.

diff --git a/Engine/Source/Runtime/Slate/Public/Widgets/Layout/SBox.h b/Engine/Source/Runtime/Slate/Public/Widgets/Layout/SBox.h
--- a/Engine/Source/Runtime/Slate/Public/Widgets/Layout/SBox.h
+++ b/Engine/Source/Runtime/Slate/Public/Widgets/Layout/SBox.h
@@ -78,6 +78,24 @@ public:
 	SLATE_ATTRIBUTE(FOptionalSize, MinAspectRatio)
 
 	SLATE_ATTRIBUTE(FOptionalSize, MaxAspectRatio)
+	
+	FArguments& Padding(float Uniform)
+	{
+		_Padding = FMargin(Uniform);
+		return *this;
+	}
+
+	FArguments& Padding(float Horizontal, float Vertical)
+	{
+		_Padding = FMargin(Horizontal, Vertical);
+		return *this;
+	}
+
+	FArguments& Padding(float Left, float Top, float Right, float Bottom)
+	{
+		_Padding = FMargin(Left, Top, Right, Bottom);
+		return *this;
+	}
 
 	SLATE_END_ARGS()
 
