Updated the SListView usage example in the class documentation.

PR #8990: Comment update: SListView example (Contributed by [REDACTED_PERSON])

diff --git a/Engine/Source/Runtime/Slate/Public/Widgets/Views/SListView.h b/Engine/Source/Runtime/Slate/Public/Widgets/Views/SListView.h
--- a/Engine/Source/Runtime/Slate/Public/Widgets/Views/SListView.h
+++ b/Engine/Source/Runtime/Slate/Public/Widgets/Views/SListView.h
@@ -33,28 +33,31 @@
 /**
  * A ListView widget observes an array of data items and creates visual representations of these items.
  * ListView relies on the property that holding a reference to a value ensures its existence. In other words,
- * neither SListView<FString> nor SListView<FString*> are valid, while SListView< TSharedPtr<FString> > and
+ * neither SListView<FText> nor SListView<FText*> are valid, while SListView< TSharedPtr<FText> > and
  * SListView< UObject* > are valid.
  *
  * A trivial use case appear below:
  *
- *   Given: TArray< TSharedPtr<FString> > Items;
+ *   Given: TArray< TSharedPtr<FText> > Items;
  *
- *   SNew( SListView< TSharedPtr<FString> > )
+ *   SNew( SListView< TSharedPtr<FText> > )
  *     .ItemHeight(24)
  *     .ListItemsSource( &Items )
- *     .OnGenerateRow( SListView< TSharedPtr<FString> >::MakeOnGenerateWidget( this, &MyClass::OnGenerateRowForList ) )
+ *     .OnGenerateRow(this, &MyClass::GenerateItemRow)
  *
  * In the example we make all our widgets be 24 screen units tall. The ListView will create widgets based on data items
- * in the Items TArray. When the ListView needs to generate an item, it will do so using the OnGenerateWidgetForList method.
+ * in the Items TArray. When the ListView needs to generate an item, it will do so using the specified OnGenerateRow method.
  *
- * A sample implementation of OnGenerateWidgetForList would simply return a TextBlock with the corresponding text:
+ * A sample implementation of MyClass::GenerateItemRow has to return a STableRow with optional content:
  *
- * TSharedRef<ITableRow> OnGenerateWidgetForList( TSharedPtr<FString> InItem, const TSharedRef<STableViewBase>& OwnerTable )
+ * TSharedRef<ITableRow> MyClass::GenerateItemRow(TSharedPtr<FText> Item, const TSharedRef<STableViewBase>& OwnerTable)
  * {
- *     return SNew(STextBlock).Text( (*InItem) )
+ *     	return SNew(STableRow<TSharedPtr<FText>>, OwnerTable)
+ *		[
+ *			SNew(STextBlock)
+ *			.Text(*Item)
+ *		];
  * }
- *
  */
 template <typename ItemType>
 class SListView : public STableViewBase, TListTypeTraits<ItemType>::SerializerType, public ITypedTableView< ItemType >
