PR #8982: Update SVirtualJoystick.h (Contributed by [REDACTED_PERSON])

diff --git a/Engine/Source/Runtime/Slate/Public/Widgets/Input/SVirtualJoystick.h b/Engine/Source/Runtime/Slate/Public/Widgets/Input/SVirtualJoystick.h
--- a/Engine/Source/Runtime/Slate/Public/Widgets/Input/SVirtualJoystick.h
+++ b/Engine/Source/Runtime/Slate/Public/Widgets/Input/SVirtualJoystick.h
@@ -19,7 +19,7 @@ class SLATE_API SVirtualJoystick : public SLeafWidget
 
 public:
 	/** The settings of each zone we render */ 
-	struct FControlInfo
+	struct SLATE_API FControlInfo
 	{
 		// Set by the game
 
@@ -52,7 +52,7 @@ public:
 	};
 
 	/** The settings and current state of each zone we render */ 
-	struct FControlData
+	struct SLATE_API FControlData
 	{
 		/** Control settings */
 		FControlInfo Info;
