Crash Fix: Check if WidgetsUnderMouse is valid before accessing it in FSlateApplication::ProcessReply

diff --git a/Engine/Source/Runtime/Slate/Private/Framework/Application/SlateApplication.cpp b/Engine/Source/Runtime/Slate/Private/Framework/Application/SlateApplication.cpp
--- a/Engine/Source/Runtime/Slate/Private/Framework/Application/SlateApplication.cpp
+++ b/Engine/Source/Runtime/Slate/Private/Framework/Application/SlateApplication.cpp
@@ -3211,41 +3211,39 @@ void FSlateApplication::ProcessReply( const FWidgetPath& CurrentEventPath, const
 		{
 			if (SlateUser->SetPointerCaptor(PointerIndex, RequestedMouseCaptor.ToSharedRef(), CurrentEventPath))
 			{
-				if (WidgetsUnderMouse)
+				const FWeakWidgetPath LastWidgetsUnderCursor = SlateUser->GetLastWidgetsUnderPointer(PointerIndex);
+				if (LastWidgetsUnderCursor.IsValid())
 				{
-					const FWeakWidgetPath LastWidgetsUnderCursor = SlateUser->GetLastWidgetsUnderPointer(PointerIndex);
-					if (LastWidgetsUnderCursor.IsValid())
+					for (int32 WidgetIndex = LastWidgetsUnderCursor.Widgets.Num() - 1; WidgetIndex >= 0; --WidgetIndex)
 					{
-						for (int32 WidgetIndex = LastWidgetsUnderCursor.Widgets.Num() - 1; WidgetIndex >= 0; --WidgetIndex)
-						{
-							TSharedPtr<SWidget> WidgetPreviouslyUnderCursor = LastWidgetsUnderCursor.Widgets[WidgetIndex].Pin();
+						TSharedPtr<SWidget> WidgetPreviouslyUnderCursor = LastWidgetsUnderCursor.Widgets[WidgetIndex].Pin();
 
-							if (WidgetPreviouslyUnderCursor.IsValid())
+						if (WidgetPreviouslyUnderCursor.IsValid())
+						{
+							if (WidgetPreviouslyUnderCursor != RequestedMouseCaptor)
 							{
-								if (WidgetPreviouslyUnderCursor != RequestedMouseCaptor)
+								// It's possible for mouse event to be null if we end up here from a keyboard event. If so, we should synthesize an event.
+								// WidgetsUnderMouse can also be invalid if the mouse is not over a Slate widget.
+								if (InMouseEvent && WidgetsUnderMouse && WidgetsUnderMouse->IsValid())
 								{
-									// It's possible for mouse event to be null if we end up here from a keyboard event. If so, we should synthesize an event
-									if (InMouseEvent)
-									{
-										FPointerEvent TransformedPointerEvent = TransformPointerEvent(*InMouseEvent, WidgetsUnderMouse->GetWindow());
-
-										// Note that the event's pointer position is not translated.
-										WidgetPreviouslyUnderCursor->OnMouseLeave(TransformedPointerEvent);
-									}
-									else
-									{
-										const FPointerEvent& SimulatedPointer = FPointerEvent();
-										WidgetPreviouslyUnderCursor->OnMouseLeave(SimulatedPointer);
-									}
-#if WITH_SLATE_DEBUGGING
-									FSlateDebugging::BroadcastInputEvent(ESlateDebuggingInputEvent::MouseLeave, InMouseEvent, WidgetPreviouslyUnderCursor);
-#endif
+									FPointerEvent TransformedPointerEvent = TransformPointerEvent(*InMouseEvent, WidgetsUnderMouse->GetWindow());
+
+									// Note that the event's pointer position is not translated.
+									WidgetPreviouslyUnderCursor->OnMouseLeave(TransformedPointerEvent);
 								}
 								else
 								{
-									// Done routing mouse leave
-									break;
+									const FPointerEvent& SimulatedPointer = FPointerEvent();
+									WidgetPreviouslyUnderCursor->OnMouseLeave(SimulatedPointer);
 								}
+#if WITH_SLATE_DEBUGGING
+								FSlateDebugging::BroadcastInputEvent(ESlateDebuggingInputEvent::MouseLeave, InMouseEvent, WidgetPreviouslyUnderCursor);
+#endif
+							}
+							else
+							{
+								// Done routing mouse leave
+								break;
 							}
 						}
 					}
