Added AddVerifiedEditableText to FMultiBoxBuilder to expose SEditableTextBlock::OnVerifyTextChanged

diff --git a/Engine/Source/Runtime/Slate/Private/Framework/MultiBox/MultiBoxBuilder.cpp b/Engine/Source/Runtime/Slate/Private/Framework/MultiBox/MultiBoxBuilder.cpp
--- a/Engine/Source/Runtime/Slate/Private/Framework/MultiBox/MultiBoxBuilder.cpp
+++ b/Engine/Source/Runtime/Slate/Private/Framework/MultiBox/MultiBoxBuilder.cpp
@@ -38,6 +38,11 @@ void FMultiBoxBuilder::AddEditableText( const FText& InLabel, const FText& InToo
 	MultiBox->AddMultiBlock( MakeShareable( new FEditableTextBlock( InLabel, InToolTip, InIcon, InTextAttribute, bInReadOnly, InOnTextCommitted, InOnTextChanged ) ) );
 }
 
+void FMultiBoxBuilder::AddVerifiedEditableText(const FText& InLabel, const FText& InToolTip, const FSlateIcon& InIcon, const TAttribute< FText >& InTextAttribute, const FOnVerifyTextChanged& InOnVerifyTextChanged, const FOnTextCommitted& InOnTextCommitted /*= FOnTextCommitted()*/, const FOnTextChanged& InOnTextChanged /*= FOnTextChanged()*/, bool bInReadOnly /*= false*/)
+{
+	MultiBox->AddMultiBlock(MakeShareable(new FEditableTextBlock(InLabel, InToolTip, InIcon, InTextAttribute, bInReadOnly, InOnTextCommitted, InOnTextChanged, InOnVerifyTextChanged)));
+}
+
 void FMultiBoxBuilder::PushCommandList( const TSharedRef< const FUICommandList > CommandList )
 {
 	CommandListStack.Push( CommandList );
diff --git a/Engine/Source/Runtime/Slate/Private/Framework/MultiBox/SEditableTextBlock.cpp b/Engine/Source/Runtime/Slate/Private/Framework/MultiBox/SEditableTextBlock.cpp
--- a/Engine/Source/Runtime/Slate/Private/Framework/MultiBox/SEditableTextBlock.cpp
+++ b/Engine/Source/Runtime/Slate/Private/Framework/MultiBox/SEditableTextBlock.cpp
@@ -8,7 +8,7 @@
 #include "Widgets/Input/SEditableTextBox.h"
 
 
-FEditableTextBlock::FEditableTextBlock( const FText& InLabel, const FText& InToolTip, const FSlateIcon& InIcon, const TAttribute< FText >& InTextAttribute, bool bInReadOnly, const FOnTextCommitted& InOnTextCommitted, const FOnTextChanged& InOnTextChanged )
+FEditableTextBlock::FEditableTextBlock( const FText& InLabel, const FText& InToolTip, const FSlateIcon& InIcon, const TAttribute< FText >& InTextAttribute, bool bInReadOnly, const FOnTextCommitted& InOnTextCommitted, const FOnTextChanged& InOnTextChanged, const FOnVerifyTextChanged& InOnVerifyTextChanged )
 	: FMultiBlock( FUIAction(), NAME_None, EMultiBlockType::EditableText ),
 	  LabelOverride( InLabel ),
 	  ToolTipOverride( InToolTip ),
@@ -16,10 +16,10 @@ FEditableTextBlock::FEditableTextBlock( const FText& InLabel, const FText& InToo
 	  TextAttribute( InTextAttribute ),
 	  OnTextCommitted( InOnTextCommitted ),
 	  OnTextChanged( InOnTextChanged ),
+	  OnVerifyTextChanged( InOnVerifyTextChanged ),
 	  bReadOnly( bInReadOnly )
 { }
 
-
 bool FEditableTextBlock::HasIcon() const
 {
 	const TSharedPtr< const FUICommandInfo >& UICommand = GetAction();
@@ -137,6 +137,7 @@ void SEditableTextBlock::BuildMultiBlockWidget( const ISlateStyle* StyleSet, con
 			.RevertTextOnEscape( true )
 			.MinDesiredWidth( MultiBoxConstants::EditableTextMinWidth )
 			.OnTextChanged( EditableTextBlock->OnTextChanged )
+			.OnVerifyTextChanged( EditableTextBlock->OnVerifyTextChanged )
 			.OnTextCommitted( EditableTextBlock->OnTextCommitted )
 			.ToolTipText( TextBlockToolTip )
 		]
diff --git a/Engine/Source/Runtime/Slate/Private/Framework/MultiBox/SEditableTextBlock.h b/Engine/Source/Runtime/Slate/Private/Framework/MultiBox/SEditableTextBlock.h
--- a/Engine/Source/Runtime/Slate/Private/Framework/MultiBox/SEditableTextBlock.h
+++ b/Engine/Source/Runtime/Slate/Private/Framework/MultiBox/SEditableTextBlock.h
@@ -19,16 +19,17 @@ public:
 	/**
 	 * Constructor
 	 *
-	 * @param	InLabel				The label to display in the menu
-	 * @param	InToolTip			The tool tip to display when the menu entry is hovered over
-	 * @param	InIcon				The icon to display to the left of the label
-	 * @param	InTextAttribute		The text string we're editing (often, a delegate will be bound to the attribute)
-	 * @param	bInReadOnly			Whether or not the text block should be read only
-	 * @param	InOnTextCommitted	Called when the user commits their change to the editable text control
-	 * @param	InOnTextChanged		Called when the text is changed interactively
+	 * @param	InLabel					The label to display in the menu
+	 * @param	InToolTip				The tool tip to display when the menu entry is hovered over
+	 * @param	InIcon					The icon to display to the left of the label
+	 * @param	InTextAttribute			The text string we're editing (often, a delegate will be bound to the attribute)
+	 * @param	bInReadOnly				Whether or not the text block should be read only
+	 * @param	InOnTextCommitted		Called when the user commits their change to the editable text control
+	 * @param	InOnTextChanged			Called when the text is changed interactively
+	 * @param	InOnVerifyTextChanged	Called to verify when the text is changed interactively
 	 */
-	FEditableTextBlock( const FText& InLabel, const FText& InToolTip, const FSlateIcon& InIcon, const TAttribute< FText >& InTextAttribute, bool bInReadOnly, const FOnTextCommitted& InOnTextCommitted, const FOnTextChanged& InOnTextChanged );
-
+	FEditableTextBlock( const FText& InLabel, const FText& InToolTip, const FSlateIcon& InIcon, const TAttribute< FText >& InTextAttribute, bool bInReadOnly, const FOnTextCommitted& InOnTextCommitted, const FOnTextChanged& InOnTextChanged, const FOnVerifyTextChanged& InOnVerifyTextChanged = FOnVerifyTextChanged() );
+	
 	/** FMultiBlock interface */
 	virtual bool HasIcon() const override;
 
@@ -64,6 +65,9 @@ private:
 	/** Called when the text is changed interactively */
 	FOnTextChanged OnTextChanged;
 
+	/** Called to verify when the text is changed interactively */
+	FOnVerifyTextChanged OnVerifyTextChanged;
+
 	/** If true the text box is read only */
 	bool bReadOnly;
 };
diff --git a/Engine/Source/Runtime/Slate/Public/Framework/MultiBox/MultiBoxBuilder.h b/Engine/Source/Runtime/Slate/Public/Framework/MultiBox/MultiBoxBuilder.h
--- a/Engine/Source/Runtime/Slate/Public/Framework/MultiBox/MultiBoxBuilder.h
+++ b/Engine/Source/Runtime/Slate/Public/Framework/MultiBox/MultiBoxBuilder.h
@@ -59,6 +59,19 @@ public:
 	 */
 	void AddEditableText( const FText& InLabel, const FText& InToolTip, const FSlateIcon& InIcon, const TAttribute< FText >& InTextAttribute, const FOnTextCommitted& InOnTextCommitted = FOnTextCommitted(), const FOnTextChanged& InOnTextChanged = FOnTextChanged(), bool bInReadOnly = false );
 
+	/**
+	 * Adds an editable text entry with a VerifyTextChanged delegate
+	 *
+	 * @param	InLabel					The label to display in the menu
+	 * @param	InToolTip				The tool tip to display when the menu entry is hovered over
+	 * @param	InIcon					The icon to display to the left of the label
+	 * @param	InTextAttribute			The text string we're editing (often, a delegate will be bound to the attribute)
+	 * @param	InOnVerifyTextChanged	Called to verify when the text is changed interactively
+	 * @param	InOnTextCommitted		Called when the user commits their change to the editable text control
+	 * @param	InOnTextChanged			Called when the text is changed interactively
+	 * @param	bInReadOnly				Whether or not the text block is read only
+	 */
+	void AddVerifiedEditableText(const FText& InLabel, const FText& InToolTip, const FSlateIcon& InIcon, const TAttribute< FText >& InTextAttribute, const FOnVerifyTextChanged& InOnVerifyTextChanged, const FOnTextCommitted& InOnTextCommitted = FOnTextCommitted(), const FOnTextChanged& InOnTextChanged = FOnTextChanged(), bool bInReadOnly = false);
 
 	/**
 	 * Creates a widget for this MultiBox
