STableRow: Hide the expander arrow for pinned rows

diff --git a/Engine/Source/Runtime/Slate/Public/Widgets/Views/SListView.h b/Engine/Source/Runtime/Slate/Public/Widgets/Views/SListView.h
--- a/Engine/Source/Runtime/Slate/Public/Widgets/Views/SListView.h
+++ b/Engine/Source/Runtime/Slate/Public/Widgets/Views/SListView.h
@@ -1143,9 +1143,11 @@ private:
 			TSharedPtr<STableRow<ItemType>> PinnedRow = StaticCastSharedPtr<STableRow<ItemType>>(InPinnedItemRow);
 
 			// If the PinnedRow inherits from STableRow (i.e has a custom border already), remove it since we will be adding our own border
+			// Also set the expander arrow to Hidden so it is not available for pinned rows, but still occupies the same space
 			if (PinnedRow.IsValid())
 			{
 				PinnedRow->SetBorderImage(FAppStyle::Get().GetBrush("NoBrush"));
+				PinnedRow->SetExpanderArrowVisibility(EVisibility::Hidden);
 			}
 
 			TSharedRef<SWidget> InPinnedItemRowWidget = InPinnedItemRow->AsWidget();
diff --git a/Engine/Source/Runtime/Slate/Public/Widgets/Views/STableRow.h b/Engine/Source/Runtime/Slate/Public/Widgets/Views/STableRow.h
--- a/Engine/Source/Runtime/Slate/Public/Widgets/Views/STableRow.h
+++ b/Engine/Source/Runtime/Slate/Public/Widgets/Views/STableRow.h
@@ -193,7 +193,7 @@ public:
 				.HAlign(HAlign_Right)
 				.VAlign(VAlign_Fill)
 				[
-					SNew(SExpanderArrow, SharedThis(this) )
+					SAssignNew(ExpanderArrowWidget, SExpanderArrow, SharedThis(this) )
 					.StyleSet(ExpanderStyleSet)
 					.ShouldDrawWires(bShowWires)
 				]
@@ -1049,6 +1049,14 @@ public:
 		return FVector2D::ZeroVector;
 	}
 
+	void SetExpanderArrowVisibility(const EVisibility InExpanderArrowVisibility)
+	{
+		if(ExpanderArrowWidget)
+		{
+			ExpanderArrowWidget->SetVisibility(InExpanderArrowVisibility);
+		}
+	}
+
 	/** Protected constructor; SWidgets should only be instantiated via declarative syntax. */
 	STableRow()
 		: IndexInList(0)
@@ -1181,6 +1189,9 @@ protected:
 	/** The slate style to use with the expander */
 	const ISlateStyle* ExpanderStyleSet;
 
+	/** A pointer to the expander arrow on the row (if it exists) */
+	TSharedPtr<SExpanderArrow> ExpanderArrowWidget;
+
 	/** @see STableRow's OnCanAcceptDrop event */
 	FOnCanAcceptDrop OnCanAcceptDrop;
 
