Always fire OnTabRemoved when the tab is removed.

diff --git a/Engine/Source/Runtime/Slate/Private/Widgets/Docking/SDockTab.cpp b/Engine/Source/Runtime/Slate/Private/Widgets/Docking/SDockTab.cpp
--- a/Engine/Source/Runtime/Slate/Private/Widgets/Docking/SDockTab.cpp
+++ b/Engine/Source/Runtime/Slate/Private/Widgets/Docking/SDockTab.cpp
@@ -956,19 +956,21 @@ bool SDockTab::RequestCloseTab()
 	// The tab can be closed if the delegate is not bound or if the delegate call indicates we cannot close it
 	const bool bCanCloseTabNow = CanCloseTab();
 
-	if(bCanCloseTabNow)
+	if (bCanCloseTabNow)
 	{
-		bool bRemovedFromSidebar = false;
-		if (GetParentDockTabStack() && GetParentDockTabStack()->GetDockArea())
+		if (GetParentDockTabStack() &&
+			GetParentDockTabStack()->GetDockArea() &&
+			GetParentDockTabStack()->GetDockArea()->IsTabInSidebar(SharedThis(this)))
 		{
-			bRemovedFromSidebar = GetParentDockTabStack()->GetDockArea()->RemoveTabFromSidebar(SharedThis(this));
+			OnTabClosed.ExecuteIfBound(SharedThis(this));
+			GetParentDockTabStack()->GetDockArea()->RemoveTabFromSidebar(SharedThis(this));
 		}
-
-		if (bCanCloseTabNow && !bRemovedFromSidebar)
+		else
 		{
 			RemoveTabFromParent();
 		}
 	}
+
 	return bCanCloseTabNow;
 }
 
