PR #9068: Fix Bug: Some calculation with shadow offset doesn't take font scale ? (Contributed by [REDACTED_PERSON])

diff --git a/Engine/Source/Runtime/Slate/Private/Framework/Text/SlateTextRun.cpp b/Engine/Source/Runtime/Slate/Private/Framework/Text/SlateTextRun.cpp
--- a/Engine/Source/Runtime/Slate/Private/Framework/Text/SlateTextRun.cpp
+++ b/Engine/Source/Runtime/Slate/Private/Framework/Text/SlateTextRun.cpp
@@ -32,13 +32,13 @@ void FSlateTextRun::SetTextRange( const FTextRange& Value )
 int16 FSlateTextRun::GetBaseLine( float Scale ) const 
 {
 	const TSharedRef< FSlateFontMeasure > FontMeasure = FSlateApplication::Get().GetRenderer()->GetFontMeasureService();
-	return FontMeasure->GetBaseline( Style.Font, Scale ) - (FMath::Min(0.0f, Style.ShadowOffset.Y) + Style.Font.OutlineSettings.OutlineSize * Scale);
+	return FontMeasure->GetBaseline( Style.Font, Scale ) - ( FMath::Min(0.0f, Style.ShadowOffset.Y) + Style.Font.OutlineSettings.OutlineSize ) * Scale;
 }
 
 int16 FSlateTextRun::GetMaxHeight( float Scale ) const 
 {
 	const TSharedRef< FSlateFontMeasure > FontMeasure = FSlateApplication::Get().GetRenderer()->GetFontMeasureService();
-	return FontMeasure->GetMaxCharacterHeight( Style.Font, Scale ) + (FMath::Abs(Style.ShadowOffset.Y) + Style.Font.OutlineSettings.OutlineSize * Scale);
+	return FontMeasure->GetMaxCharacterHeight( Style.Font, Scale ) + ( FMath::Abs(Style.ShadowOffset.Y) + Style.Font.OutlineSettings.OutlineSize ) * Scale;
 }
 
 FVector2D FSlateTextRun::Measure( int32 BeginIndex, int32 EndIndex, float Scale, const FRunTextContext& TextContext ) const 
