UMG: Update the Slider widget with accessors.

diff --git a/Engine/Plugins/Experimental/CommonUI/Source/CommonUI/Private/AnalogSlider.cpp b/Engine/Plugins/Experimental/CommonUI/Source/CommonUI/Private/AnalogSlider.cpp
--- a/Engine/Plugins/Experimental/CommonUI/Source/CommonUI/Private/AnalogSlider.cpp
+++ b/Engine/Plugins/Experimental/CommonUI/Source/CommonUI/Private/AnalogSlider.cpp
@@ -12,17 +12,17 @@
 UAnalogSlider::UAnalogSlider(const FObjectInitializer& ObjectInitializer)
 	: Super(ObjectInitializer)
 {
-	Orientation = EOrientation::Orient_Horizontal;
-	SliderBarColor = FLinearColor::White;
-	SliderHandleColor = FLinearColor::White;
-	StepSize = 0.01f;
+	SetOrientation(EOrientation::Orient_Horizontal);
+	SetSliderBarColor(FLinearColor::White);
+	SetSliderHandleColor(FLinearColor::White);
+	SetStepSize(0.01f);
 	IsFocusable = true;
 }
 
 TSharedRef<SWidget> UAnalogSlider::RebuildWidget()
 {
 	MySlider = MyAnalogSlider = SNew(SAnalogSlider)
-		.Style(&WidgetStyle)
+		.Style(&GetWidgetStyle())
 		.IsFocusable(IsFocusable)
 		.OnMouseCaptureBegin(BIND_UOBJECT_DELEGATE(FSimpleDelegate, HandleOnMouseCaptureBegin))
 		.OnMouseCaptureEnd(BIND_UOBJECT_DELEGATE(FSimpleDelegate, HandleOnMouseCaptureEnd))
diff --git a/Engine/Plugins/Experimental/CommonUI/Source/CommonUI/Private/SAnalogSlider.cpp b/Engine/Plugins/Experimental/CommonUI/Source/CommonUI/Private/SAnalogSlider.cpp
--- a/Engine/Plugins/Experimental/CommonUI/Source/CommonUI/Private/SAnalogSlider.cpp
+++ b/Engine/Plugins/Experimental/CommonUI/Source/CommonUI/Private/SAnalogSlider.cpp
@@ -51,7 +51,7 @@ FReply SAnalogSlider::OnAnalogValueChanged(const FGeometry& MyGeometry, const FA
 	const float TimeToSlideFromZeroToOneAtFullSpeed = 3;
 	const float FastestStepTime = TimeToSlideFromZeroToOneAtFullSpeed * StepSizeValue;
 
-	float NewValue = ValueAttribute.Get();
+	float NewValue = GetValue();
 	bool bCommitNewValue = false;
 
 	if (Orientation == EOrientation::Orient_Horizontal && KeyPressed == EKeys::Gamepad_LeftX)
@@ -129,14 +129,14 @@ FNavigationReply SAnalogSlider::OnNavigation(const FGeometry& MyGeometry, const
 	{
 		if (InNavigationEvent.GetNavigationType() == EUINavigation::Left)
 		{
-			CommitValue(FMath::Clamp(ValueAttribute.Get() - GetStepSize(), 0.0f, 1.0f));
+			CommitValue(FMath::Clamp(GetValue() - GetStepSize(), 0.0f, 1.0f));
 			OnAnalogCapture.ExecuteIfBound(GetValue());
 
 			return FNavigationReply::Explicit(nullptr);
 		}
 		else if (InNavigationEvent.GetNavigationType() == EUINavigation::Right)
 		{
-			CommitValue(FMath::Clamp(ValueAttribute.Get() + GetStepSize(), 0.0f, 1.0f));
+			CommitValue(FMath::Clamp(GetValue() + GetStepSize(), 0.0f, 1.0f));
 			OnAnalogCapture.ExecuteIfBound(GetValue());
 
 			return FNavigationReply::Explicit(nullptr);
diff --git a/Engine/Source/Runtime/Slate/Private/Widgets/Input/SSlider.cpp b/Engine/Source/Runtime/Slate/Private/Widgets/Input/SSlider.cpp
--- a/Engine/Source/Runtime/Slate/Private/Widgets/Input/SSlider.cpp
+++ b/Engine/Source/Runtime/Slate/Private/Widgets/Input/SSlider.cpp
@@ -7,7 +7,24 @@
 #include "Widgets/Accessibility/SlateAccessibleWidgets.h"
 #endif
 
+SLATE_IMPLEMENT_WIDGET(SSlider)
+void SSlider::PrivateRegisterAttributes(FSlateAttributeInitializer& AttributeInitializer)
+{
+	SLATE_ADD_MEMBER_ATTRIBUTE_DEFINITION_WITH_NAME(AttributeInitializer, "Value", ValueSlateAttribute, EInvalidateWidgetReason::Paint);
+	SLATE_ADD_MEMBER_ATTRIBUTE_DEFINITION_WITH_NAME(AttributeInitializer, "IndentHandle", IndentHandleSlateAttribute, EInvalidateWidgetReason::Paint);
+	SLATE_ADD_MEMBER_ATTRIBUTE_DEFINITION_WITH_NAME(AttributeInitializer, "Locked", LockedSlateAttribute, EInvalidateWidgetReason::Paint);
+	SLATE_ADD_MEMBER_ATTRIBUTE_DEFINITION_WITH_NAME(AttributeInitializer, "SliderBarColor", SliderBarColorSlateAttribute, EInvalidateWidgetReason::Paint);
+	SLATE_ADD_MEMBER_ATTRIBUTE_DEFINITION_WITH_NAME(AttributeInitializer, "SliderHandleColor", SliderHandleColorSlateAttribute, EInvalidateWidgetReason::Paint);
+}
+
 SSlider::SSlider()
+	: Style(nullptr)
+	, PressedScreenSpaceTouchDownPosition(FVector2D(0, 0))
+	, ValueSlateAttribute(*this, 1.f)
+	, IndentHandleSlateAttribute(*this, true)
+	, LockedSlateAttribute(*this, false)
+	, SliderBarColorSlateAttribute(*this, FLinearColor::White)
+	, SliderHandleColorSlateAttribute(*this, FLinearColor::White)
 {
 #if WITH_ACCESSIBILITY
 	AccessibleBehavior = EAccessibleBehavior::Summary;
@@ -21,17 +38,17 @@ void SSlider::Construct( const SSlider::FArguments& InDeclaration )
 
 	Style = InDeclaration._Style;
 
-	IndentHandle = InDeclaration._IndentHandle;
+	IndentHandleSlateAttribute.Assign(*this, InDeclaration._IndentHandle);
 	bMouseUsesStep = InDeclaration._MouseUsesStep;
 	bRequiresControllerLock = InDeclaration._RequiresControllerLock;
-	LockedAttribute = InDeclaration._Locked;
+	LockedSlateAttribute.Assign(*this, InDeclaration._Locked);
 	Orientation = InDeclaration._Orientation;
 	StepSize = InDeclaration._StepSize;
-	ValueAttribute = InDeclaration._Value;
+	ValueSlateAttribute.Assign(*this, InDeclaration._Value);
 	MinValue = InDeclaration._MinValue;
 	MaxValue = InDeclaration._MaxValue;
-	SliderBarColor = InDeclaration._SliderBarColor;
-	SliderHandleColor = InDeclaration._SliderHandleColor;
+	SliderBarColorSlateAttribute.Assign(*this, InDeclaration._SliderBarColor);
+	SliderHandleColorSlateAttribute.Assign(*this, InDeclaration._SliderHandleColor);
 	bIsFocusable = InDeclaration._IsFocusable;
 	OnMouseCaptureBegin = InDeclaration._OnMouseCaptureBegin;
 	OnMouseCaptureEnd = InDeclaration._OnMouseCaptureEnd;
@@ -57,7 +74,7 @@ int32 SSlider::OnPaint( const FPaintArgs& Args, const FGeometry& AllottedGeometr
 	// calculate slider geometry as if it's a horizontal slider (we'll rotate it later if it's vertical)
 	const FVector2D HandleSize = GetThumbImage()->ImageSize;
 	const FVector2D HalfHandleSize = 0.5f * HandleSize;
-	const float Indentation = IndentHandle.Get() ? HandleSize.X : 0.0f;
+	const float Indentation = IndentHandleSlateAttribute.Get() ? HandleSize.X : 0.0f;
 
 	// We clamp to make sure that the slider cannot go out of the slider Length.
 	const float SliderPercent = FMath::Clamp(GetNormalizedValue(), 0.0f, 1.0f); 
@@ -99,7 +116,7 @@ int32 SSlider::OnPaint( const FPaintArgs& Args, const FGeometry& AllottedGeometr
 		SliderGeometry.ToPaintGeometry(BarTopLeft, BarSize),
 		BarImage,
 		DrawEffects,
-		BarImage->GetTint(InWidgetStyle) * SliderBarColor.Get().GetColor(InWidgetStyle) * InWidgetStyle.GetColorAndOpacityTint()
+		BarImage->GetTint(InWidgetStyle) * SliderBarColorSlateAttribute.Get().GetColor(InWidgetStyle) * InWidgetStyle.GetColorAndOpacityTint()
 		);
 
 	++LayerId;
@@ -111,7 +128,7 @@ int32 SSlider::OnPaint( const FPaintArgs& Args, const FGeometry& AllottedGeometr
 		SliderGeometry.ToPaintGeometry(HandleTopLeftPoint, GetThumbImage()->ImageSize),
 		ThumbImage,
 		DrawEffects,
-		ThumbImage->GetTint(InWidgetStyle) * SliderHandleColor.Get().GetColor(InWidgetStyle) * InWidgetStyle.GetColorAndOpacityTint()
+		ThumbImage->GetTint(InWidgetStyle) * SliderHandleColorSlateAttribute.Get().GetColor(InWidgetStyle) * InWidgetStyle.GetColorAndOpacityTint()
 	);
 
 	return LayerId;
@@ -137,9 +154,15 @@ FVector2D SSlider::ComputeDesiredSize( float ) const
 	return FVector2D(SSliderDesiredSize.X, Thickness);
 }
 
+void SSlider::SetStyle(const FSliderStyle* InStyle)
+{
+	Style = InStyle;
+	Invalidate(EInvalidateWidgetReason::Layout);
+}
+
 bool SSlider::IsLocked() const
 {
-	return LockedAttribute.Get();
+	return LockedSlateAttribute.Get();
 }
 
 bool SSlider::IsInteractable() const
@@ -167,7 +190,7 @@ FNavigationReply SSlider::OnNavigation(const FGeometry& MyGeometry, const FNavig
 	{
 		FNavigationReply Reply = FNavigationReply::Escape();
 
-		float NewValue = ValueAttribute.Get();
+		float NewValue = ValueSlateAttribute.Get();
 		if (Orientation == EOrientation::Orient_Horizontal)
 		{
 			if (InNavigationEvent.GetNavigationType() == EUINavigation::Left)
@@ -194,7 +217,7 @@ FNavigationReply SSlider::OnNavigation(const FGeometry& MyGeometry, const FNavig
 				Reply = FNavigationReply::Stop();
 			}
 		}
-		if (ValueAttribute.Get() != NewValue)
+		if (ValueSlateAttribute.Get() != NewValue)
 		{
 			CommitValue(FMath::Clamp(NewValue, MinValue, MaxValue));
 			return Reply;
@@ -256,7 +279,7 @@ void SSlider::OnFocusLost(const FFocusEvent& InFocusEvent)
 	if (bControllerInputCaptured)
 	{
 		// Commit and reset state
-		CommitValue(ValueAttribute.Get());
+		CommitValue(ValueSlateAttribute.Get());
 		ResetControllerState();
 	}
 }
@@ -382,9 +405,9 @@ void SSlider::CommitValue(float NewValue)
 {
 	const float OldValue = GetValue();
 
-	if (!ValueAttribute.IsBound())
+	if (!ValueSlateAttribute.IsBound(*this))
 	{
-		ValueAttribute.Set(NewValue);
+		ValueSlateAttribute.Assign(*this, NewValue);
 	}
 
 	Invalidate(EInvalidateWidgetReason::Paint);
@@ -399,7 +422,7 @@ float SSlider::PositionToValue( const FGeometry& MyGeometry, const FVector2D& Ab
 	float RelativeValue;
 	float Denominator;
 	// Only need X as we rotate the thumb image when rendering vertically
-	const float Indentation = GetThumbImage()->ImageSize.X * (IndentHandle.Get() ? 2.f : 1.f);
+	const float Indentation = GetThumbImage()->ImageSize.X * (IndentHandleSlateAttribute.Get() ? 2.f : 1.f);
 	const float HalfIndentation = 0.5f * Indentation;
 
 	if (Orientation == Orient_Horizontal)
@@ -417,18 +440,19 @@ float SSlider::PositionToValue( const FGeometry& MyGeometry, const FVector2D& Ab
 	RelativeValue = FMath::Clamp(RelativeValue, 0.0f, 1.0f) * (MaxValue - MinValue) + MinValue;
 	if (bMouseUsesStep)
 	{
-		float direction = ValueAttribute.Get() - RelativeValue;
-		if (direction > StepSize.Get() / 2.0f)
+		float direction = ValueSlateAttribute.Get() - RelativeValue;
+		float CurrentStepSize = StepSize.Get();
+		if (direction > CurrentStepSize / 2.0f)
 		{
-			return FMath::Clamp(ValueAttribute.Get() - StepSize.Get(), MinValue, MaxValue);
+			return FMath::Clamp(ValueSlateAttribute.Get() - CurrentStepSize, MinValue, MaxValue);
 		}
-		else if (direction < StepSize.Get() / -2.0f)
+		else if (direction < CurrentStepSize / -2.0f)
 		{
-			return FMath::Clamp(ValueAttribute.Get() + StepSize.Get(), MinValue, MaxValue);
+			return FMath::Clamp(ValueSlateAttribute.Get() + CurrentStepSize, MinValue, MaxValue);
 		}
 		else
 		{
-			return ValueAttribute.Get();
+			return ValueSlateAttribute.Get();
 		}
 	}
 	return RelativeValue;
@@ -436,7 +460,7 @@ float SSlider::PositionToValue( const FGeometry& MyGeometry, const FVector2D& Ab
 
 const FSlateBrush* SSlider::GetBarImage() const
 {
-	if (!IsEnabled() || LockedAttribute.Get())
+	if (!IsEnabled() || LockedSlateAttribute.Get())
 	{
 		return &Style->DisabledBarImage;
 	}
@@ -452,7 +476,7 @@ const FSlateBrush* SSlider::GetBarImage() const
 
 const FSlateBrush* SSlider::GetThumbImage() const
 {
-	if (!IsEnabled() || LockedAttribute.Get())
+	if (!IsEnabled() || LockedSlateAttribute.Get())
 	{
 		return &Style->DisabledThumbImage;
 	}
@@ -468,7 +492,7 @@ const FSlateBrush* SSlider::GetThumbImage() const
 
 float SSlider::GetValue() const
 {
-	return ValueAttribute.Get();
+	return ValueSlateAttribute.Get();
 }
 
 float SSlider::GetNormalizedValue() const
@@ -479,13 +503,13 @@ float SSlider::GetNormalizedValue() const
 	}
 	else
 	{
-		return (ValueAttribute.Get() - MinValue) / (MaxValue - MinValue);
+		return (ValueSlateAttribute.Get() - MinValue) / (MaxValue - MinValue);
 	}
 }
 
-void SSlider::SetValue(const TAttribute<float>& InValueAttribute)
+void SSlider::SetValue(TAttribute<float> InValueAttribute)
 {
-	SetAttribute(ValueAttribute, InValueAttribute, EInvalidateWidgetReason::Paint);
+	ValueSlateAttribute.Assign(*this, MoveTemp(InValueAttribute));
 }
 
 void SSlider::SetMinAndMaxValues(float InMinValue, float InMaxValue)
@@ -498,14 +522,14 @@ void SSlider::SetMinAndMaxValues(float InMinValue, float InMaxValue)
 	}
 }
 
-void SSlider::SetIndentHandle(const TAttribute<bool>& InIndentHandle)
+void SSlider::SetIndentHandle(TAttribute<bool> InIndentHandle)
 {
-	SetAttribute(IndentHandle, InIndentHandle, EInvalidateWidgetReason::Paint);
+	IndentHandleSlateAttribute.Assign(*this, MoveTemp(InIndentHandle));
 }
 
-void SSlider::SetLocked(const TAttribute<bool>& InLocked)
+void SSlider::SetLocked(TAttribute<bool> InLocked)
 {
-	SetAttribute(LockedAttribute, InLocked, EInvalidateWidgetReason::Paint);
+	LockedSlateAttribute.Assign(*this, MoveTemp(InLocked));
 }
 
 void SSlider::SetOrientation(EOrientation InOrientation)
@@ -517,14 +541,14 @@ void SSlider::SetOrientation(EOrientation InOrientation)
 	}
 }
 
-void SSlider::SetSliderBarColor(FSlateColor InSliderBarColor)
+void SSlider::SetSliderBarColor(TAttribute<FSlateColor> InSliderBarColor)
 {
-	SetAttribute(SliderBarColor, TAttribute<FSlateColor>(InSliderBarColor), EInvalidateWidgetReason::Paint);
+	SliderBarColorSlateAttribute.Assign(*this, MoveTemp(InSliderBarColor));
 }
 
-void SSlider::SetSliderHandleColor(FSlateColor InSliderHandleColor)
+void SSlider::SetSliderHandleColor(TAttribute<FSlateColor> InSliderHandleColor)
 {
-	SetAttribute(SliderHandleColor, TAttribute<FSlateColor>(InSliderHandleColor), EInvalidateWidgetReason::Paint);
+	SliderHandleColorSlateAttribute.Assign(*this, MoveTemp(InSliderHandleColor));
 }
 
 float SSlider::GetStepSize() const
@@ -532,9 +556,9 @@ float SSlider::GetStepSize() const
 	return StepSize.Get();
 }
 
-void SSlider::SetStepSize(const TAttribute<float>& InStepSize)
+void SSlider::SetStepSize(TAttribute<float> InStepSize)
 {
-	StepSize = InStepSize;
+	StepSize = MoveTemp(InStepSize);
 }
 
 void SSlider::SetMouseUsesStep(bool MouseUsesStep)
diff --git a/Engine/Source/Runtime/Slate/Public/Widgets/Input/SSlider.h b/Engine/Source/Runtime/Slate/Public/Widgets/Input/SSlider.h
--- a/Engine/Source/Runtime/Slate/Public/Widgets/Input/SSlider.h
+++ b/Engine/Source/Runtime/Slate/Public/Widgets/Input/SSlider.h
@@ -19,9 +19,10 @@ class FSlateWindowElementList;
 /**
  * A Slate slider control is a linear scale and draggable handle.
  */
-class SLATE_API SSlider
-	: public SLeafWidget
+class SLATE_API SSlider : public SLeafWidget
 {
+	SLATE_DECLARE_WIDGET(SSlider, SLeafWidget)
+
 public:
 
 	SLATE_BEGIN_ARGS(SSlider)
@@ -108,6 +109,9 @@ public:
 	 */
 	void Construct( const SSlider::FArguments& InDeclaration );
 
+	/** Set the widget style. */
+	void SetStyle(const FSliderStyle* InStyle);
+
 	/** Get the MinValue attribute */
 	float GetMinValue() const { return MinValue; }
 
@@ -121,31 +125,31 @@ public:
 	float GetNormalizedValue() const;
 
 	/** Set the Value attribute */
-	void SetValue(const TAttribute<float>& InValueAttribute);
+	void SetValue(TAttribute<float> InValueAttribute);
 
 	/** Set the MinValue and MaxValue attributes. If the new MinValue is more than the new MaxValue, MaxValue will be changed to equal MinValue. */
 	void SetMinAndMaxValues(float InMinValue, float InMaxValue);
 	
 	/** Set the IndentHandle attribute */
-	void SetIndentHandle(const TAttribute<bool>& InIndentHandle);
+	void SetIndentHandle(TAttribute<bool> InIndentHandle);
 	
 	/** Set the Locked attribute */
-	void SetLocked(const TAttribute<bool>& InLocked);
+	void SetLocked(TAttribute<bool> InLocked);
 
 	/** Set the Orientation attribute */
 	void SetOrientation(EOrientation InOrientation);
 	
 	/** Set the SliderBarColor attribute */
-	void SetSliderBarColor(FSlateColor InSliderBarColor);
+	void SetSliderBarColor(TAttribute<FSlateColor> InSliderBarColor);
 	
 	/** Set the SliderHandleColor attribute */
-	void SetSliderHandleColor(FSlateColor InSliderHandleColor);
+	void SetSliderHandleColor(TAttribute<FSlateColor> InSliderHandleColor);
 
 	/** Get the StepSize attribute */
 	float GetStepSize() const;
 
 	/** Set the StepSize attribute */
-	void SetStepSize(const TAttribute<float>& InStepSize);
+	void SetStepSize(TAttribute<float> InStepSize);
 
 	/** Set the MouseUsesStep attribute */
 	void SetMouseUsesStep(bool MouseUsesStep);
@@ -202,33 +206,60 @@ protected:
 	const FSlateBrush* GetThumbImage() const;
 
 protected:
+#if WITH_EDITORONLY_DATA
+	UE_DEPRECATED(5.1, "Direct access to Value is now deprecated. Use the setter or getter.")
+	TSlateDeprecatedTAttribute<float> ValueAttribute;
+	UE_DEPRECATED(5.1, "Direct access to IndentHandle is now deprecated. Use the setter or getter.")
+	TSlateDeprecatedTAttribute<bool> IndentHandle;
+	UE_DEPRECATED(5.1, "Direct access to LockedAttribute is now deprecated. Use the setter or getter.")
+	TSlateDeprecatedTAttribute<bool> LockedAttribute;
+	UE_DEPRECATED(5.1, "Direct access to SliderBarColor is now deprecated. Use the setter or getter.")
+	TSlateDeprecatedTAttribute<FSlateColor> SliderBarColor;
+	UE_DEPRECATED(5.1, "Direct access to SliderHandleColor is now deprecated. Use the setter or getter.")
+	TSlateDeprecatedTAttribute<FSlateColor> SliderHandleColor;
+#endif
+
+	/** @return an attribute reference of IndentHandle */
+	TSlateAttributeRef<float> GetValueAttribute() const
+	{
+		return TSlateAttributeRef<float>(SharedThis(this), ValueSlateAttribute);
+	}
+
+	/** @return an attribute reference of IndentHandle */
+	TSlateAttributeRef<bool> GetIndentHandleAttribute() const
+	{
+		return TSlateAttributeRef<bool>(SharedThis(this), IndentHandleSlateAttribute);
+	}
+
+	/** @return an attribute reference of Locked */
+	TSlateAttributeRef<bool> GetLockedAttribute() const
+	{
+		return TSlateAttributeRef<bool>(SharedThis(this), LockedSlateAttribute);
+	}
+
+	/** @return an attribute reference of SliderBarColor */
+	TSlateAttributeRef<FSlateColor> GetSliderBarColorAttribute() const
+	{
+		return TSlateAttributeRef<FSlateColor>(SharedThis(this), SliderBarColorSlateAttribute);
+	}
+
+	/** @return an attribute reference of SliderHandleColor */
+	TSlateAttributeRef<FSlateColor> GetSliderHandleColorAttribute() const
+	{
+		return TSlateAttributeRef<FSlateColor>(SharedThis(this), SliderHandleColorSlateAttribute);
+	}
 
 	// Holds the style passed to the widget upon construction.
 	const FSliderStyle* Style;
 
-	// Holds a flag indicating whether the slideable area should be indented to fit the handle.
-	TAttribute<bool> IndentHandle;
-
-	// Holds a flag indicating whether the slider is locked.
-	TAttribute<bool> LockedAttribute;
-
 	// Holds the slider's orientation.
 	EOrientation Orientation;
 
-	// Holds the color of the slider bar.
-	TAttribute<FSlateColor> SliderBarColor;
-
-	// Holds the color of the slider handle.
-	TAttribute<FSlateColor> SliderHandleColor;
-
-	// Holds the slider's current value.
-	TAttribute<float> ValueAttribute;
-
 	// Holds the initial cursor in case a custom cursor has been specified, so we can restore it after dragging the slider
 	EMouseCursor::Type CachedCursor;
 
 	/** The location in screenspace the slider was pressed by a touch */
-	FVector2D PressedScreenSpaceTouchDownPosition = FVector2D(0, 0);
+	FVector2D PressedScreenSpaceTouchDownPosition;
 
 	/** Holds the amount to adjust the value by when using a controller or keyboard */
 	TAttribute<float> StepSize;
@@ -254,6 +285,21 @@ private:
 	// Resets controller input state. Fires delegates.
 	void ResetControllerState();
 
+	// Holds the slider's current value.
+	TSlateAttribute<float> ValueSlateAttribute;
+
+	// Holds a flag indicating whether the slideable area should be indented to fit the handle.
+	TSlateAttribute<bool> IndentHandleSlateAttribute;
+
+	// Holds a flag indicating whether the slider is locked.
+	TSlateAttribute<bool> LockedSlateAttribute;
+
+	// Holds the color of the slider bar.
+	TSlateAttribute<FSlateColor> SliderBarColorSlateAttribute;
+
+	// Holds the color of the slider handle.
+	TSlateAttribute<FSlateColor> SliderHandleColorSlateAttribute;
+
 	// Holds a delegate that is executed when the mouse is pressed and a capture begins.
 	FSimpleDelegate OnMouseCaptureBegin;
 
diff --git a/Engine/Source/Runtime/UMG/Private/Components/Slider.cpp b/Engine/Source/Runtime/UMG/Private/Components/Slider.cpp
--- a/Engine/Source/Runtime/UMG/Private/Components/Slider.cpp
+++ b/Engine/Source/Runtime/UMG/Private/Components/Slider.cpp
@@ -19,6 +19,7 @@ static FSliderStyle* EditorSliderStyle = nullptr;
 USlider::USlider(const FObjectInitializer& ObjectInitializer)
 	: Super(ObjectInitializer)
 {
+PRAGMA_DISABLE_DEPRECATION_WARNINGS
 	MinValue = 0.0f;
 	MaxValue = 1.0f;
 	Orientation = EOrientation::Orient_Horizontal;
@@ -58,6 +59,8 @@ USlider::USlider(const FObjectInitializer& ObjectInitializer)
 	}
 #endif // WITH_EDITOR
 
+PRAGMA_ENABLE_DEPRECATION_WARNINGS
+
 #if WITH_EDITORONLY_DATA
 	AccessibleBehavior = ESlateAccessibleBehavior::Summary;
 	bCanChildrenBeAccessible = false;
@@ -66,6 +69,7 @@ USlider::USlider(const FObjectInitializer& ObjectInitializer)
 
 TSharedRef<SWidget> USlider::RebuildWidget()
 {
+PRAGMA_DISABLE_DEPRECATION_WARNINGS
 	MySlider = SNew(SSlider)
 		.Style(&WidgetStyle)
 		.IsFocusable(IsFocusable)
@@ -74,7 +78,7 @@ TSharedRef<SWidget> USlider::RebuildWidget()
 		.OnControllerCaptureBegin(BIND_UOBJECT_DELEGATE(FSimpleDelegate, HandleOnControllerCaptureBegin))
 		.OnControllerCaptureEnd(BIND_UOBJECT_DELEGATE(FSimpleDelegate, HandleOnControllerCaptureEnd))
 		.OnValueChanged(BIND_UOBJECT_DELEGATE(FOnFloatValueChanged, HandleOnValueChanged));
-
+PRAGMA_ENABLE_DEPRECATION_WARNINGS
 	return MySlider.ToSharedRef();
 }
 
@@ -82,6 +86,7 @@ void USlider::SynchronizeProperties()
 {
 	Super::SynchronizeProperties();
 
+PRAGMA_DISABLE_DEPRECATION_WARNINGS
 	TAttribute<float> ValueBinding = PROPERTY_BINDING(float, Value);
 	
 	MySlider->SetOrientation(Orientation);
@@ -94,6 +99,7 @@ void USlider::SynchronizeProperties()
 	MySlider->SetLocked(Locked);
 	MySlider->SetIndentHandle(IndentHandle);
 	MySlider->SetStepSize(StepSize);
+PRAGMA_ENABLE_DEPRECATION_WARNINGS
 }
 
 void USlider::ReleaseSlateResources(bool bReleaseChildren)
@@ -106,6 +112,7 @@ void USlider::ReleaseSlateResources(bool bReleaseChildren)
 void USlider::HandleOnValueChanged(float InValue)
 {
 	OnValueChanged.Broadcast(InValue);
+	BroadcastFieldValueChanged(FFieldNotificationClassDescriptor::Value);
 }
 
 void USlider::HandleOnMouseCaptureBegin()
@@ -128,6 +135,7 @@ void USlider::HandleOnControllerCaptureEnd()
 	OnControllerCaptureEnd.Broadcast();
 }
 
+PRAGMA_DISABLE_DEPRECATION_WARNINGS
 float USlider::GetValue() const
 {
 	if ( MySlider.IsValid() )
@@ -164,6 +172,15 @@ void USlider::SetValue(float InValue)
 	}
 }
 
+float USlider::GetMinValue() const
+{
+	if (MySlider.IsValid())
+	{
+		return MySlider->GetMinValue();
+	}
+	return MinValue;
+}
+
 void USlider::SetMinValue(float InValue)
 {
 	MinValue = InValue;
@@ -176,6 +193,15 @@ void USlider::SetMinValue(float InValue)
 	}
 }
 
+float USlider::GetMaxValue() const
+{
+	if (MySlider.IsValid())
+	{
+		return MySlider->GetMaxValue();
+	}
+	return MaxValue;
+}
+
 void USlider::SetMaxValue(float InValue)
 {
 	MaxValue = InValue;
@@ -185,6 +211,39 @@ void USlider::SetMaxValue(float InValue)
 	}
 }
 
+const FSliderStyle& USlider::GetWidgetStyle() const
+{
+	return WidgetStyle;
+}
+
+void USlider::SetWidgetStyle(const FSliderStyle& InStyle)
+{
+	WidgetStyle = InStyle;
+	if (MySlider.IsValid())
+	{
+		MySlider->SetStyle(&WidgetStyle);
+	}
+}
+
+EOrientation USlider::GetOrientation() const
+{
+	return Orientation;
+}
+
+void USlider::SetOrientation(EOrientation InOrientation)
+{
+	Orientation = InOrientation;
+	if (MySlider.IsValid())
+	{
+		MySlider->SetOrientation(Orientation);
+	}
+}
+
+bool USlider::HasIndentHandle() const
+{
+	return IndentHandle;
+}
+
 void USlider::SetIndentHandle(bool InIndentHandle)
 {
 	IndentHandle = InIndentHandle;
@@ -194,6 +253,11 @@ void USlider::SetIndentHandle(bool InIndentHandle)
 	}
 }
 
+bool USlider::IsLocked() const
+{
+	return Locked;
+}
+
 void USlider::SetLocked(bool InLocked)
 {
 	Locked = InLocked;
@@ -203,6 +267,15 @@ void USlider::SetLocked(bool InLocked)
 	}
 }
 
+float USlider::GetStepSize() const
+{
+	if (MySlider.IsValid())
+	{
+		return MySlider->GetStepSize();
+	}
+	return StepSize;
+}
+
 void USlider::SetStepSize(float InValue)
 {
 	StepSize = InValue;
@@ -212,6 +285,11 @@ void USlider::SetStepSize(float InValue)
 	}
 }
 
+FLinearColor USlider::GetSliderHandleColor() const
+{
+	return SliderHandleColor;
+}
+
 void USlider::SetSliderHandleColor(FLinearColor InValue)
 {
 	SliderHandleColor = InValue;
@@ -221,6 +299,11 @@ void USlider::SetSliderHandleColor(FLinearColor InValue)
 	}
 }
 
+FLinearColor USlider::GetSliderBarColor() const
+{
+	return SliderBarColor;
+}
+
 void USlider::SetSliderBarColor(FLinearColor InValue)
 {
 	SliderBarColor = InValue;
@@ -229,6 +312,7 @@ void USlider::SetSliderBarColor(FLinearColor InValue)
 		MySlider->SetSliderBarColor(InValue);
 	}
 }
+PRAGMA_ENABLE_DEPRECATION_WARNINGS
 
 #if WITH_ACCESSIBILITY
 TSharedPtr<SWidget> USlider::GetAccessibleWidget() const
diff --git a/Engine/Source/Runtime/UMG/Public/Components/Slider.h b/Engine/Source/Runtime/UMG/Public/Components/Slider.h
--- a/Engine/Source/Runtime/UMG/Public/Components/Slider.h
+++ b/Engine/Source/Runtime/UMG/Public/Components/Slider.h
@@ -29,7 +29,8 @@ class UMG_API USlider : public UWidget
 
 public:
 	/** The volume value to display. */
-	UPROPERTY(EditAnywhere, Category=Appearance, meta=(UIMin="0", UIMax="1"))
+	UE_DEPRECATED(5.1, "Direct access to Value is deprecated. Please use the getter or setter.")
+	UPROPERTY(EditAnywhere, BlueprintReadWrite, Getter, Setter, FieldNotify, BlueprintGetter="GetValue", BlueprintSetter="SetValue", Category="Appearance", meta=(UIMin="0", UIMax="1"))
 	float Value;
 
 	/** A bindable delegate to allow logic to drive the value of the widget */
@@ -37,37 +38,44 @@ public:
 	FGetFloat ValueDelegate;
 
 	/** The minimum value the slider can be set to. */
-	UPROPERTY(EditAnywhere, BlueprintReadOnly, Category = Appearance)
+	UE_DEPRECATED(5.1, "Direct access to MinValue is deprecated. Please use the getter or setter.")
+	UPROPERTY(EditAnywhere, BlueprintReadWrite, Getter, Setter, BlueprintSetter="SetMinValue", Category="Appearance")
 	float MinValue;
 
 	/** The maximum value the slider can be set to. */
-	UPROPERTY(EditAnywhere, BlueprintReadOnly, Category = Appearance)
+	UE_DEPRECATED(5.1, "Direct access to MaxValue is deprecated. Please use the getter or setter.")
+	UPROPERTY(EditAnywhere, BlueprintReadWrite, Getter, Setter, BlueprintSetter="SetMaxValue", Category="Appearance")
 	float MaxValue;
 
 public:
-	
 	/** The progress bar style */
-	UPROPERTY(EditAnywhere, BlueprintReadWrite, Category="Style", meta=( DisplayName="Style" ))
+	UE_DEPRECATED(5.1, "Direct access to WidgetStyle is deprecated. Please use the getter or setter.")
+	UPROPERTY(EditAnywhere, BlueprintReadWrite, Getter, Setter, Category="Style", meta=( DisplayName="Style" ))
 	FSliderStyle WidgetStyle;
 
 	/** The slider's orientation. */
-	UPROPERTY(EditAnywhere, BlueprintReadOnly, Category=Appearance)
+	UE_DEPRECATED(5.1, "Direct access to Orientation is deprecated. Please use the getter or setter.")
+	UPROPERTY(EditAnywhere, BlueprintReadWrite, Getter, Setter, Category=Appearance)
 	TEnumAsByte<EOrientation> Orientation;
 
 	/** The color to draw the slider bar in. */
-	UPROPERTY(EditAnywhere, BlueprintReadOnly, Category=Appearance)
+	UE_DEPRECATED(5.1, "Direct access to SliderBarColor is deprecated. Please use the getter or setter.")
+	UPROPERTY(EditAnywhere, BlueprintReadWrite, Getter, Setter, BlueprintSetter="SetSliderBarColor", Category="Appearance")
 	FLinearColor SliderBarColor;
 
 	/** The color to draw the slider handle in. */
-	UPROPERTY(EditAnywhere, BlueprintReadOnly, Category=Appearance)
+	UE_DEPRECATED(5.1, "Direct access to SliderHandleColor is deprecated. Please use the getter or setter.")
+	UPROPERTY(EditAnywhere, BlueprintReadWrite, Getter, Setter, BlueprintSetter="SetSliderHandleColor", Category="Appearance")
 	FLinearColor SliderHandleColor;
 
 	/** Whether the slidable area should be indented to fit the handle. */
-	UPROPERTY(EditAnywhere, BlueprintReadOnly, Category=Appearance, AdvancedDisplay)
+	UE_DEPRECATED(5.1, "Direct access to IndentHandle is deprecated. Please use the getter or setter.")
+	UPROPERTY(EditAnywhere, BlueprintReadWrite, Getter="HasIndentHandle", Setter, BlueprintSetter="SetIndentHandle", Category="Appearance", AdvancedDisplay)
 	bool IndentHandle;
 
 	/** Whether the handle is interactive or fixed. */
-	UPROPERTY(EditAnywhere, BlueprintReadOnly, Category=Appearance, AdvancedDisplay)
+	UE_DEPRECATED(5.1, "Direct access to Locked is deprecated. Please use the getter or setter.")
+	UPROPERTY(EditAnywhere, BlueprintReadWrite, Getter="IsLocked", Setter, BlueprintSetter="SetLocked", Category="Appearance", AdvancedDisplay)
 	bool Locked;
 
 	/** Sets new value if mouse position is greater/less than half the step size. */
@@ -79,7 +87,8 @@ public:
 	bool RequiresControllerLock;
 
 	/** The amount to adjust the value by, when using a controller or keyboard */
-	UPROPERTY(EditAnywhere, BlueprintReadOnly, Category=Appearance, meta=(UIMin="0", UIMax="1"))
+	UE_DEPRECATED(5.1, "Direct access to StepSize is deprecated. Please use the getter or setter.")
+	UPROPERTY(EditAnywhere,BlueprintReadWrite, Getter, Setter, BlueprintSetter="SetStepSize", Category="Appearance", meta=(UIMin="0", UIMax="1"))
 	float StepSize;
 
 	/** Should the slider be focusable? */
@@ -120,30 +129,63 @@ public:
 	UFUNCTION(BlueprintCallable, Category="Behavior")
 	void SetValue(float InValue);
 
+	/** Gets the minimum value of the slider. */
+	float GetMinValue() const;
+
 	/** Sets the minimum value of the slider. */
 	UFUNCTION(BlueprintCallable, Category = "Behavior")
 	void SetMinValue(float InValue);
 
+	/** Gets the maximum value of the slider. */
+	float GetMaxValue() const;
+
 	/** Sets the maximum value of the slider. */
 	UFUNCTION(BlueprintCallable, Category = "Behavior")
 	void SetMaxValue(float InValue);
 
-	/** Sets if the slidable area should be indented to fit the handle */
+	/** Get the style used by the widget. */
+	const FSliderStyle& GetWidgetStyle() const;
+
+	/** Set the style used by the widget. */
+	void SetWidgetStyle(const FSliderStyle& InStyle);
+
+	/** Getg the slider's orientation. */
+	EOrientation GetOrientation() const;
+
+	/** Sets the slider's orientation. */
+	void SetOrientation(EOrientation InOrientation);
+
+	/** Gets if the slidable area should be indented to fit the handle. */
+	bool HasIndentHandle() const;
+
+	/** Sets if the slidable area should be indented to fit the handle. */
 	UFUNCTION(BlueprintCallable, Category="Behavior")
 	void SetIndentHandle(bool InValue);
 
-	/** Sets the handle to be interactive or fixed */
+	/** Returns true when the handle is fixed. */
+	bool IsLocked() const;
+
+	/** Sets the handle to be interactive or fixed. */
 	UFUNCTION(BlueprintCallable, Category="Behavior")
 	void SetLocked(bool InValue);
 
-	/** Sets the amount to adjust the value by, when using a controller or keyboard */
+	/** Gets the amount to adjust the value by. */
+	float GetStepSize() const;
+
+	/** Sets the amount to adjust the value by, when using a controller or keyboard. */
 	UFUNCTION(BlueprintCallable, Category="Behavior")
 	void SetStepSize(float InValue);
 
-	/** Sets the color of the slider bar */
+	/** Gets the color of the slider bar. */
+	FLinearColor GetSliderBarColor() const;
+
+	/** Sets the color of the slider bar. */
 	UFUNCTION(BlueprintCallable, Category="Appearance")
 	void SetSliderBarColor(FLinearColor InValue);
 
+	/** Gets the color of the handle bar */
+	FLinearColor GetSliderHandleColor() const;
+
 	/** Sets the color of the handle bar */
 	UFUNCTION(BlueprintCallable, Category="Appearance")
 	void SetSliderHandleColor(FLinearColor InValue);
