SpinBox option to disable sliding in cases where you don't need sliding and just looking for numeric entry box

diff --git a/Engine/Source/Runtime/Slate/Public/Widgets/Input/SSpinBox.h b/Engine/Source/Runtime/Slate/Public/Widgets/Input/SSpinBox.h
--- a/Engine/Source/Runtime/Slate/Public/Widgets/Input/SSpinBox.h
+++ b/Engine/Source/Runtime/Slate/Public/Widgets/Input/SSpinBox.h
@@ -74,6 +74,7 @@ public:
 		, _MinFractionalDigits(DefaultMinFractionalDigits)
 		, _MaxFractionalDigits(DefaultMaxFractionalDigits)
 		, _AlwaysUsesDeltaSnap(false)
+		, _EnableSlider(true)
 		, _Delta(0)
 		, _ShiftMouseMovePixelPerDelta(1)
 		, _SupportDynamicSliderMaxValue(false)
@@ -107,6 +108,8 @@ public:
 		SLATE_ATTRIBUTE(TOptional< int32 >, MaxFractionalDigits)
 		/** Whether typed values should use delta snapping, defaults to false */
 		SLATE_ATTRIBUTE(bool, AlwaysUsesDeltaSnap)
+		/** Whether this spin box should have slider feature enabled, defaults to true */
+		SLATE_ATTRIBUTE(bool, EnableSlider)
 		/** Delta to increment the value as the slider moves.  If not specified will determine automatically */
 		SLATE_ATTRIBUTE(NumericType, Delta)
 		/** How many pixel the mouse must move to change the value of the delta step */
@@ -190,6 +193,7 @@ public:
 		SetMinFractionalDigits(MinFractionalDigits);
 
 		AlwaysUsesDeltaSnap = InArgs._AlwaysUsesDeltaSnap;
+		EnableSlider = InArgs._EnableSlider;
 
 		SupportDynamicSliderMaxValue = InArgs._SupportDynamicSliderMaxValue;
 		SupportDynamicSliderMinValue = InArgs._SupportDynamicSliderMinValue;
@@ -511,7 +515,8 @@ public:
 	 */
 	virtual FReply OnMouseMove(const FGeometry& MyGeometry, const FPointerEvent& MouseEvent) override
 	{
-		if (PointerDraggingSliderIndex == MouseEvent.GetPointerIndex())
+		const bool bEnableSlider = GetEnableSlider();
+		if (PointerDraggingSliderIndex == MouseEvent.GetPointerIndex() && bEnableSlider)
 		{
 			if (!this->HasMouseCapture())
 			{
@@ -646,6 +651,13 @@ public:
 
 	virtual FCursorReply OnCursorQuery(const FGeometry& MyGeometry, const FPointerEvent& CursorEvent) const override
 	{
+		const bool bEnableSlider = GetEnableSlider();
+
+		if (!bEnableSlider)
+		{
+			return FCursorReply::Cursor(EMouseCursor::Default);
+		}
+
 		return bDragging ?
 			FCursorReply::Cursor(EMouseCursor::None) :
 			FCursorReply::Cursor(EMouseCursor::ResizeLeftRight);
@@ -785,6 +797,10 @@ public:
 	bool GetAlwaysUsesDeltaSnap() const { return AlwaysUsesDeltaSnap.Get(); }
 	void SetAlwaysUsesDeltaSnap(bool bNewValue) { AlwaysUsesDeltaSnap.Set(bNewValue); }
 
+	/** See the EnableSlider attribute */
+	bool GetEnableSlider() const { return EnableSlider.Get(); }
+	void SetEnableSlider(bool bNewValue) { EnableSlider.Set(bNewValue); }
+
 	/** See the Delta attribute */
 	NumericType GetDelta() const { return Delta.Get(); }
 	void SetDelta(NumericType InDelta) { Delta.Set(InDelta); }
@@ -1057,6 +1073,7 @@ private:
 	TAttribute< TOptional<int32> > MinFractionalDigits;
 	TAttribute< TOptional<int32> > MaxFractionalDigits;
 	TAttribute<bool> AlwaysUsesDeltaSnap;
+	TAttribute<bool> EnableSlider;
 	TAttribute<bool> SupportDynamicSliderMaxValue;
 	TAttribute<bool> SupportDynamicSliderMinValue;
 	FOnDynamicSliderMinMaxValueChanged OnDynamicSliderMaxValueChanged;
diff --git a/Engine/Source/Runtime/UMG/Private/Components/SpinBox.cpp b/Engine/Source/Runtime/UMG/Private/Components/SpinBox.cpp
--- a/Engine/Source/Runtime/UMG/Private/Components/SpinBox.cpp
+++ b/Engine/Source/Runtime/UMG/Private/Components/SpinBox.cpp
@@ -33,6 +33,7 @@ USpinBox::USpinBox(const FObjectInitializer& ObjectInitializer)
 	MinFractionalDigits = 1;
 	MaxFractionalDigits = 6;
 	bAlwaysUsesDeltaSnap = false;
+	bEnableSlider = true;
 	Delta = 0;
 	SliderExponent = 1;
 	MinDesiredWidth = 0;
@@ -107,6 +108,7 @@ void USpinBox::SynchronizeProperties()
 	MySpinBox->SetMinFractionalDigits(MinFractionalDigits);
 	MySpinBox->SetMaxFractionalDigits(MaxFractionalDigits);
 	MySpinBox->SetAlwaysUsesDeltaSnap(bAlwaysUsesDeltaSnap);
+	MySpinBox->SetEnableSlider(bEnableSlider);
 
 	// Set optional values
 	bOverride_MinValue ? SetMinValue(MinValue) : ClearMinValue();
diff --git a/Engine/Source/Runtime/UMG/Public/Components/SpinBox.h b/Engine/Source/Runtime/UMG/Public/Components/SpinBox.h
--- a/Engine/Source/Runtime/UMG/Public/Components/SpinBox.h
+++ b/Engine/Source/Runtime/UMG/Public/Components/SpinBox.h
@@ -52,6 +52,10 @@ public:
 	UPROPERTY(EditAnywhere, BlueprintReadWrite, BlueprintSetter = SetAlwaysUsesDeltaSnap, BlueprintGetter = GetAlwaysUsesDeltaSnap, Category = "Slider")
 	bool bAlwaysUsesDeltaSnap;
 
+	/** Whether this spin box should have slider feature enabled */
+	UPROPERTY(EditAnywhere, BlueprintReadWrite, Category = "Slider")
+	bool bEnableSlider = true;
+
 	/** The amount by which to change the spin box value as the slider moves. */
 	UPROPERTY(EditAnywhere, BlueprintReadWrite, BlueprintSetter = SetDelta, BlueprintGetter = GetDelta, Category = "Slider")
 	float Delta;
