Exposed SBorder's protected methods GetBorderImage() and GetBorderBackgroundColor().

Added doc comments for other getters.

Requested in UDN case 00430145.

diff --git a/Engine/Source/Runtime/Slate/Public/Widgets/Layout/SBorder.h b/Engine/Source/Runtime/Slate/Public/Widgets/Layout/SBorder.h
--- a/Engine/Source/Runtime/Slate/Public/Widgets/Layout/SBorder.h
+++ b/Engine/Source/Runtime/Slate/Public/Widgets/Layout/SBorder.h
@@ -108,6 +108,9 @@ public:
 	/** Sets the color and opacity of the background image of this border. */
 	void SetBorderBackgroundColor(TAttribute<FSlateColor> InColorAndOpacity);
 
+	/** Gets the color and opacity of the background image of this border. */
+	FSlateColor GetBorderBackgroundColor() const { return BorderBackgroundColorAttribute.Get(); }
+
 	/** Set the desired size scale multiplier */
 	void SetDesiredSizeScale(TAttribute<FVector2D> InDesiredSizeScale);
 	
@@ -120,12 +123,15 @@ public:
 	/** See Padding attribute */
 	void SetPadding(TAttribute<FMargin> InPadding);
 
-	/** See ShowEffectWhenDisabled attribute */
+	/** Set whether or not to show the disabled effect when this border is disabled */
 	void SetShowEffectWhenDisabled(TAttribute<bool> InShowEffectWhenDisabled);
 
-	/** See BorderImage attribute */
+	/** Set the image to draw for this border. */
 	void SetBorderImage(TAttribute<const FSlateBrush*> InBorderImage);
 
+	/** Get the image to draw for this border. */
+	const FSlateBrush* GetBorderImage() const { return BorderImageAttribute.Get(); }
+
 public:
 	// SWidget interface
 	virtual int32 OnPaint( const FPaintArgs& Args, const FGeometry& AllottedGeometry, const FSlateRect& MyCullingRect, FSlateWindowElementList& OutDrawElements, int32 LayerId, const FWidgetStyle& InWidgetStyle, bool bParentEnabled ) const override;
@@ -136,10 +142,12 @@ protected:
 	virtual FVector2D ComputeDesiredSize(float) const override;
 	//~End SWidget overrides.
 
-	const FSlateBrush* GetBorderImage() const { return BorderImageAttribute.Get(); }
-	FSlateColor GetBorderBackgroundColor() const { return BorderBackgroundColorAttribute.Get(); }
-	FVector2D GetDesiredSizeScale() const { return DesiredSizeScaleAttribute.Get(); }
+	/** Get whether or not to show the disabled effect when this border is disabled */
 	bool GetShowDisabledEffect() const { return ShowDisabledEffectAttribute.Get(); }
+
+	/** Get the desired size scale multiplier */
+	FVector2D GetDesiredSizeScale() const { return DesiredSizeScaleAttribute.Get(); }
+
 	TSlateAttributeRef<const FSlateBrush*> GetBorderImageAttribute() const { return TSlateAttributeRef<const FSlateBrush*>(SharedThis(this), BorderImageAttribute); }
 	TSlateAttributeRef<FSlateColor> GetBorderBackgroundColorAttribute() const { return TSlateAttributeRef<FSlateColor>(SharedThis(this), BorderBackgroundColorAttribute); }
 	TSlateAttributeRef<FVector2D> GetDesiredSizeScaleAttribute() const { return TSlateAttributeRef<FVector2D>(SharedThis(this), DesiredSizeScaleAttribute); }
