STableViewBase:  fix scroll wheel failing to be consumed bug
- Issue: When ConsumeMouseWheel is set to WhenScrollingPossible, input is not properly consumed when scrolling up. We are checking NewScrollOffset >= DesiredScrollOffset right after we set DesiredScrollOffset to NewScrollOffset so it is always true.
- Solution: Do this check before we set the DesiredScrollOffset.  I believe this logic is just predicting if we are at the end of the list and trying to scroll further we aren't actually going to go anywhere.

diff --git a/Engine/Source/Runtime/Slate/Private/Widgets/Views/STableViewBase.cpp b/Engine/Source/Runtime/Slate/Private/Widgets/Views/STableViewBase.cpp
--- a/Engine/Source/Runtime/Slate/Private/Widgets/Views/STableViewBase.cpp
+++ b/Engine/Source/Runtime/Slate/Private/Widgets/Views/STableViewBase.cpp
@@ -819,13 +819,13 @@ float STableViewBase::ScrollTo( float InScrollOffset)
 	const float NewScrollOffset = FMath::Clamp( InScrollOffset, -10.0f, GetNumItemsBeingObserved()+10.0f );
 	float AmountScrolled = FMath::Abs( DesiredScrollOffset - NewScrollOffset );
 
-	SetScrollOffset( NewScrollOffset );
-	
-	if ( bWasAtEndOfList && NewScrollOffset >= DesiredScrollOffset )
+	if (bWasAtEndOfList && NewScrollOffset >= DesiredScrollOffset)
 	{
 		AmountScrolled = 0;
 	}
 
+	SetScrollOffset( NewScrollOffset );
+
 	return AmountScrolled;
 }
 
