Fix text fields not supporting IME when programmatically focused

diff --git a/Engine/Source/Runtime/Slate/Private/Widgets/Text/SlateEditableTextLayout.cpp b/Engine/Source/Runtime/Slate/Private/Widgets/Text/SlateEditableTextLayout.cpp
--- a/Engine/Source/Runtime/Slate/Private/Widgets/Text/SlateEditableTextLayout.cpp
+++ b/Engine/Source/Runtime/Slate/Private/Widgets/Text/SlateEditableTextLayout.cpp
@@ -826,6 +826,13 @@ bool FSlateEditableTextLayout::HandleFocusReceived(const FFocusEvent& InFocusEve
 			}
 
 			TextInputMethodContext->CacheWindow();
+
+			// Make sure to set Native OS window focus as well to ensure IME support
+			if (TSharedPtr<FGenericWindow> NativeWindow = TextInputMethodContext->GetWindow())
+			{
+				NativeWindow->SetWindowFocus();
+			}
+
 			TextInputMethodSystem->ActivateContext(TextInputMethodContext.ToSharedRef());
 		}
 	}
