Adding a bool to prevent throttling in the scroll bar. Fixes an issue where list view entries didn't update when scrolling while holding scroll bar.

diff --git a/Engine/Source/Runtime/Slate/Private/Widgets/Layout/SScrollBar.cpp b/Engine/Source/Runtime/Slate/Private/Widgets/Layout/SScrollBar.cpp
--- a/Engine/Source/Runtime/Slate/Private/Widgets/Layout/SScrollBar.cpp
+++ b/Engine/Source/Runtime/Slate/Private/Widgets/Layout/SScrollBar.cpp
@@ -24,6 +24,7 @@ void SScrollBar::Construct(const FArguments& InArgs)
 	EVerticalAlignment VerticalAlignment = Orientation == Orient_Vertical ? VAlign_Fill : VAlign_Center;
 
 	bHideWhenNotInUse = InArgs._HideWhenNotInUse;
+	bPreventThrottling = InArgs._PreventThrottling;
 	bIsScrolling = false;
 	LastInteractionTime = 0;
 
@@ -146,7 +147,14 @@ FReply SScrollBar::OnMouseButtonDown( const FGeometry& MyGeometry, const FPointe
 
 	if( bDraggingThumb )
 	{
-		return FReply::Handled().CaptureMouse(AsShared()).SetUserFocus(AsShared(), DragFocusCause);
+		FReply ReturnReply = FReply::Handled().CaptureMouse(AsShared()).SetUserFocus(AsShared(), DragFocusCause);
+
+		if (bPreventThrottling)
+		{
+			ReturnReply.PreventThrottling();
+		}
+
+		return ReturnReply;
 	}
 	else
 	{
diff --git a/Engine/Source/Runtime/Slate/Private/Widgets/Views/STableViewBase.cpp b/Engine/Source/Runtime/Slate/Private/Widgets/Views/STableViewBase.cpp
--- a/Engine/Source/Runtime/Slate/Private/Widgets/Views/STableViewBase.cpp
+++ b/Engine/Source/Runtime/Slate/Private/Widgets/Views/STableViewBase.cpp
@@ -50,7 +50,7 @@ FTableViewDimensions::FTableViewDimensions(EOrientation InOrientation, const FVe
 	}
 }
 
-void STableViewBase::ConstructChildren( const TAttribute<float>& InItemWidth, const TAttribute<float>& InItemHeight, const TAttribute<EListItemAlignment>& InItemAlignment, const TSharedPtr<SHeaderRow>& InHeaderRow, const TSharedPtr<SScrollBar>& InScrollBar, EOrientation InScrollOrientation, const FOnTableViewScrolled& InOnTableViewScrolled, const FScrollBarStyle* InScrollBarStyle )
+void STableViewBase::ConstructChildren( const TAttribute<float>& InItemWidth, const TAttribute<float>& InItemHeight, const TAttribute<EListItemAlignment>& InItemAlignment, const TSharedPtr<SHeaderRow>& InHeaderRow, const TSharedPtr<SScrollBar>& InScrollBar, EOrientation InScrollOrientation, const FOnTableViewScrolled& InOnTableViewScrolled, const FScrollBarStyle* InScrollBarStyle, const bool bInPreventThrottling )
 {
 	bItemsNeedRefresh = true;
 	
@@ -92,7 +92,8 @@ void STableViewBase::ConstructChildren( const TAttribute<float>& InItemWidth, co
 		ScrollBar = SNew(SScrollBar)
 			.OnUserScrolled(this, &STableViewBase::ScrollBar_OnUserScrolled)
 			.Orientation(Orientation)
-			.Style(InScrollBarStyle ? InScrollBarStyle : &FAppStyle::Get().GetWidgetStyle<FScrollBarStyle>("ScrollBar"));
+			.Style(InScrollBarStyle ? InScrollBarStyle : &FAppStyle::Get().GetWidgetStyle<FScrollBarStyle>("ScrollBar"))
+			.PreventThrottling(bInPreventThrottling);
 
 		const FOptionalSize ScrollBarSize(16.f);
 
diff --git a/Engine/Source/Runtime/Slate/Public/Widgets/Layout/SScrollBar.h b/Engine/Source/Runtime/Slate/Public/Widgets/Layout/SScrollBar.h
--- a/Engine/Source/Runtime/Slate/Public/Widgets/Layout/SScrollBar.h
+++ b/Engine/Source/Runtime/Slate/Public/Widgets/Layout/SScrollBar.h
@@ -38,6 +38,7 @@ public:
 #else
 		, _HideWhenNotInUse(false)
 #endif
+		, _PreventThrottling(false)
 		, _Orientation( Orient_Vertical )
 		, _DragFocusCause( EFocusCause::Mouse )
 		, _Thickness()
@@ -50,6 +51,7 @@ public:
 		SLATE_ARGUMENT( bool, AlwaysShowScrollbar )
 		SLATE_ARGUMENT( bool, AlwaysShowScrollbarTrack )
 		SLATE_ARGUMENT( bool, HideWhenNotInUse )
+		SLATE_ARGUMENT( bool, PreventThrottling )
 		SLATE_ARGUMENT( EOrientation, Orientation )
 		SLATE_ARGUMENT( EFocusCause, DragFocusCause )
 		/** The thickness of the scrollbar thumb */
@@ -168,6 +170,11 @@ protected:
 	bool bAlwaysShowScrollbarTrack;
 	EFocusCause DragFocusCause;
 	bool bHideWhenNotInUse;
+	/*
+	 * Holds whether or not to prevent throttling during mouse capture
+	 * When true, the viewport will be updated with every single change to the value during dragging
+	 */
+	bool bPreventThrottling;
 	bool bIsScrolling;
 	double LastInteractionTime;
 
diff --git a/Engine/Source/Runtime/Slate/Public/Widgets/Views/STableViewBase.h b/Engine/Source/Runtime/Slate/Public/Widgets/Views/STableViewBase.h
--- a/Engine/Source/Runtime/Slate/Public/Widgets/Views/STableViewBase.h
+++ b/Engine/Source/Runtime/Slate/Public/Widgets/Views/STableViewBase.h
@@ -93,7 +93,7 @@ class SLATE_API STableViewBase
 public:
 
 	/** Create the child widgets that comprise the list */
-	void ConstructChildren( const TAttribute<float>& InItemWidth, const TAttribute<float>& InItemHeight, const TAttribute<EListItemAlignment>& InItemAlignment, const TSharedPtr<SHeaderRow>& InHeaderRow, const TSharedPtr<SScrollBar>& InScrollBar, EOrientation InScrollOrientation, const FOnTableViewScrolled& InOnTableViewScrolled, const FScrollBarStyle* InScrollBarStyle = nullptr );
+	void ConstructChildren( const TAttribute<float>& InItemWidth, const TAttribute<float>& InItemHeight, const TAttribute<EListItemAlignment>& InItemAlignment, const TSharedPtr<SHeaderRow>& InHeaderRow, const TSharedPtr<SScrollBar>& InScrollBar, EOrientation InScrollOrientation, const FOnTableViewScrolled& InOnTableViewScrolled, const FScrollBarStyle* InScrollBarStyle = nullptr, const bool bInPreventThrottling = false );
 
 	/** Sets the item height */
 	void SetItemHeight(TAttribute<float> Height);
diff --git a/Engine/Source/Runtime/Slate/Public/Widgets/Views/STreeView.h b/Engine/Source/Runtime/Slate/Public/Widgets/Views/STreeView.h
--- a/Engine/Source/Runtime/Slate/Public/Widgets/Views/STreeView.h
+++ b/Engine/Source/Runtime/Slate/Public/Widgets/Views/STreeView.h
@@ -136,6 +136,7 @@ public:
 		, _ConsumeMouseWheel( EConsumeMouseWheel::WhenScrollingPossible )
 		, _AllowOverscroll(EAllowOverscroll::Yes)
 		, _ScrollBarStyle(&FAppStyle::Get().GetWidgetStyle<FScrollBarStyle>("ScrollBar"))
+		, _PreventThrottling(false)
 		, _WheelScrollMultiplier(GetGlobalScrollAmount())
 		, _OnItemToString_Debug()
 		, _OnEnteredBadState()
@@ -203,6 +204,8 @@ public:
 		
 		SLATE_STYLE_ARGUMENT( FScrollBarStyle, ScrollBarStyle );
 
+		SLATE_ARGUMENT( bool, PreventThrottling )
+
 		SLATE_ARGUMENT( float, WheelScrollMultiplier );
 
 		/** Assign this to get more diagnostics from the list view. */
@@ -312,7 +315,7 @@ public:
 		else
 		{
 			// Make the TableView
-			this->ConstructChildren( 0, InArgs._ItemHeight, EListItemAlignment::LeftAligned, InArgs._HeaderRow, InArgs._ExternalScrollbar, Orient_Vertical, InArgs._OnTreeViewScrolled, InArgs._ScrollBarStyle );
+			this->ConstructChildren( 0, InArgs._ItemHeight, EListItemAlignment::LeftAligned, InArgs._HeaderRow, InArgs._ExternalScrollbar, Orient_Vertical, InArgs._OnTreeViewScrolled, InArgs._ScrollBarStyle, InArgs._PreventThrottling );
 			if (this->ScrollBar.IsValid())
 			{
 				this->ScrollBar->SetDragFocusCause(InArgs._ScrollbarDragFocusCause);
diff --git a/Engine/Source/Runtime/UMG/Public/Components/ListViewBase.h b/Engine/Source/Runtime/UMG/Public/Components/ListViewBase.h
--- a/Engine/Source/Runtime/UMG/Public/Components/ListViewBase.h
+++ b/Engine/Source/Runtime/UMG/Public/Components/ListViewBase.h
@@ -236,6 +236,7 @@ protected:
 		EOrientation Orientation = Orient_Vertical;
 		const FTableViewStyle* ListViewStyle = &FUMGCoreStyle::Get().GetWidgetStyle<FTableViewStyle>("ListView");
 		const FScrollBarStyle* ScrollBarStyle = &FUMGCoreStyle::Get().GetWidgetStyle<FScrollBarStyle>("ScrollBar");
+		bool bPreventThrottling = false;
 	};
 
 	template <template<typename> class ListViewT = SListView, typename UListViewBaseT>
@@ -255,6 +256,7 @@ protected:
 			.Orientation(Args.Orientation)
 			.ListViewStyle(Args.ListViewStyle)
 			.ScrollBarStyle(Args.ScrollBarStyle)
+			.PreventThrottling(Args.bPreventThrottling)
 			.OnGenerateRow_UObject(Implementer, &UListViewBaseT::HandleGenerateRow)
 			.OnSelectionChanged_UObject(Implementer, &UListViewBaseT::HandleSelectionChanged)
 			.OnIsSelectableOrNavigable_UObject(Implementer, &UListViewBaseT::HandleIsSelectableOrNavigable)
