Fix CIS Compile Error: ListViewBase.h(259,5): error: no member named 'PreventThrottling' due to missing file from 20062084

diff --git a/Engine/Source/Runtime/Slate/Public/Widgets/Views/SListView.h b/Engine/Source/Runtime/Slate/Public/Widgets/Views/SListView.h
--- a/Engine/Source/Runtime/Slate/Public/Widgets/Views/SListView.h
+++ b/Engine/Source/Runtime/Slate/Public/Widgets/Views/SListView.h
@@ -105,6 +105,7 @@ public:
 		, _ScrollbarDragFocusCause(EFocusCause::Mouse)
 		, _AllowOverscroll(EAllowOverscroll::Yes)
 		, _ScrollBarStyle(&FAppStyle::Get().GetWidgetStyle<FScrollBarStyle>("ScrollBar"))
+		, _PreventThrottling(false)
 		, _ConsumeMouseWheel(EConsumeMouseWheel::WhenScrollingPossible)
 		, _WheelScrollMultiplier(GetGlobalScrollAmount())
 		, _NavigationScrollOffset(0.5f)
@@ -173,6 +174,8 @@ public:
 		
 		SLATE_STYLE_ARGUMENT( FScrollBarStyle, ScrollBarStyle );
 
+		SLATE_ARGUMENT(bool, PreventThrottling);
+
 		SLATE_ARGUMENT( EConsumeMouseWheel, ConsumeMouseWheel );
 
 		SLATE_ARGUMENT( float, WheelScrollMultiplier );
@@ -281,7 +284,7 @@ public:
 		else
 		{
 			// Make the TableView
-			ConstructChildren( 0, InArgs._ItemHeight, EListItemAlignment::LeftAligned, InArgs._HeaderRow, InArgs._ExternalScrollbar, InArgs._Orientation, InArgs._OnListViewScrolled, InArgs._ScrollBarStyle );
+			ConstructChildren( 0, InArgs._ItemHeight, EListItemAlignment::LeftAligned, InArgs._HeaderRow, InArgs._ExternalScrollbar, InArgs._Orientation, InArgs._OnListViewScrolled, InArgs._ScrollBarStyle, InArgs._PreventThrottling );
 			if(ScrollBar.IsValid())
 			{
 				ScrollBar->SetDragFocusCause(InArgs._ScrollbarDragFocusCause);
