FAsyncTaskNotification: Refactor to remove SSlateAsyncTaskNotificationWidget and use different SNotificationItem notifications depending on the state

diff --git a/Engine/Source/Runtime/Core/Private/Misc/AsyncTaskNotification.cpp b/Engine/Source/Runtime/Core/Private/Misc/AsyncTaskNotification.cpp
--- a/Engine/Source/Runtime/Core/Private/Misc/AsyncTaskNotification.cpp
+++ b/Engine/Source/Runtime/Core/Private/Misc/AsyncTaskNotification.cpp
@@ -12,22 +12,21 @@ FAsyncTaskNotification::FAsyncTaskNotification(const FAsyncTaskNotificationConfi
 
 FAsyncTaskNotification::~FAsyncTaskNotification()
 {
-	delete NotificationImpl;
 }
 
 FAsyncTaskNotification::FAsyncTaskNotification(FAsyncTaskNotification&& InOther)
 	: NotificationImpl(InOther.NotificationImpl)
 {
-	InOther.NotificationImpl = nullptr;
+	InOther.NotificationImpl.Reset();
 }
 
 FAsyncTaskNotification& FAsyncTaskNotification::operator=(FAsyncTaskNotification&& InOther)
 {
 	if (this != &InOther)
 	{
-		delete NotificationImpl;
+		NotificationImpl.Reset();
 		NotificationImpl = InOther.NotificationImpl;
-		InOther.NotificationImpl = nullptr;
+		InOther.NotificationImpl.Reset();
 	}
 	return *this;
 }
diff --git a/Engine/Source/Runtime/Core/Private/Misc/CoreAsyncTaskNotificationImpl.cpp b/Engine/Source/Runtime/Core/Private/Misc/CoreAsyncTaskNotificationImpl.cpp
--- a/Engine/Source/Runtime/Core/Private/Misc/CoreAsyncTaskNotificationImpl.cpp
+++ b/Engine/Source/Runtime/Core/Private/Misc/CoreAsyncTaskNotificationImpl.cpp
@@ -215,7 +215,7 @@ FAsyncTaskNotificationFactory::FImplPointerType FAsyncTaskNotificationFactory::I
 {
 	return Factories.Num() > 0
 		? Factories.Last().Value()
-		: new FCoreAsyncTaskNotificationImpl();
+		: MakeShared<FCoreAsyncTaskNotificationImpl>();
 }
 
 #undef LOCTEXT_NAMESPACE
diff --git a/Engine/Source/Runtime/Core/Public/Misc/AsyncTaskNotification.h b/Engine/Source/Runtime/Core/Public/Misc/AsyncTaskNotification.h
--- a/Engine/Source/Runtime/Core/Public/Misc/AsyncTaskNotification.h
+++ b/Engine/Source/Runtime/Core/Public/Misc/AsyncTaskNotification.h
@@ -190,5 +190,5 @@ public:
 
 private:
 	/** Pointer to the real notification implementation */
-	IAsyncTaskNotificationImpl* NotificationImpl = nullptr;
+	TSharedPtr<IAsyncTaskNotificationImpl> NotificationImpl;
 };
diff --git a/Engine/Source/Runtime/Core/Public/Misc/CoreAsyncTaskNotificationImpl.h b/Engine/Source/Runtime/Core/Public/Misc/CoreAsyncTaskNotificationImpl.h
--- a/Engine/Source/Runtime/Core/Public/Misc/CoreAsyncTaskNotificationImpl.h
+++ b/Engine/Source/Runtime/Core/Public/Misc/CoreAsyncTaskNotificationImpl.h
@@ -137,7 +137,7 @@ class CORE_API FAsyncTaskNotificationFactory
 	friend class FAsyncTaskNotification;
 
 public:
-	typedef IAsyncTaskNotificationImpl* FImplPointerType;
+	typedef TSharedPtr<IAsyncTaskNotificationImpl> FImplPointerType;
 	typedef TFunction<FImplPointerType()> FFactoryFunc;
 
 	/**
diff --git a/Engine/Source/Runtime/Slate/Private/Framework/Application/SlateApplication.cpp b/Engine/Source/Runtime/Slate/Private/Framework/Application/SlateApplication.cpp
--- a/Engine/Source/Runtime/Slate/Private/Framework/Application/SlateApplication.cpp
+++ b/Engine/Source/Runtime/Slate/Private/Framework/Application/SlateApplication.cpp
@@ -717,7 +717,7 @@ TSharedRef<FSlateApplication> FSlateApplication::Create(const TSharedRef<class G
 		PlatformApplication->OnDisplayMetricsChanged().AddSP(CurrentApplication.ToSharedRef(), &FSlateApplication::OnVirtualDesktopSizeChanged);
 	}
 
-	FAsyncTaskNotificationFactory::Get().RegisterFactory(TEXT("Slate"), []() -> FAsyncTaskNotificationFactory::FImplPointerType { return new FSlateAsyncTaskNotificationImpl(); });
+	FAsyncTaskNotificationFactory::Get().RegisterFactory(TEXT("Slate"), []() -> FAsyncTaskNotificationFactory::FImplPointerType { return MakeShared<FSlateAsyncTaskNotificationImpl>(); });
 
 	return CurrentApplication.ToSharedRef();
 }
diff --git a/Engine/Source/Runtime/Slate/Private/Framework/Notifications/NotificationManager.cpp b/Engine/Source/Runtime/Slate/Private/Framework/Notifications/NotificationManager.cpp
--- a/Engine/Source/Runtime/Slate/Private/Framework/Notifications/NotificationManager.cpp
+++ b/Engine/Source/Runtime/Slate/Private/Framework/Notifications/NotificationManager.cpp
@@ -92,6 +92,10 @@ void FSlateNotificationManager::ShutdownOnPreExit()
 	FCoreDelegates::OnPreExit.RemoveAll(this);
 
 	RegionalLists.Empty();
+
+	FScopeLock Lock(&StagedNotificationCS);
+	StagedNotifications.Empty();
+
 }
 
 void FSlateNotificationManager::SetRootWindow( const TSharedRef<SWindow> InRootWindow )
@@ -344,4 +348,19 @@ void FSlateNotificationManager::ForceNotificationsInFront( const TSharedRef<SWin
 	}
 }
 
+void FSlateNotificationManager::RegisterStagedNotification(TSharedPtr<IAsyncTaskNotificationImpl> InNotification)
+{
+	FScopeLock Lock(&StagedNotificationCS);
+
+	StagedNotifications.Add(InNotification);
+}
+
+void FSlateNotificationManager::UnregisterStagedNotification(TSharedPtr<IAsyncTaskNotificationImpl> InNotification)
+{
+	FScopeLock Lock(&StagedNotificationCS);
+
+	StagedNotifications.Remove(InNotification);
+}
+
+
 #undef LOCTEXT_NAMESPACE
diff --git a/Engine/Source/Runtime/Slate/Private/Framework/Notifications/SlateAsyncTaskNotificationImpl.cpp b/Engine/Source/Runtime/Slate/Private/Framework/Notifications/SlateAsyncTaskNotificationImpl.cpp
--- a/Engine/Source/Runtime/Slate/Private/Framework/Notifications/SlateAsyncTaskNotificationImpl.cpp
+++ b/Engine/Source/Runtime/Slate/Private/Framework/Notifications/SlateAsyncTaskNotificationImpl.cpp
@@ -2,7 +2,6 @@
 
 #include "Framework/Notifications/SlateAsyncTaskNotificationImpl.h"
 
-#include "Widgets/Notifications/SNotificationList.h"
 #include "Widgets/Notifications/INotificationWidget.h"
 #include "Framework/Notifications/NotificationManager.h"
 
@@ -21,405 +20,199 @@
 
 #define LOCTEXT_NAMESPACE "SlateAsyncTaskNotification"
 
-class SSlateAsyncTaskNotificationWidget : public SCompoundWidget, public INotificationWidget
-{
-public:
-	SLATE_BEGIN_ARGS(SSlateAsyncTaskNotificationWidget) {}
-	SLATE_END_ARGS()
-
-	//~ SWidget
-	void Construct(const FArguments& InArgs, const FAsyncTaskNotificationConfig& InConfig);
-	virtual void Tick(const FGeometry& AllottedGeometry, const double InCurrentTime, const float InDeltaTime) override;
-
-	//~ INotificationWidget
-	virtual void OnSetCompletionState(SNotificationItem::ECompletionState State) override;
-	virtual TSharedRef<SWidget> AsWidget() override;
-
-	/** Set the notification item that owns this widget */
-	void SetOwner(TSharedPtr<SNotificationItem> InOwningNotification);
-
-	/** Update the notification state */
-	void UpdateNotification(const FText& InTitleText, const FText& InProgressText, const FText& InPromptText, const FSimpleDelegate& InHyperlink, const FText& InHyperlinkText);
-
-	/** Set the pending completion state of the notification (applied during the next Tick) and reset the external UI reference */
-	void SetPendingCompletionState(const EAsyncTaskNotificationState InPendingCompletionState, TSharedPtr<SSlateAsyncTaskNotificationWidget>* ExternalReferenceToReset);
-
-	/** Set whether this task be canceled */
-	void SetCanCancel(const TAttribute<bool>& InCanCancel);
-
-	/** Set whether to keep this notification open on success */
-	void SetKeepOpenOnSuccess(const TAttribute<bool>& InKeepOpenOnSuccess);
-
-	/** Set whether to keep this notification open on failure */
-	void SetKeepOpenOnFailure(const TAttribute<bool>& InKeepOpenOnFailure);
-
-	EAsyncTaskNotificationPromptAction GetPromptAction() const;
-private:
-	/** Sync attribute bindings with the cached values (once per-frame from the game thread) */
-	void SyncAttributes();
-
-	/** UI Text */
-	EVisibility GetTitleTextVisibility() const;
-	FText GetTitleText() const;
-	EVisibility GetProgressTextVisibility() const;
-	FText GetProgressText() const;
-
-	/** Throbber */
-	EVisibility GetThrobberVisibility() const;
-
-	/** Status Icon */
-	EVisibility GetStatusIconVisibility() const;
-	const FSlateBrush* GetStatusIconBrush() const;
-
-	/** Cancel button */
-	bool IsCancelButtonEnabled() const;
-	EVisibility GetCancelButtonVisibility() const;
-	FReply OnCancelButtonClicked();
-
-	/** Prompt button */
-	bool IsPromptButtonEnabled() const;
-	EVisibility GetPromptButtonVisibility() const;
-	FReply OnPromptButtonClicked();
-	FText GetPromptButtonText() const;
-
-	/** Close button */
-	EVisibility GetCloseButtonVisibility() const;
-	FReply OnCloseButtonClicked();
-
-	/** Hyperlink */
-	void OnHyperlinkClicked() const;
-	FText GetHyperlinkText() const;
-	EVisibility GetHyperlinkVisibility() const;
-
-	/** Get the current completion state from the parent notification */
-	SNotificationItem::ECompletionState GetNotificationCompletionState() const;
-
-	/** Action taken for the task, resets to none on notification state change. */
-	TAtomic<EAsyncTaskNotificationPromptAction> PromptAction;
-
-	/** Can this task be canceled? Will show a cancel button for in-progress tasks */
-	TAttribute<bool> bCanCancelAttr = false;
-	bool bCanCancel = false;
-
-	/** Keep this notification open on success? Will show a close button */
-	TAttribute<bool> bKeepOpenOnSuccessAttr = false;
-	bool bKeepOpenOnSuccess = false;
-
-	/** Keep this notification open on failure? Will show an close button */
-	TAttribute<bool> bKeepOpenOnFailureAttr = false;
-	bool bKeepOpenOnFailure = false;
+/*
+ * FSlateAsyncTaskNotificationImpl
+ */
 
-	/** The title text displayed in the notification (if any) */
-	FText TitleText;
-
-	/** The progress text displayed in the notification (if any) */
-	FText ProgressText;
-
-	/** The prompt text displayed on the prompt button in the notification (if any)*/
-	FText PromptText;
-
-	/** When set this will display as a hyperlink on the right side of the notification. */
-	FSimpleDelegate Hyperlink;
+FSlateAsyncTaskNotificationImpl::FSlateAsyncTaskNotificationImpl() : PromptAction(EAsyncTaskNotificationPromptAction::None)
+{
+	
+}
 
-	/** Text to display for the hyperlink message */
-	FText HyperlinkText;
+FSlateAsyncTaskNotificationImpl::~FSlateAsyncTaskNotificationImpl()
+{
+	FTSTicker::GetCoreTicker().RemoveTicker(TickerHandle);
+}
 
-	/** The pending completion state of the notification (if any, applied during the next Tick) */
-	TOptional<EAsyncTaskNotificationState> PendingCompletionState;
+void FSlateAsyncTaskNotificationImpl::Initialize(const FAsyncTaskNotificationConfig& InConfig)
+{
+	NotificationConfig = InConfig;
+	
+	// Note: FCoreAsyncTaskNotificationImpl guarantees this is being called from the game thread
 
-	/** The current completion state of this widget. */
-	EAsyncTaskNotificationState NotificationState = EAsyncTaskNotificationState::None;
+	// Initialize the UI if the Notification is not headless
+	if (!NotificationConfig.bIsHeadless)
+	{
+		// Register the ticker to update the notification ever frame
+		TickerHandle = FTSTicker::GetCoreTicker().AddTicker(FTickerDelegate::CreateSP(this, &FSlateAsyncTaskNotificationImpl::TickNotification));
 
-	/** Pointer to the notification item that owns this widget (this is a deliberate reference cycle as we need this object alive until we choose to expire it, at which point we release our reference to allow everything to be destroyed) */
-	TSharedPtr<SNotificationItem> OwningNotification;
+		// Register this as a Staged Notification (Allows notifications to remain open even after this is destroyed)
+		FSlateNotificationManager::Get().RegisterStagedNotification(AsShared());
 
-	/** Critical section preventing concurrent access to the attributes */
-	FCriticalSection AttributesCS;
+		PromptAction = FApp::IsUnattended() ? EAsyncTaskNotificationPromptAction::Unattended : EAsyncTaskNotificationPromptAction::None;
+		bCanCancelAttr = InConfig.bCanCancel;
+		bKeepOpenOnSuccessAttr = InConfig.bKeepOpenOnSuccess;
+		bKeepOpenOnFailureAttr = InConfig.bKeepOpenOnFailure;
+		SyncAttributes();
 
-	/** Critical section preventing the game thread from completing this widget while another thread is in the progress of setting the completion state and cleaning up its UI references */
-	FCriticalSection CompletionCS;
-};
+		// Mark the notification as pending so the UI can initialize
+		PreviousCompletionState = EAsyncTaskNotificationState::None;
+		SetPendingCompletionState(EAsyncTaskNotificationState::Pending); 
+	}
 
-void SSlateAsyncTaskNotificationWidget::Construct(const FArguments& InArgs, const FAsyncTaskNotificationConfig& InConfig)
-{
-	PromptAction = FApp::IsUnattended() ? EAsyncTaskNotificationPromptAction::Unattended : EAsyncTaskNotificationPromptAction::None;
-	bCanCancelAttr = InConfig.bCanCancel;
-	bKeepOpenOnSuccessAttr = InConfig.bKeepOpenOnSuccess;
-	bKeepOpenOnFailureAttr = InConfig.bKeepOpenOnFailure;
-	SyncAttributes();
+	
+	// This calls UpdateNotification to update the UI initialized above
+	FCoreAsyncTaskNotificationImpl::Initialize(InConfig);
+}
 
-	ChildSlot
-	[
-		SNew(SBox)
-		.Padding(FMargin(15.0f))
-		[
-			SNew(SHorizontalBox)
-			// Main Icon
-			+SHorizontalBox::Slot()
-			.AutoWidth()
-			.VAlign(VAlign_Center)
-			.HAlign(HAlign_Left)
-			[
-				SNew(SImage)
-				.Image(InConfig.Icon ? InConfig.Icon : FCoreStyle::Get().GetBrush(TEXT("NotificationList.DefaultMessage")))
-			]
-
-			// Text
-			+SHorizontalBox::Slot()
-			.VAlign(VAlign_Center)
-			[
-				SNew(SVerticalBox)
-
-				// Title Text
-				+SVerticalBox::Slot()
-				.HAlign(HAlign_Left)
-				[
-					SNew(STextBlock)
-					.Margin(FMargin(10.0f, 0.0f, 0.0f, 0.0f))
-					.Font(FCoreStyle::Get().GetFontStyle(TEXT("NotificationList.FontBold")))
-					.Text(this, &SSlateAsyncTaskNotificationWidget::GetTitleText)
-					.Visibility(this, &SSlateAsyncTaskNotificationWidget::GetTitleTextVisibility)
-				]
-
-				// Progress Text
-				+SVerticalBox::Slot()
-				.HAlign(HAlign_Left)
-				[
-					SNew(STextBlock)
-					.Margin(FMargin(10.0f, 5.0f, 0.0f, 0.0f))
-					.Font(FCoreStyle::Get().GetFontStyle(TEXT("NotificationList.FontLight")))
-					.Text(this, &SSlateAsyncTaskNotificationWidget::GetProgressText)
-					.Visibility(this, &SSlateAsyncTaskNotificationWidget::GetProgressTextVisibility)
-				]
-			]
-
-			// Throbber/Status Icon + Buttons
-			+SHorizontalBox::Slot()
-			.AutoWidth()
-			.Padding(FMargin(15.0f, 0.0f, 0.0f, 0.0f))
-			[
-				SNew(SVerticalBox)
-
-				// Throbber/Status Icon
-				+SVerticalBox::Slot()
-				.HAlign(HAlign_Center)
-				.VAlign(VAlign_Center)
-				[
-					SNew(SVerticalBox)
-
-					// Throbber
-					+SVerticalBox::Slot()
-					.AutoHeight()
-					[
-						SNew(SBox)
-						.Padding(FMargin(5.0f, 0.0f, 10.0f, 0.0f))
-						.Visibility(this, &SSlateAsyncTaskNotificationWidget::GetThrobberVisibility)
-						[
-							SNew(SThrobber)
-						]
-					]
-
-					// Status Icon
-					+SVerticalBox::Slot()
-					.AutoHeight()
-					[
-						SNew(SBox)
-						.Padding(FMargin(8.0f, 0.0f, 10.0f, 0.0f))
-						.Visibility(this, &SSlateAsyncTaskNotificationWidget::GetStatusIconVisibility)
-						[
-							SNew(SImage)
-							.Image(this, &SSlateAsyncTaskNotificationWidget::GetStatusIconBrush)
-						]
-					]
-				]
-
-				// Hyperlink
-				+ SVerticalBox::Slot()
-				.AutoHeight()
-				.HAlign(HAlign_Right)
-				.VAlign(VAlign_Bottom)
-				[
-					SNew(SBox)
-					.Padding(FMargin(0.0f, 2.0f, 0.0f, 2.0f))
-					.VAlign(VAlign_Center)
-					.HAlign(HAlign_Left)
-					.Visibility(this, &SSlateAsyncTaskNotificationWidget::GetHyperlinkVisibility)
-					[
-						SNew(SHyperlink)
-						.Text(this, &SSlateAsyncTaskNotificationWidget::GetHyperlinkText)
-						.OnNavigate(this, &SSlateAsyncTaskNotificationWidget::OnHyperlinkClicked)
-					]
-				]
-
-				// Buttons
-				+SVerticalBox::Slot()
-				.AutoHeight()
-				.HAlign(HAlign_Center)
-				.VAlign(VAlign_Bottom)
-				.Padding(FMargin(0.0f, 5.0f, 0.0f, 0.0f))
-				[
-					SNew(SHorizontalBox)
-
-					// Prompt Button
-					+ SHorizontalBox::Slot()
-					.AutoWidth()
-					.VAlign(VAlign_Center)
-					.Padding(FMargin(5.0f, 0.0f, 5.0f, 0.0f))
-					[
-						SNew(SButton)
-						.Text_Raw(this, &SSlateAsyncTaskNotificationWidget::GetPromptButtonText)
-						.Visibility(this, &SSlateAsyncTaskNotificationWidget::GetPromptButtonVisibility)
-						.OnClicked(this, &SSlateAsyncTaskNotificationWidget::OnPromptButtonClicked)
-					]
-
-					// Cancel Button
-					+SHorizontalBox::Slot()
-					.AutoWidth()
-					.VAlign(VAlign_Center)
-					[
-						SNew(SButton)
-						.Text(LOCTEXT("CancelButton", "Cancel"))
-						.Visibility(this, &SSlateAsyncTaskNotificationWidget::GetCancelButtonVisibility)
-						.OnClicked(this, &SSlateAsyncTaskNotificationWidget::OnCancelButtonClicked)
-					]
-
-					// Close Button
-					+SHorizontalBox::Slot()
-					.AutoWidth()
-					.VAlign(VAlign_Center)
-					[
-						SNew(SButton)
-						.Text(LOCTEXT("CloseButton", "Close"))
-						.Visibility(this, &SSlateAsyncTaskNotificationWidget::GetCloseButtonVisibility)
-						.OnClicked(this, &SSlateAsyncTaskNotificationWidget::OnCloseButtonClicked)
-					]
-				]
-			]
-		]
-	];
-}
-
-void SSlateAsyncTaskNotificationWidget::Tick(const FGeometry& AllottedGeometry, const double InCurrentTime, const float InDeltaTime)
+void FSlateAsyncTaskNotificationImpl::DestroyCurrentNotification()
 {
-	SyncAttributes();
-
-	EAsyncTaskNotificationState CompletionStateToApply = EAsyncTaskNotificationState::None;
+	if(OwningNotification)
 	{
-		FScopeLock Lock(&CompletionCS);
-
-		if (PendingCompletionState.IsSet())
-		{
-			CompletionStateToApply = PendingCompletionState.GetValue();
-			PendingCompletionState.Reset();
-		}
-	}
+		// Perform the normal automatic fadeout
+		OwningNotification->ExpireAndFadeout();
 
-	if (OwningNotification && CompletionStateToApply != EAsyncTaskNotificationState::None)
-	{
-		NotificationState = CompletionStateToApply;
-		SNotificationItem::ECompletionState OwningCompletionState = SNotificationItem::CS_None;
-		switch (NotificationState)
-		{
-		case EAsyncTaskNotificationState::Pending:
-			OwningCompletionState = SNotificationItem::CS_Pending;
-			break;
-		case EAsyncTaskNotificationState::Failure:
-			OwningCompletionState = SNotificationItem::CS_Fail;
-			break;
-		case EAsyncTaskNotificationState::Success:
-			OwningCompletionState = SNotificationItem::CS_Success;
-			break;
-		case EAsyncTaskNotificationState::Prompt:
-			OwningNotification->Pulse(FLinearColor(0.f, 0.f, 1.f));
-
-			break;
-		}
-		if (OwningCompletionState != SNotificationItem::CS_None && OwningCompletionState != OwningNotification->GetCompletionState())
-		{
-			OwningNotification->SetCompletionState(OwningCompletionState);
-		}
+		// Release our reference to our owner so that everything can be destroyed
+		OwningNotification.Reset();
 	}
 }
 
-void SSlateAsyncTaskNotificationWidget::OnSetCompletionState(SNotificationItem::ECompletionState State)
+void FSlateAsyncTaskNotificationImpl::CreateNewNotificationItem(EAsyncTaskNotificationState NewNotificationState)
 {
-	check(State == GetNotificationCompletionState());
-
-	// If we completed and we aren't keeping the notification open (which will show the Close button), then expire the notification immediately
-	if ((State == SNotificationItem::CS_Success || State == SNotificationItem::CS_Fail) && GetCloseButtonVisibility() == EVisibility::Collapsed)
+	DestroyCurrentNotification();
+	
+	switch (NewNotificationState)
 	{
-		if (OwningNotification)
-		{
-			// Perform the normal automatic fadeout
-			OwningNotification->ExpireAndFadeout();
-
-			// Release our reference to our owner so that everything can be destroyed
-			OwningNotification.Reset();
-		}
+	case EAsyncTaskNotificationState::Pending:
+		CreatePendingNotification();
+		break;
+	case EAsyncTaskNotificationState::Failure:
+		CreateFailureNotification();
+		break;
+	case EAsyncTaskNotificationState::Success:
+		CreateSuccessNotification();
+		break;
+	case EAsyncTaskNotificationState::Prompt:
+		CreatePromptNotification();
+		break;
 	}
 
-	// Reset the `PromptAction` state when changing completion state
-	PromptAction = FApp::IsUnattended() ? EAsyncTaskNotificationPromptAction::Unattended : EAsyncTaskNotificationPromptAction::None;
 }
 
-TSharedRef<SWidget> SSlateAsyncTaskNotificationWidget::AsWidget()
+void FSlateAsyncTaskNotificationImpl::SetOwner(TSharedPtr<SNotificationItem> InOwningNotification)
 {
-	return AsShared();
+	OwningNotification = InOwningNotification;
+
+	// Update the notification here to make sure it has the correct Text/Subtext/Hyperlink etc
+	UpdateNotification();
 }
 
-void SSlateAsyncTaskNotificationWidget::SetOwner(TSharedPtr<SNotificationItem> InOwningNotification)
+TSharedPtr<SNotificationItem> FSlateAsyncTaskNotificationImpl::SetupNotificationItem(FNotificationInfo& NotificationInfo)
 {
-	OwningNotification = InOwningNotification;
+	NotificationInfo.FadeOutDuration = NotificationConfig.FadeOutDuration;
+	NotificationInfo.ExpireDuration = NotificationConfig.ExpireDuration;
+	NotificationInfo.FadeInDuration = NotificationConfig.FadeInDuration;
+	NotificationInfo.bFireAndForget = false;
+
+	TSharedPtr<SNotificationItem> NotificationItem = FSlateNotificationManager::Get().AddNotification(NotificationInfo);
+	check(NotificationItem);
+	
+	SetOwner(NotificationItem);
+	return NotificationItem;
 }
 
-void SSlateAsyncTaskNotificationWidget::UpdateNotification(const FText& InTitleText, const FText& InProgressText, const FText& InPromptText, const FSimpleDelegate& InHyperlink, const FText& InHyperlinkText)
+void FSlateAsyncTaskNotificationImpl::CreatePendingNotification()
 {
-	TitleText = InTitleText;
-	ProgressText = InProgressText;
-	PromptText = InPromptText;
-	Hyperlink = InHyperlink;
-	HyperlinkText = InHyperlinkText;
+	FNotificationInfo NotificationInfo(FText::GetEmpty());
+
+	// Pending Notifications have a throbber to show progress
+	NotificationInfo.bUseThrobber = true;
+	
+	AddCancelButton(NotificationInfo, SNotificationItem::CS_Pending);
+	
+	TSharedPtr<SNotificationItem> NotificationItem = SetupNotificationItem(NotificationInfo);
 }
 
-void SSlateAsyncTaskNotificationWidget::SetPendingCompletionState(const EAsyncTaskNotificationState InPendingCompletionState, TSharedPtr<SSlateAsyncTaskNotificationWidget>* ExternalReferenceToReset)
+void FSlateAsyncTaskNotificationImpl::CreateSuccessNotification()
 {
-	FScopeLock Lock(&CompletionCS);
+	FNotificationInfo NotificationInfo(FText::GetEmpty());
 
-	// Set the completion state and reset the external UI reference while we have the lock to avoid the game thread potentially destroying this notification (via a Tick) while another thread is still clearing its references to it
-	PendingCompletionState = InPendingCompletionState;
-	if (ExternalReferenceToReset)
-	{
-		ExternalReferenceToReset->Reset();
-	}
+	NotificationInfo.Image = FAppStyle::Get().GetBrush("NotificationList.SuccessImage");
+	AddCloseButton(NotificationInfo);
+	
+	TSharedPtr<SNotificationItem> NotificationItem = SetupNotificationItem(NotificationInfo);
 }
 
-void SSlateAsyncTaskNotificationWidget::SetCanCancel(const TAttribute<bool>& InCanCancel)
+void FSlateAsyncTaskNotificationImpl::CreateFailureNotification()
 {
-	FScopeLock Lock(&AttributesCS);
+	FNotificationInfo NotificationInfo(FText::GetEmpty());
+	
+	NotificationInfo.Image = FAppStyle::Get().GetBrush("NotificationList.FailImage");
+	AddCloseButton(NotificationInfo);
 
-	bCanCancelAttr = InCanCancel;
+	TSharedPtr<SNotificationItem> NotificationItem = SetupNotificationItem(NotificationInfo);
 }
 
-void SSlateAsyncTaskNotificationWidget::SetKeepOpenOnSuccess(const TAttribute<bool>& InKeepOpenOnSuccess)
+void FSlateAsyncTaskNotificationImpl::CreatePromptNotification()
 {
-	FScopeLock Lock(&AttributesCS);
+	FNotificationInfo NotificationInfo(FText::GetEmpty());
 
-	bKeepOpenOnSuccessAttr = InKeepOpenOnSuccess;
+	AddPromptButton(NotificationInfo);
+	
+	AddCancelButton(NotificationInfo, SNotificationItem::CS_None);
+	
+	TSharedPtr<SNotificationItem> NotificationItem = SetupNotificationItem(NotificationInfo);
 }
 
-void SSlateAsyncTaskNotificationWidget::SetKeepOpenOnFailure(const TAttribute<bool>& InKeepOpenOnFailure)
+void FSlateAsyncTaskNotificationImpl::AddPromptButton(FNotificationInfo &NotificationInfo)
 {
-	FScopeLock Lock(&AttributesCS);
+	if(GetPromptButtonVisibility() == EVisibility::Visible)
+	{
+		FNotificationButtonInfo PromptButtonInfo(
+		PromptText,
+		FText::GetEmpty(),
+		FSimpleDelegate::CreateSP(this, &FSlateAsyncTaskNotificationImpl::OnPromptButtonClicked),
+		SNotificationItem::CS_None
+		);
+	
+		NotificationInfo.ButtonDetails.Add(PromptButtonInfo);
+	}
+}
+
 
-	bKeepOpenOnFailureAttr = InKeepOpenOnFailure;
+void FSlateAsyncTaskNotificationImpl::AddCancelButton(FNotificationInfo &NotificationInfo, SNotificationItem::ECompletionState VisibleInState)
+{
+	if(GetCancelButtonVisibility() == EVisibility::Visible)
+	{
+		FNotificationButtonInfo CancelButtonInfo(
+		LOCTEXT("CancelButton", "Cancel"),
+		FText::GetEmpty(),
+		FSimpleDelegate::CreateSP(this, &FSlateAsyncTaskNotificationImpl::OnCancelButtonClicked),
+		VisibleInState
+		);
+	
+		NotificationInfo.ButtonDetails.Add(CancelButtonInfo);
+	}
 }
 
-EAsyncTaskNotificationPromptAction SSlateAsyncTaskNotificationWidget::GetPromptAction() const
+void FSlateAsyncTaskNotificationImpl::AddCloseButton(FNotificationInfo &NotificationInfo)
 {
-	return PromptAction;
+	if(GetCloseButtonVisibility() == EVisibility::Visible)
+	{
+		FNotificationButtonInfo CloseButtonInfo(
+		LOCTEXT("CloseButton", "Close"),
+		FText::GetEmpty(),
+		FSimpleDelegate::CreateSP(this, &FSlateAsyncTaskNotificationImpl::OnCloseButtonClicked)
+		);
+
+		CloseButtonInfo.VisibilityOnSuccess = EVisibility::Visible;
+		CloseButtonInfo.VisibilityOnFail = EVisibility::Visible;
+		
+		NotificationInfo.ButtonDetails.Add(CloseButtonInfo);
+	}
 }
 
-void SSlateAsyncTaskNotificationWidget::SyncAttributes()
+void FSlateAsyncTaskNotificationImpl::SyncAttributes()
 {
 	FScopeLock Lock(&AttributesCS);
 
@@ -428,100 +221,107 @@ void SSlateAsyncTaskNotificationWidget::SyncAttributes()
 	bKeepOpenOnFailure = bKeepOpenOnFailureAttr.Get(false);
 }
 
-EVisibility SSlateAsyncTaskNotificationWidget::GetTitleTextVisibility() const
+void FSlateAsyncTaskNotificationImpl::OnSetCompletionState(SNotificationItem::ECompletionState InState)
 {
-	return (TitleText.IsEmpty())
-		? EVisibility::Collapsed
-		: EVisibility::Visible;
-}
+	check(InState == GetNotificationCompletionState());
 
-FText SSlateAsyncTaskNotificationWidget::GetTitleText() const
-{
-	return TitleText;
-}
+	// If we completed and we aren't keeping the notification open (which will show the Close button), then expire the notification immediately
+	if ((InState == SNotificationItem::CS_Success || InState == SNotificationItem::CS_Fail) && GetCloseButtonVisibility() == EVisibility::Collapsed)
+	{
+		DestroyCurrentNotification();
 
-EVisibility SSlateAsyncTaskNotificationWidget::GetProgressTextVisibility() const
-{
-	return (ProgressText.IsEmpty())
-		? EVisibility::Collapsed
-		: EVisibility::Visible;
+		FSlateNotificationManager::Get().UnregisterStagedNotification(AsShared());
+	}
+
+	// Reset the `PromptAction` state when changing completion state
+	PromptAction = FApp::IsUnattended() ? EAsyncTaskNotificationPromptAction::Unattended : EAsyncTaskNotificationPromptAction::None;
 }
 
-FText SSlateAsyncTaskNotificationWidget::GetProgressText() const
+void FSlateAsyncTaskNotificationImpl::SetPendingCompletionState(const EAsyncTaskNotificationState InPendingCompletionState)
 {
-	return ProgressText;
+	FScopeLock Lock(&CompletionCS);
+
+	// Set the completion state
+	PendingCompletionState = InPendingCompletionState;
 }
 
-EVisibility SSlateAsyncTaskNotificationWidget::GetThrobberVisibility() const
+void FSlateAsyncTaskNotificationImpl::SetCanCancel(const TAttribute<bool>& InCanCancel)
 {
-	return (NotificationState == EAsyncTaskNotificationState::Pending)
-		? EVisibility::Visible
-		: EVisibility::Collapsed;
+	if (!NotificationConfig.bIsHeadless)
+	{
+		FScopeLock Lock(&AttributesCS);
+
+		bCanCancelAttr = InCanCancel;
+	}
 }
 
-EVisibility SSlateAsyncTaskNotificationWidget::GetStatusIconVisibility() const
+void FSlateAsyncTaskNotificationImpl::SetKeepOpenOnSuccess(const TAttribute<bool>& InKeepOpenOnSuccess)
 {
-	return (NotificationState == EAsyncTaskNotificationState::Success || NotificationState == EAsyncTaskNotificationState::Failure)
-		? EVisibility::Visible
-		: EVisibility::Collapsed;
+	if (!NotificationConfig.bIsHeadless)
+	{
+		FScopeLock Lock(&AttributesCS);
+
+		bKeepOpenOnSuccessAttr = InKeepOpenOnSuccess;
+	}
 }
 
-const FSlateBrush* SSlateAsyncTaskNotificationWidget::GetStatusIconBrush() const
+void FSlateAsyncTaskNotificationImpl::SetKeepOpenOnFailure(const TAttribute<bool>& InKeepOpenOnFailure)
 {
-	return (NotificationState == EAsyncTaskNotificationState::Success)
-		? FCoreStyle::Get().GetBrush("NotificationList.SuccessImage")
-		: FCoreStyle::Get().GetBrush("NotificationList.FailImage");
+	if (!NotificationConfig.bIsHeadless)
+	{
+		FScopeLock Lock(&AttributesCS);
+
+		bKeepOpenOnFailureAttr = InKeepOpenOnFailure;
+	}
 }
 
-bool SSlateAsyncTaskNotificationWidget::IsCancelButtonEnabled() const
+bool FSlateAsyncTaskNotificationImpl::IsCancelButtonEnabled() const
 {
 	return bCanCancel && PromptAction == EAsyncTaskNotificationPromptAction::None;
 }
 
-EVisibility SSlateAsyncTaskNotificationWidget::GetCancelButtonVisibility() const
+EVisibility FSlateAsyncTaskNotificationImpl::GetCancelButtonVisibility() const
 {
-	return (bCanCancel && (NotificationState == EAsyncTaskNotificationState::Pending || NotificationState == EAsyncTaskNotificationState::Prompt))
+	return (bCanCancel && (State == EAsyncTaskNotificationState::Pending || State == EAsyncTaskNotificationState::Prompt))
 		? EVisibility::Visible
 		: EVisibility::Collapsed;
 }
 
-FReply SSlateAsyncTaskNotificationWidget::OnCancelButtonClicked()
+void FSlateAsyncTaskNotificationImpl::OnCancelButtonClicked()
 {
 	PromptAction = EAsyncTaskNotificationPromptAction::Cancel;
-	return FReply::Handled();
 }
 
-bool SSlateAsyncTaskNotificationWidget::IsPromptButtonEnabled() const
+bool FSlateAsyncTaskNotificationImpl::IsPromptButtonEnabled() const
 {
 	return PromptAction == EAsyncTaskNotificationPromptAction::None;
 }
 
-EVisibility SSlateAsyncTaskNotificationWidget::GetPromptButtonVisibility() const
+EVisibility FSlateAsyncTaskNotificationImpl::GetPromptButtonVisibility() const
 {
-	return (!FApp::IsUnattended() && NotificationState == EAsyncTaskNotificationState::Prompt)
+	return (!FApp::IsUnattended() && State == EAsyncTaskNotificationState::Prompt)
 		? EVisibility::Visible
 		: EVisibility::Collapsed;
 }
 
-FReply SSlateAsyncTaskNotificationWidget::OnPromptButtonClicked()
+void FSlateAsyncTaskNotificationImpl::OnPromptButtonClicked()
 {
 	PromptAction = EAsyncTaskNotificationPromptAction::Continue;
-	return FReply::Handled();
 }
 
-FText SSlateAsyncTaskNotificationWidget::GetPromptButtonText() const
+FText FSlateAsyncTaskNotificationImpl::GetPromptButtonText() const
 {
 	return PromptText;
 }
 
-EVisibility SSlateAsyncTaskNotificationWidget::GetCloseButtonVisibility() const
+EVisibility FSlateAsyncTaskNotificationImpl::GetCloseButtonVisibility() const
 {
-	return (!FApp::IsUnattended() && ((bKeepOpenOnSuccess && NotificationState == EAsyncTaskNotificationState::Success) || (bKeepOpenOnFailure && NotificationState == EAsyncTaskNotificationState::Failure)))
+	return (!FApp::IsUnattended() && ((bKeepOpenOnSuccess && State == EAsyncTaskNotificationState::Success) || (bKeepOpenOnFailure && State == EAsyncTaskNotificationState::Failure)))
 		? EVisibility::Visible
 		: EVisibility::Collapsed;
 }
 
-FReply SSlateAsyncTaskNotificationWidget::OnCloseButtonClicked()
+void FSlateAsyncTaskNotificationImpl::OnCloseButtonClicked()
 {
 	if (OwningNotification)
 	{
@@ -532,26 +332,28 @@ FReply SSlateAsyncTaskNotificationWidget::OnCloseButtonClicked()
 
 		// Release our reference to our owner so that everything can be destroyed
 		OwningNotification.Reset();
+
+		// Unregister the Staged Notification to complete the cleanup
+		FSlateNotificationManager::Get().UnregisterStagedNotification(AsShared());
 	}
-	return FReply::Handled();
 }
 
-void SSlateAsyncTaskNotificationWidget::OnHyperlinkClicked() const
+void FSlateAsyncTaskNotificationImpl::OnHyperlinkClicked() const
 {
 	Hyperlink.ExecuteIfBound();
 }
 
-FText SSlateAsyncTaskNotificationWidget::GetHyperlinkText() const
+FText FSlateAsyncTaskNotificationImpl::GetHyperlinkText() const
 {
 	return HyperlinkText;
 }
 
-EVisibility SSlateAsyncTaskNotificationWidget::GetHyperlinkVisibility() const
+EVisibility FSlateAsyncTaskNotificationImpl::GetHyperlinkVisibility() const
 {
 	return Hyperlink.IsBound() ? EVisibility::Visible : EVisibility::Collapsed;
 }
 
-SNotificationItem::ECompletionState SSlateAsyncTaskNotificationWidget::GetNotificationCompletionState() const
+SNotificationItem::ECompletionState FSlateAsyncTaskNotificationImpl::GetNotificationCompletionState() const
 {
 	if (OwningNotification)
 	{
@@ -560,74 +362,97 @@ SNotificationItem::ECompletionState SSlateAsyncTaskNotificationWidget::GetNotifi
 	return SNotificationItem::CS_None;
 }
 
-void FSlateAsyncTaskNotificationImpl::Initialize(const FAsyncTaskNotificationConfig& InConfig)
+void FSlateAsyncTaskNotificationImpl::UpdateNotification()
 {
-	// Note: FCoreAsyncTaskNotificationImpl guarantees this is being called from the game thread
-	if (!InConfig.bIsHeadless)
+	FCoreAsyncTaskNotificationImpl::UpdateNotification();
+	
+	if (!NotificationConfig.bIsHeadless)
 	{
-		// Set-up the notification UI
-		NotificationItemWidget = SNew(SSlateAsyncTaskNotificationWidget, InConfig);
-		check(NotificationItemWidget);
-
-		FNotificationInfo NotificationInfo(NotificationItemWidget);
-		NotificationInfo.FadeInDuration = InConfig.FadeInDuration;
-		NotificationInfo.FadeOutDuration = InConfig.FadeOutDuration;
-		NotificationInfo.ExpireDuration = InConfig.ExpireDuration;
-		NotificationInfo.bFireAndForget = false;
-
-		TSharedPtr<SNotificationItem> NotificationItem = FSlateNotificationManager::Get().AddNotification(NotificationInfo);
-		check(NotificationItem);
+		// Update the notification UI only if the state hasn't changed (i.e this notification will not be deleted)
+		if(OwningNotification && State == PreviousCompletionState)
+		{
+			OwningNotification->SetText(TitleText);
+			OwningNotification->SetSubText(ProgressText);
+			OwningNotification->SetHyperlink(Hyperlink, HyperlinkText);
+		}
 
-		NotificationItemWidget->SetOwner(NotificationItem);
-		NotificationItem->SetCompletionState(SNotificationItem::CS_Pending);
+		// Set the Pending Completion State in case the notification has to change
+		SetPendingCompletionState(State);
 	}
-
-	// This calls UpdateNotification to update the UI initialized above
-	FCoreAsyncTaskNotificationImpl::Initialize(InConfig);
 }
 
-void FSlateAsyncTaskNotificationImpl::SetCanCancel(const TAttribute<bool>& InCanCancel)
+EAsyncTaskNotificationPromptAction FSlateAsyncTaskNotificationImpl::GetPromptAction() const
 {
-	if (NotificationItemWidget)
+	if(NotificationConfig.bIsHeadless)
 	{
-		NotificationItemWidget->SetCanCancel(InCanCancel);
+		return EAsyncTaskNotificationPromptAction::Unattended;
 	}
+	return PromptAction;
 }
 
-void FSlateAsyncTaskNotificationImpl::SetKeepOpenOnSuccess(const TAttribute<bool>& InKeepOpenOnSuccess)
+bool FSlateAsyncTaskNotificationImpl::TickNotification(float InDeltaTime)
 {
-	if (NotificationItemWidget)
+	SyncAttributes();
+	
+	EAsyncTaskNotificationState CompletionStateToApply = EAsyncTaskNotificationState::None;
 	{
-		NotificationItemWidget->SetKeepOpenOnSuccess(InKeepOpenOnSuccess);
-	}
-}
+		FScopeLock Lock(&CompletionCS);
 
-void FSlateAsyncTaskNotificationImpl::SetKeepOpenOnFailure(const TAttribute<bool>& InKeepOpenOnSuccess)
-{
-	if (NotificationItemWidget)
-	{
-		NotificationItemWidget->SetKeepOpenOnFailure(InKeepOpenOnSuccess);
+		if (PendingCompletionState.IsSet())
+		{
+			CompletionStateToApply = PendingCompletionState.GetValue();
+			PendingCompletionState.Reset();
+		}
 	}
-}
 
-void FSlateAsyncTaskNotificationImpl::UpdateNotification()
-{
-	FCoreAsyncTaskNotificationImpl::UpdateNotification();
-	
-	if (NotificationItemWidget)
+	// Create a new notification if the state changed to a valid state
+	if (PreviousCompletionState != CompletionStateToApply && CompletionStateToApply != EAsyncTaskNotificationState::None)
 	{
-		// Update the notification text
-		NotificationItemWidget->UpdateNotification(TitleText, ProgressText, PromptText, Hyperlink, HyperlinkText);
+		// Reset the State of the previous notification if it was 'Pending', to make any misleading buttons disappear
+		if(OwningNotification && PreviousCompletionState == EAsyncTaskNotificationState::Pending)
+		{
+			OwningNotification->SetCompletionState(SNotificationItem::CS_None);
+		}
+		
+		PreviousCompletionState = CompletionStateToApply;
+
+		// Create a new notification based on the new state
+		CreateNewNotificationItem(CompletionStateToApply);
 
-		// Complete the notification and remove our references to it in a single atomic operation if needed
-		// NotificationItemWidget will be null once this call completes if completion is success or failure
-		NotificationItemWidget->SetPendingCompletionState(State, State == EAsyncTaskNotificationState::Failure || State == EAsyncTaskNotificationState::Success ? &NotificationItemWidget : nullptr);
+		if(OwningNotification)
+		{
+			SNotificationItem::ECompletionState OwningCompletionState = SNotificationItem::CS_None;
+			switch (State)
+			{
+			case EAsyncTaskNotificationState::Pending:
+				OwningCompletionState = SNotificationItem::CS_Pending;
+				break;
+			case EAsyncTaskNotificationState::Failure:
+				OwningCompletionState = SNotificationItem::CS_Fail;
+				break;
+			case EAsyncTaskNotificationState::Success:
+				OwningCompletionState = SNotificationItem::CS_Success;
+				break;
+			case EAsyncTaskNotificationState::Prompt:
+				OwningNotification->Pulse(FLinearColor(0.f, 0.f, 1.f));
+
+				break;
+			}
+			if (OwningCompletionState != SNotificationItem::CS_None && OwningCompletionState != OwningNotification->GetCompletionState())
+			{
+				OwningNotification->SetCompletionState(OwningCompletionState);
+				OnSetCompletionState(OwningCompletionState);
+
+				// We don't need the ticker anymore if the notification is complete
+				if(OwningCompletionState == SNotificationItem::CS_Success || OwningCompletionState == SNotificationItem::CS_Fail)
+				{
+					return false;
+				}
+			}
+		}
 	}
-}
 
-EAsyncTaskNotificationPromptAction FSlateAsyncTaskNotificationImpl::GetPromptAction() const
-{
-	return NotificationItemWidget ? NotificationItemWidget->GetPromptAction() : EAsyncTaskNotificationPromptAction::Unattended;
+	return true;
 }
 
 #undef LOCTEXT_NAMESPACE
diff --git a/Engine/Source/Runtime/Slate/Private/Framework/Notifications/SlateAsyncTaskNotificationImpl.h b/Engine/Source/Runtime/Slate/Private/Framework/Notifications/SlateAsyncTaskNotificationImpl.h
--- a/Engine/Source/Runtime/Slate/Private/Framework/Notifications/SlateAsyncTaskNotificationImpl.h
+++ b/Engine/Source/Runtime/Slate/Private/Framework/Notifications/SlateAsyncTaskNotificationImpl.h
@@ -3,14 +3,16 @@
 #pragma once
 
 #include "Misc/CoreAsyncTaskNotificationImpl.h"
+#include "Containers/Ticker.h"
+#include "SlateCore/Public/Layout/Visibility.h"
+#include "Widgets/Notifications/SNotificationList.h"
 
-class SNotificationItem;
 class SSlateAsyncTaskNotificationWidget;
 
 /**
  * Slate asynchronous task notification that uses a notification item.
  */
-class FSlateAsyncTaskNotificationImpl : public FCoreAsyncTaskNotificationImpl
+class FSlateAsyncTaskNotificationImpl : public TSharedFromThis<FSlateAsyncTaskNotificationImpl, ESPMode::ThreadSafe>,  public FCoreAsyncTaskNotificationImpl
 {
 public:
 	//~ IAsyncTaskNotificationImpl
@@ -19,10 +21,115 @@ public:
 	virtual void SetKeepOpenOnSuccess(const TAttribute<bool>& InKeepOpenOnSuccess) override;
 	virtual void SetKeepOpenOnFailure(const TAttribute<bool>& InKeepOpenOnFailure) override;
 	virtual EAsyncTaskNotificationPromptAction GetPromptAction() const override;
+
+	FSlateAsyncTaskNotificationImpl();
+	virtual ~FSlateAsyncTaskNotificationImpl() override;
+	
 private:
 	//~ FCoreAsyncTaskNotificationImpl
 	virtual void UpdateNotification() override;
 
-	/** Active Slate notification widget */
-	TSharedPtr<SSlateAsyncTaskNotificationWidget> NotificationItemWidget;
+	/** Run every frame on the game thread to update the notification */
+	bool TickNotification(float InDeltaTime);
+
+	/** Set the notification item that owns this widget */
+	void SetOwner(TSharedPtr<SNotificationItem> InOwningNotification);
+
+	/** Sync attribute bindings with the cached values (once per-frame from the game thread) */
+	void SyncAttributes();
+
+	/** Set the pending completion state of the notification (applied during the next Tick) */
+	void SetPendingCompletionState(const EAsyncTaskNotificationState InPendingCompletionState);
+
+	virtual void OnSetCompletionState(SNotificationItem::ECompletionState State);
+
+	/** Cancel button */
+	bool IsCancelButtonEnabled() const;
+	EVisibility GetCancelButtonVisibility() const;
+	void OnCancelButtonClicked();
+
+	/** Prompt button */
+	bool IsPromptButtonEnabled() const;
+	EVisibility GetPromptButtonVisibility() const;
+	void OnPromptButtonClicked();
+	FText GetPromptButtonText() const;
+
+	/** Close button */
+	EVisibility GetCloseButtonVisibility() const;
+	void OnCloseButtonClicked();
+
+	/** Hyperlink */
+	void OnHyperlinkClicked() const;
+	FText GetHyperlinkText() const;
+	EVisibility GetHyperlinkVisibility() const;
+
+	/** Get the current completion state from the parent notification */
+	SNotificationItem::ECompletionState GetNotificationCompletionState() const;
+
+	/* Create a new notification item when the state changes */
+	void CreateNewNotificationItem(EAsyncTaskNotificationState NewNotificationState);
+
+	/* Create a notification for the Pending state */
+	void CreatePendingNotification();
+
+	/* Create a notification for the Success state */
+	void CreateSuccessNotification();
+
+	/* Create a notification for the Failure state */
+	void CreateFailureNotification();
+
+	/* Create a notification for the Prompt state */
+	void CreatePromptNotification();
+
+	/* Helper function to add the "Cancel" button to a notification */
+	void AddCancelButton(FNotificationInfo &NotificationInfo, SNotificationItem::ECompletionState VisibleInState);
+
+	/* Helper function to add the "Close" button to a notification */
+	void AddCloseButton(FNotificationInfo &NotificationInfo);
+
+	/* Helper function to add a Prompt button to a notification */
+	void AddPromptButton(FNotificationInfo &NotificationInfo);
+
+	/* Helper function to setup the information common to notifications in all states */
+	TSharedPtr<SNotificationItem> SetupNotificationItem(FNotificationInfo& NotificationInfo);
+
+	/* Cleanly destroy the current notification */
+	void DestroyCurrentNotification();
+
+private:
+
+	/** The Config used for all notifications */
+	FAsyncTaskNotificationConfig NotificationConfig;
+
+	/** Handle for TickNotification() */
+	FTSTicker::FDelegateHandle TickerHandle;
+
+	/** Action taken for the task, resets to none on notification state change. */
+	TAtomic<EAsyncTaskNotificationPromptAction> PromptAction;
+
+	/** Can this task be canceled? Will show a cancel button for in-progress tasks */
+	TAttribute<bool> bCanCancelAttr = false;
+	bool bCanCancel = false;
+
+	/** Keep this notification open on success? Will show a close button */
+	TAttribute<bool> bKeepOpenOnSuccessAttr = false;
+	bool bKeepOpenOnSuccess = false;
+
+	/** Keep this notification open on failure? Will show an close button */
+	TAttribute<bool> bKeepOpenOnFailureAttr = false;
+	bool bKeepOpenOnFailure = false;
+
+	/** The pending completion state of the notification (if any, applied during the next Tick) */
+	TOptional<EAsyncTaskNotificationState> PendingCompletionState;
+
+	EAsyncTaskNotificationState PreviousCompletionState;
+
+	/** Pointer to the notification item that owns this widget (this is a deliberate reference cycle as we need this object alive until we choose to expire it, at which point we release our reference to allow everything to be destroyed) */
+	TSharedPtr<SNotificationItem> OwningNotification;
+
+	/** Critical section preventing concurrent access to the attributes */
+	FCriticalSection AttributesCS;
+
+	/** Critical section preventing the game thread from completing this widget while another thread is in the progress of setting the completion state and cleaning up its UI references */
+	FCriticalSection CompletionCS;
 };
diff --git a/Engine/Source/Runtime/Slate/Public/Framework/Notifications/NotificationManager.h b/Engine/Source/Runtime/Slate/Public/Framework/Notifications/NotificationManager.h
--- a/Engine/Source/Runtime/Slate/Public/Framework/Notifications/NotificationManager.h
+++ b/Engine/Source/Runtime/Slate/Public/Framework/Notifications/NotificationManager.h
@@ -9,6 +9,7 @@
 #include "Widgets/SWindow.h"
 
 #include "HAL/IConsoleManager.h"
+#include "Misc/CoreAsyncTaskNotificationImpl.h"
 
 struct FNotificationInfo;
 
@@ -166,6 +167,16 @@ public:
 		return this->bAllowNotifications;
 	}
 
+	/**
+	 * Register a Staged Async Notification, allowing the NotificationManager to keep a reference to it
+	 */
+	void RegisterStagedNotification(TSharedPtr<IAsyncTaskNotificationImpl> InNotification);
+
+	/**
+	 * Unregister a previously added Staged Notification
+	 */
+	void UnregisterStagedNotification(TSharedPtr<IAsyncTaskNotificationImpl> InNotification);
+
 protected:
 	/** Protect constructor as this is a singleton */
 	FSlateNotificationManager();
@@ -209,6 +220,12 @@ private:
 	/** Thread safe queue of notifications to display */
 	TLockFreePointerListLIFO<FNotificationInfo> PendingNotifications;
 
+	/** An array of staged notifications */
+	TArray<TSharedPtr<IAsyncTaskNotificationImpl>> StagedNotifications;
+
+	/** Critical Section to guard access to StagedNotifications */
+	FCriticalSection StagedNotificationCS;
+
 	/** Counter used to create progress handles */
 	static int32 ProgressHandleCounter;
 
