Tabs should only show close button if the tab reports it is closeable

diff --git a/Engine/Source/Runtime/Slate/Private/Widgets/Docking/SDockTab.cpp b/Engine/Source/Runtime/Slate/Private/Widgets/Docking/SDockTab.cpp
--- a/Engine/Source/Runtime/Slate/Private/Widgets/Docking/SDockTab.cpp
+++ b/Engine/Source/Runtime/Slate/Private/Widgets/Docking/SDockTab.cpp
@@ -550,6 +550,7 @@ void SDockTab::Construct( const FArguments& InArgs )
 	this->TabRole = InArgs._TabRole;
 	this->OnTabClosed = InArgs._OnTabClosed;
 	this->OnCanCloseTab = InArgs._OnCanCloseTab;
+	this->bCanEverClose = InArgs._CanEverClose;
 	this->OnPersistVisualState = InArgs._OnPersistVisualState;
 	this->OnExtendContextMenu = InArgs._OnExtendContextMenu;
 	this->TabLabel = InArgs._Label;
@@ -935,7 +936,7 @@ FText SDockTab::GetCloseButtonToolTipText() const
 
 EVisibility SDockTab::HandleIsCloseButtonVisible() const
 {
-	return ((IsHovered() || IsForeground()) && MyTabManager.Pin()->IsTabCloseable(SharedThis(this))) ? EVisibility::Visible : EVisibility::Hidden;
+	return bCanEverClose && ((IsHovered() || IsForeground()) && MyTabManager.Pin()->IsTabCloseable(SharedThis(this))) ? EVisibility::Visible : EVisibility::Hidden;
 }
 
 TOptional<FVector2D> SDockTab::GetTabIconSize() const
diff --git a/Engine/Source/Runtime/Slate/Public/Widgets/Docking/SDockTab.h b/Engine/Source/Runtime/Slate/Public/Widgets/Docking/SDockTab.h
--- a/Engine/Source/Runtime/Slate/Public/Widgets/Docking/SDockTab.h
+++ b/Engine/Source/Runtime/Slate/Public/Widgets/Docking/SDockTab.h
@@ -108,7 +108,10 @@ public:
 		SLATE_EVENT( FSimpleDelegate, OnTabRelocated )
 		SLATE_EVENT( FSimpleDelegate, OnTabDraggedOverDockArea )
 		SLATE_ARGUMENT( bool, ShouldAutosize )
+		/** When the close button is pressed, checks whether the tab can be closed in that moment. Example: Show dialog and ask user whether they're sure to close. */
 		SLATE_EVENT( FCanCloseTab, OnCanCloseTab )
+		/** Whether this tab can ever be closed. Example: certain programs may want to show tabs for the lifetime of the program. */
+		SLATE_ARGUMENT( bool, CanEverClose )
 		SLATE_EVENT( FOnPersistVisualState, OnPersistVisualState )
 		SLATE_EVENT( FExtendContextMenu, OnExtendContextMenu )
 		/** Invoked when a tab is closed from a drawer. This does not mean the tab or its contents is destroyed, just hidden. Use OnTabClosed for that */
@@ -230,7 +233,7 @@ public:
 	/** Set whether this tab should be sized based on its content. */
 	void SetShouldAutosize(const bool bNewShouldAutosize);
 
-	/** @return true if the tab can be closed */
+	/** @return true if the tab can be closed right now. Example: Callback could ask user via dialog. */
 	bool CanCloseTab() const;
 
 	/** Requests that the tab be closed.  Tabs may prevent closing depending on their state */	
@@ -415,6 +418,9 @@ protected:
 	/** Is this an MajorTab? A tool panel tab? */
 	ETabRole TabRole;
 
+	/** Determines whether the close button for the tab is shown. */
+	bool bCanEverClose;
+
 	/** The tab's parent tab well. Null if it is a floating tab. */
 	TWeakPtr<SDockingTabWell> ParentPtr;
 
