Fix tabs not closing

diff --git a/Engine/Source/Runtime/Slate/Public/Widgets/Docking/SDockTab.h b/Engine/Source/Runtime/Slate/Public/Widgets/Docking/SDockTab.h
--- a/Engine/Source/Runtime/Slate/Public/Widgets/Docking/SDockTab.h
+++ b/Engine/Source/Runtime/Slate/Public/Widgets/Docking/SDockTab.h
@@ -84,6 +84,7 @@ public:
 		, _OnTabDraggedOverDockArea()
 		, _ShouldAutosize(false)
 		, _OnCanCloseTab()
+		, _CanEverClose(true)
 		, _OnPersistVisualState()
 		, _TabColorScale(FLinearColor::Transparent)
 		, _ForegroundColor(FSlateColor::UseStyle())
