TreeView: Fix a bug where the hierarchy would flicker if you were at the end of the list and the number of stacked items was greater than the allowed maximum

diff --git a/Engine/Source/Runtime/Slate/Public/Widgets/Views/STreeView.h b/Engine/Source/Runtime/Slate/Public/Widgets/Views/STreeView.h
--- a/Engine/Source/Runtime/Slate/Public/Widgets/Views/STreeView.h
+++ b/Engine/Source/Runtime/Slate/Public/Widgets/Views/STreeView.h
@@ -768,8 +768,8 @@ public:
 					continue;
 				}
 
-				// If the hierarchy is big enough, AND the smallest we have found so far
-				if (NumParents - IndexOffset < MinSpaceOccupied)
+				// If hierarchy is the smallest we have found so far, AND the number of pinned items it will require is < the allowed max
+				if (NumParents - IndexOffset < MinSpaceOccupied && IndexOffset <= CurrentMaxPinnedItems)
 				{
 					MinSpaceOccupied = NumParents - IndexOffset;
 					MinIndex = ItemIndex;
