Speculative fix crash generating non visible parents in STreeView

diff --git a/Engine/Source/Runtime/Slate/Public/Widgets/Views/STreeView.h b/Engine/Source/Runtime/Slate/Public/Widgets/Views/STreeView.h
--- a/Engine/Source/Runtime/Slate/Public/Widgets/Views/STreeView.h
+++ b/Engine/Source/Runtime/Slate/Public/Widgets/Views/STreeView.h
@@ -710,6 +710,10 @@ public:
 		int32 CurrentItemIndex = StartIndex;
 
 		auto GetNonVisibleParents = [this, &InItemsSource, StartIndex](TArray<ItemType>& OutParents, int32 ItemIndex) {
+			if (DenseItemInfos.IsValidIndex(ItemIndex))
+			{
+				return;
+			}
 
 			int32 ParentIndex = ItemIndex;
 
@@ -719,7 +723,7 @@ public:
 				ParentIndex = DenseItemInfos[ItemIndex].ParentIndex;
 
 				// If the current item has a parent, and the parent is not visible, add the parent to the list of pinned items
-				if (ParentIndex != INDEX_NONE && ParentIndex < StartIndex)
+				if (InItemsSource.IsValidIndex(ParentIndex) && ParentIndex < StartIndex)
 				{
 					OutParents.Add(InItemsSource[ParentIndex]);
 				}
