MultiBox: Disable SelectAllTextWhenFocused for SSearchBox in SMultiBoxWidget

diff --git a/Engine/Source/Runtime/Slate/Private/Framework/MultiBox/MultiBox.cpp b/Engine/Source/Runtime/Slate/Private/Framework/MultiBox/MultiBox.cpp
--- a/Engine/Source/Runtime/Slate/Private/Framework/MultiBox/MultiBox.cpp
+++ b/Engine/Source/Runtime/Slate/Private/Framework/MultiBox/MultiBox.cpp
@@ -882,6 +882,7 @@ void SMultiBoxWidget::CreateSearchTextWidget()
 	SearchTextWidget = 
 		SNew(SSearchBox)
 		   .HintText(LOCTEXT("SearchHint", "Search"))
+		   .SelectAllTextWhenFocused(false)
 		   .OnTextChanged(this, &SMultiBoxWidget::OnFilterTextChanged);
 
 	TSharedRef< FWidgetBlock > NewWidgetBlock(new FWidgetBlock(SearchTextWidget.ToSharedRef(), FText::GetEmpty(), false));
