Remove obsolete compiler crash workaround
#rnx

diff --git a/Engine/Source/Runtime/Slate/Private/Widgets/Layout/SSplitter.cpp b/Engine/Source/Runtime/Slate/Private/Widgets/Layout/SSplitter.cpp
--- a/Engine/Source/Runtime/Slate/Private/Widgets/Layout/SSplitter.cpp
+++ b/Engine/Source/Runtime/Slate/Private/Widgets/Layout/SSplitter.cpp
@@ -1103,8 +1103,8 @@ void SSplitter2x2::ResizeChildren( const FGeometry& MyGeometry, const TArray<FLa
 		NewSizeTL.X = TopLeftSize.X + Delta.X;
 		NewSizeBL.X = BotLeftSize.X + Delta.X;
 
-		//workaround PS4 compiler crash.
-#if PLATFORM_PS4 || PLATFORM_ANDROID_X64
+		// @todo: revisit workaround for compiler crash on clang version > 9
+#if PLATFORM_ANDROID_X64
 		volatile float workaround = 0.0f;
 		NewSizeTR.X = TopRightSize.X - Delta.X + workaround;
 #else
