[Insights] Added LLM_SCOPE(ELLMTag::UI); in FSlateApplication::OnSizeChanged() in effort to reduce untagged allocations. It captures ~10000 untagged allocations (~10 MiB) in a short Editor run session.

diff --git a/Engine/Source/Runtime/Slate/Private/Framework/Application/SlateApplication.cpp b/Engine/Source/Runtime/Slate/Private/Framework/Application/SlateApplication.cpp
--- a/Engine/Source/Runtime/Slate/Private/Framework/Application/SlateApplication.cpp
+++ b/Engine/Source/Runtime/Slate/Private/Framework/Application/SlateApplication.cpp
@@ -6170,6 +6170,8 @@ void FSlateApplication::ProcessMotionDetectedEvent( const FMotionEvent& MotionEv
 
 bool FSlateApplication::OnSizeChanged( const TSharedRef< FGenericWindow >& PlatformWindow, const int32 Width, const int32 Height, bool bWasMinimized )
 {
+	LLM_SCOPE(ELLMTag::UI);
+
 	TSharedPtr< SWindow > Window = FSlateWindowHelper::FindWindowByPlatformWindow( SlateWindows, PlatformWindow );
 
 	if ( Window.IsValid() )
