Add a FPlatformUserId to the Slate User and define a default Platform User for the SlateApplication to use. There is no functionality difference here, it is just updating to use the FPlatformUserId struct wrapper instead of the raw int32 ControllerId concept. This is another piece in the large puzzle that is properly supporting local multiplayer input in UE.

#rnx

diff --git a/Engine/Source/Runtime/Slate/Private/Framework/Application/SlateApplication.cpp b/Engine/Source/Runtime/Slate/Private/Framework/Application/SlateApplication.cpp
--- a/Engine/Source/Runtime/Slate/Private/Framework/Application/SlateApplication.cpp
+++ b/Engine/Source/Runtime/Slate/Private/Framework/Application/SlateApplication.cpp
@@ -26,6 +26,8 @@
 #include "Input/Events.h"
 #include "Input/HittestGrid.h"
 #include "HAL/PlatformApplicationMisc.h"
+#include "GenericPlatform/GenericPlatformInputDeviceMapper.h"
+
 #if WITH_ACCESSIBILITY
 #include "Widgets/Accessibility/SlateAccessibleMessageHandler.h"
 #endif
@@ -809,7 +811,9 @@ FSlateApplication::FSlateApplication()
 	NormalExecutionGetter.BindRaw( this, &FSlateApplication::IsNormalExecution );
 
 	// Add the standard 'default' user (there's always guaranteed to be at least one)
-	RegisterNewUser(CursorUserIndex);
+	// The default cursor platform user id the primary platform user
+	ensure(SlateAppPrimaryPlatformUser.IsValid() && SlateAppPrimaryPlatformUser == IPlatformInputDeviceMapper::Get().GetPrimaryPlatformUser());
+	RegisterNewUser(SlateAppPrimaryPlatformUser);
 
 	NavigationConfig->OnRegister();
 #if WITH_EDITOR
diff --git a/Engine/Source/Runtime/Slate/Private/Framework/Application/SlateUser.cpp b/Engine/Source/Runtime/Slate/Private/Framework/Application/SlateUser.cpp
--- a/Engine/Source/Runtime/Slate/Private/Framework/Application/SlateUser.cpp
+++ b/Engine/Source/Runtime/Slate/Private/Framework/Application/SlateUser.cpp
@@ -98,6 +98,11 @@ TSharedRef<FSlateUser> FSlateUser::Create(int32 InUserIndex, TSharedPtr<ICursor>
 	return MakeShareable(new FSlateUser(InUserIndex, InCursor));
 }
 
+TSharedRef<FSlateUser> FSlateUser::Create(FPlatformUserId InPlatformUserId, TSharedPtr<ICursor> InCursor)
+{
+	return MakeShareable(new FSlateUser(InPlatformUserId, InCursor));
+}
+
 FSlateUser::FSlateUser(int32 InUserIndex, TSharedPtr<ICursor> InCursor)
 	: UserIndex(InUserIndex)
 	, Cursor(InCursor)
@@ -108,6 +113,19 @@ FSlateUser::FSlateUser(int32 InUserIndex, TSharedPtr<ICursor> InCursor)
 	PreviousPointerPositionsByIndex.Add(FSlateApplication::CursorPointerIndex, FVector2D::ZeroVector);
 }
 
+FSlateUser::FSlateUser(FPlatformUserId InPlatformUser, TSharedPtr<ICursor> InCursor)
+	: PlatformUser(InPlatformUser)
+	, Cursor(InCursor)
+{
+	// TODO: Remove this part, its backwards compatible for now
+	UserIndex = InPlatformUser.GetInternalId();
+	
+	UE_LOG(LogSlate, Log, TEXT("New Slate User Created.  Platform User Id %d,  Old User Index: %d  , Is Virtual User: %d"), PlatformUser.GetInternalId(), UserIndex, IsVirtualUser());
+	
+	PointerPositionsByIndex.Add(FSlateApplication::CursorPointerIndex, FVector2D::ZeroVector);
+	PreviousPointerPositionsByIndex.Add(FSlateApplication::CursorPointerIndex, FVector2D::ZeroVector);
+}
+
 FSlateUser::~FSlateUser()
 {
 	UE_LOG(LogSlate, Log, TEXT("Slate User Destroyed.  User Index %d, Is Virtual User: %d"), UserIndex, IsVirtualUser());
diff --git a/Engine/Source/Runtime/Slate/Public/Framework/Application/SlateUser.h b/Engine/Source/Runtime/Slate/Public/Framework/Application/SlateUser.h
--- a/Engine/Source/Runtime/Slate/Public/Framework/Application/SlateUser.h
+++ b/Engine/Source/Runtime/Slate/Public/Framework/Application/SlateUser.h
@@ -43,6 +43,7 @@ public:
 	virtual ~FSlateUser();
 
 	FORCEINLINE int32 GetUserIndex() const { return UserIndex; }
+	FORCEINLINE FPlatformUserId GetPlatformUserId() const { return PlatformUser; }
 	FORCEINLINE bool IsVirtualUser() const { return !Cursor.IsValid(); }
 
 	TSharedPtr<SWidget> GetFocusedWidget() const;
@@ -125,6 +126,7 @@ public:
 
 SLATE_SCOPE:
 	static TSharedRef<FSlateUser> Create(int32 InUserIndex, TSharedPtr<ICursor> InCursor);
+	static TSharedRef<FSlateUser> Create(FPlatformUserId InPlatformUserId, TSharedPtr<ICursor> InCursor);
 	
 	FORCEINLINE bool HasValidFocusPath() const { return WeakFocusPath.IsValid(); }
 	FORCEINLINE const FWeakWidgetPath& GetWeakFocusPath() const { return WeakFocusPath; }
@@ -191,6 +193,8 @@ SLATE_SCOPE:
 
 private:
 	FSlateUser(int32 InUserIndex, TSharedPtr<ICursor> InCursor);
+	
+	FSlateUser(FPlatformUserId InPlatformUser, TSharedPtr<ICursor> InCursor);
 	void UpdatePointerPosition(uint32 PointerIndex, const FVector2D& Position);
 	void LockCursorInternal(const FWidgetPath& WidgetPath);
 	TSharedRef<SWindow> GetOrCreateTooltipWindow();
@@ -198,6 +202,9 @@ private:
 	/** The index the user was assigned. */
 	int32 UserIndex = INDEX_NONE;
 
+	/** The owning platform user of this slate user. */
+	FPlatformUserId PlatformUser = PLATFORMUSERID_NONE;
+
 	/** The cursor this user is in control of. Guaranteed to be valid for all real users, absence implies this is a virtual user. */
 	TSharedPtr<ICursor> Cursor;
 
diff --git a/Engine/Source/Runtime/SlateCore/Private/Application/SlateApplicationBase.cpp b/Engine/Source/Runtime/SlateCore/Private/Application/SlateApplicationBase.cpp
--- a/Engine/Source/Runtime/SlateCore/Private/Application/SlateApplicationBase.cpp
+++ b/Engine/Source/Runtime/SlateCore/Private/Application/SlateApplicationBase.cpp
@@ -18,6 +18,7 @@ TSharedPtr<GenericApplication> FSlateApplicationBase::PlatformApplication = null
 // TODO: Identifier the cursor index in a smarter way.
 const uint32 FSlateApplicationBase::CursorPointerIndex = ETouchIndex::CursorPointerIndex;
 const uint32 FSlateApplicationBase::CursorUserIndex = 0;
+const FPlatformUserId FSlateApplicationBase::SlateAppPrimaryPlatformUser = FPlatformUserId::CreateFromInternalId(0);
 
 FWidgetPath FHitTesting::LocateWidgetInWindow(FVector2D ScreenspaceMouseCoordinate, const TSharedRef<SWindow>& Window, bool bIgnoreEnabledStatus, int32 UserIndex) const
 {
diff --git a/Engine/Source/Runtime/SlateCore/Public/Application/SlateApplicationBase.h b/Engine/Source/Runtime/SlateCore/Public/Application/SlateApplicationBase.h
--- a/Engine/Source/Runtime/SlateCore/Public/Application/SlateApplicationBase.h
+++ b/Engine/Source/Runtime/SlateCore/Public/Application/SlateApplicationBase.h
@@ -521,6 +521,9 @@ public:
 	const static uint32 CursorPointerIndex;
 	const static uint32 CursorUserIndex;
 
+	/** The primary platform user for this application. They will own the mouse and cursor by default. Default ID is 0 */
+	const static FPlatformUserId SlateAppPrimaryPlatformUser;
+
 	/**
 	 * Returns the current instance of the application. The application should have been initialized before
 	 * this method is called
