Fix sidebar drawers preventing drag actions

diff --git a/Engine/Source/Runtime/Slate/Private/Framework/Docking/STabSidebar.cpp b/Engine/Source/Runtime/Slate/Private/Framework/Docking/STabSidebar.cpp
--- a/Engine/Source/Runtime/Slate/Private/Framework/Docking/STabSidebar.cpp
+++ b/Engine/Source/Runtime/Slate/Private/Framework/Docking/STabSidebar.cpp
@@ -824,6 +824,12 @@ void STabSidebar::CloseDrawerInternal(TSharedRef<SDockTab> ForTab)
 	if (TSharedPtr<STabDrawer> OpenedDrawer = FindOpenedDrawer(ForTab))
 	{
 		OpenedDrawer->Close();
+		TSharedRef<STabDrawer> OpenedDrawerRef = OpenedDrawer.ToSharedRef();
+
+		bool bRemoveSuccessful = DrawersOverlay->RemoveSlot(OpenedDrawerRef);
+		ensure(bRemoveSuccessful);
+
+		OpenedDrawers.Remove(OpenedDrawerRef);
 	}
 
 	SummonPinnedTabIfNothingOpened();
