Fix certain context menus disappear the second time a submenu within them gets dismissed

diff --git a/Engine/Source/Runtime/Slate/Private/Framework/MultiBox/SClippingHorizontalBox.cpp b/Engine/Source/Runtime/Slate/Private/Framework/MultiBox/SClippingHorizontalBox.cpp
--- a/Engine/Source/Runtime/Slate/Private/Framework/MultiBox/SClippingHorizontalBox.cpp
+++ b/Engine/Source/Runtime/Slate/Private/Framework/MultiBox/SClippingHorizontalBox.cpp
@@ -96,6 +96,7 @@ void SClippingHorizontalBox::AddWrapButton()
 	const FToolBarStyle& ToolBarStyle = StyleSet->GetWidgetStyle<FToolBarStyle>(StyleName);
 
 	// Construct the wrap button used in toolbars and menubars
+	// Always allow this to be focusable to prevent the menu from collapsing during interaction
 	WrapButton = 
 		SNew( SComboButton )
 		.HasDownArrow( false )
@@ -105,7 +106,7 @@ void SClippingHorizontalBox::AddWrapButton()
 		.OnGetMenuContent( OnWrapButtonClicked )
 		.Cursor( EMouseCursor::Default )
 		.OnMenuOpenChanged(this, &SClippingHorizontalBox::OnWrapButtonOpenChanged)
-		.IsFocusable(bIsFocusable)
+		.IsFocusable(true)
 		.ButtonContent()
 		[
 			SNew(SImage)
