Slate:
- Add support for changing value of a SSpinBox using the mouse wheel (Enabled using the EnableSlider flag)
- Add new meta data to specify the step value: WheelStep (Default step is 1.0 or 0.1 if clamp range is <= 10.0)
- Step can be modified by holding Control (1/10) or Shift (10x)

diff --git a/Engine/Source/Editor/PropertyEditor/Private/UserInterface/PropertyEditor/SPropertyEditorNumeric.h b/Engine/Source/Editor/PropertyEditor/Private/UserInterface/PropertyEditor/SPropertyEditorNumeric.h
--- a/Engine/Source/Editor/PropertyEditor/Private/UserInterface/PropertyEditor/SPropertyEditorNumeric.h
+++ b/Engine/Source/Editor/PropertyEditor/Private/UserInterface/PropertyEditor/SPropertyEditorNumeric.h
@@ -331,6 +331,16 @@ public:
 				UE_LOG(LogPropertyNode, Warning, TEXT("UI Min (%s) >= UI Max (%s) for Ranged Numeric property %s"), *UIMinString, *UIMaxString, *Property->GetPathName());
 			}
 
+			const FString& WheelStepString = GetMetaDataFromKey("WheelStep");
+			NumericType WheelStepValue = NumericType(0);
+
+			if (!WheelStepString.IsEmpty())
+			{
+				TTypeFromString<NumericType>::FromString(WheelStepValue, *WheelStepString);
+			}
+
+			TOptional<NumericType> WheelStep = WheelStepString.Len() ? WheelStepValue : TOptional<NumericType>();
+
 			FObjectPropertyNode* ObjectPropertyNode = PropertyNode->FindObjectItemParent();
 			const bool bAllowSpin = (!ObjectPropertyNode || (1 == ObjectPropertyNode->GetNumObjects()))
 				&& !PropertyHandle->GetBoolMetaData("NoSpinbox");
@@ -392,6 +402,8 @@ public:
 				.Delta(Delta)
 				// LinearDeltaSensitivity needs to be left unset if not provided, rather than being set to some default
 				.LinearDeltaSensitivity(LinearDeltaSensitivity != 0 ? LinearDeltaSensitivity : TAttribute<int32>())
+				.AllowWheel(bAllowSpin)
+				.WheelStep(WheelStep)
 				.UndeterminedString(LOCTEXT("MultipleValues", "Multiple Values"))
 				.OnValueChanged(this, &SPropertyEditorNumeric<NumericType>::OnValueChanged)
 				.OnValueCommitted(this, &SPropertyEditorNumeric<NumericType>::OnValueCommitted)
diff --git a/Engine/Source/Runtime/Slate/Public/Widgets/Input/SNumericEntryBox.h b/Engine/Source/Runtime/Slate/Public/Widgets/Input/SNumericEntryBox.h
--- a/Engine/Source/Runtime/Slate/Public/Widgets/Input/SNumericEntryBox.h
+++ b/Engine/Source/Runtime/Slate/Public/Widgets/Input/SNumericEntryBox.h
@@ -86,14 +86,15 @@ public:
 		, _Delta(NumericType(0))
 		, _MinValue(TNumericLimits<NumericType>::Lowest())
 		, _MaxValue(TNumericLimits<NumericType>::Max())
-		, _MinSliderValue(NumericType(0))				
+		, _MinSliderValue(NumericType(0))
 		, _MaxSliderValue(NumericType(100))
 		, _SliderExponent(1.f)
+		, _AllowWheel(true)
 		, _MinDesiredValueWidth(0.f)
 		, _DisplayToggle(false)
 		, _ToggleChecked(ECheckBoxState::Checked)
 		, _TogglePadding(FMargin(1.f,0.f,1.f,0.f) )
-	{}		
+	{}
 
 		/** Style to use for the editable text box within this widget */
 		SLATE_STYLE_ARGUMENT( FEditableTextBoxStyle, EditableTextBoxStyle )
@@ -149,6 +150,10 @@ public:
 		SLATE_ATTRIBUTE( float, SliderExponent )
 		/** When using exponential scale specify a neutral value where we want the maximum precision (by default it is the smallest slider value)*/
 		SLATE_ATTRIBUTE(NumericType, SliderExponentNeutralValue )
+		/** Whether this spin box should have mouse wheel feature enabled, defaults to true */
+		SLATE_ARGUMENT( bool, AllowWheel )
+		/** Step to increment or decrement the value by when scrolling the mouse wheel. If not specified will determine automatically */
+		SLATE_ATTRIBUTE( TOptional< NumericType >, WheelStep )
 		/** The minimum desired width for the value portion of the control. */
 		SLATE_ATTRIBUTE( float, MinDesiredValueWidth )
 		/** The text margin to use if overridden. */
@@ -238,6 +243,8 @@ public:
 				.MinValue(InArgs._MinValue)
 				.SliderExponent(InArgs._SliderExponent)
 				.SliderExponentNeutralValue(InArgs._SliderExponentNeutralValue)
+				.EnableWheel(InArgs._AllowWheel)
+				.WheelStep(InArgs._WheelStep)
 				.OnBeginSliderMovement(InArgs._OnBeginSliderMovement)
 				.OnEndSliderMovement(InArgs._OnEndSliderMovement)
 				.MinDesiredWidth(InArgs._MinDesiredValueWidth)
diff --git a/Engine/Source/Runtime/Slate/Public/Widgets/Input/SSpinBox.h b/Engine/Source/Runtime/Slate/Public/Widgets/Input/SSpinBox.h
--- a/Engine/Source/Runtime/Slate/Public/Widgets/Input/SSpinBox.h
+++ b/Engine/Source/Runtime/Slate/Public/Widgets/Input/SSpinBox.h
@@ -80,6 +80,7 @@ public:
 		, _SupportDynamicSliderMaxValue(false)
 		, _SupportDynamicSliderMinValue(false)
 		, _SliderExponent(1.f)
+		, _EnableWheel(true)
 		, _Font(FCoreStyle::Get().GetFontStyle(TEXT("NormalFont")))
 		, _ContentPadding(FMargin(2.0f, 1.0f))
 		, _OnValueChanged()
@@ -128,6 +129,10 @@ public:
 		SLATE_ATTRIBUTE(float, SliderExponent)
 		/** When use exponential scale for the slider which is the neutral value */
 		SLATE_ATTRIBUTE(NumericType, SliderExponentNeutralValue)
+		/** Whether this spin box should have mouse wheel feature enabled, defaults to true */
+		SLATE_ARGUMENT(bool, EnableWheel)
+		/** Step to increment or decrement the value by when scrolling the mouse wheel. If not specified will determine automatically */
+		SLATE_ATTRIBUTE(TOptional< NumericType >, WheelStep)
 		/** Font used to display text in the slider */
 		SLATE_ATTRIBUTE(FSlateFontInfo, Font)
 		/** Padding to add around this widget and its internal widgets */
@@ -200,6 +205,9 @@ public:
 		OnDynamicSliderMaxValueChanged = InArgs._OnDynamicSliderMaxValueChanged;
 		OnDynamicSliderMinValueChanged = InArgs._OnDynamicSliderMinValueChanged;
 
+		bEnableWheel = InArgs._EnableWheel;
+		WheelStep = InArgs._WheelStep;
+
 		bPreventThrottling = InArgs._PreventThrottling;
 
 		CachedExternalValue = ValueAttribute.Get();
@@ -624,15 +632,15 @@ public:
 				else
 				{
 					// If this control has a specified delta and sensitivity then we use that instead of the current value for determining how much to change.
-					const float Sign = (MouseEvent.GetCursorDelta().X > 0) ? 1.f : -1.f;
+					const double Sign = (MouseEvent.GetCursorDelta().X > 0) ? 1.0 : -1.0;
 					if (LinearDeltaSensitivity.IsSet() && LinearDeltaSensitivity.Get() != 0 && Delta.IsSet() && Delta.Get() > 0)
 					{
-						const float MouseDelta = FMath::Abs(MouseEvent.GetCursorDelta().X / LinearDeltaSensitivity.Get());
-						NewValue = InternalValue + (Sign * MouseDelta * FMath::Pow((float)Delta.Get(), SliderExponent.Get()));
+						const double MouseDelta = FMath::Abs(MouseEvent.GetCursorDelta().X / LinearDeltaSensitivity.Get());
+						NewValue = InternalValue + (Sign * MouseDelta * FMath::Pow((double)Delta.Get(), SliderExponent.Get()));
 					}
 					else
 					{
-						const float MouseDelta = FMath::Abs(MouseEvent.GetCursorDelta().X / SliderWidthInSlateUnits);
+						const double MouseDelta = FMath::Abs(MouseEvent.GetCursorDelta().X / SliderWidthInSlateUnits);
 						const double CurrentValue = FMath::Clamp<double>(FMath::Abs(InternalValue), 1.0, (double)std::numeric_limits<NumericType>::max());
 						NewValue = InternalValue + (Sign * MouseDelta * FMath::Pow((float)CurrentValue, SliderExponent.Get()));
 					}
@@ -649,6 +657,49 @@ public:
 		return FReply::Unhandled();
 	}
 
+	/**
+	 * Called when the mouse wheel is spun. This event is bubbled.
+	 *
+	 * @param  MouseEvent  Mouse event
+	 * @return  Returns whether the event was handled, along with other possible actions
+	 */
+	virtual FReply OnMouseWheel(const FGeometry& MyGeometry, const FPointerEvent& MouseEvent) override
+	{
+		if (bEnableWheel && PointerDraggingSliderIndex == INDEX_NONE && HasKeyboardFocus())
+		{
+			// If there is no WheelStep defined, we use 1.0 (Or 0.1 if slider range is <= 10)
+			constexpr bool bIsIntegral = TIsIntegral<NumericType>::Value;
+			const bool bIsSmallStep = !bIsIntegral && (GetMaxSliderValue() - GetMinSliderValue()) <= 10.0;
+			double Step = WheelStep.IsSet() && WheelStep.Get().IsSet() ? WheelStep.Get().GetValue() : (bIsSmallStep ? 0.1 : 1.0);
+
+			if (MouseEvent.IsControlDown())
+			{
+				// If no value is set for WheelSmallStep, we use the DefaultStep divided by 10
+				Step /= 10.0;
+			}
+			else if (MouseEvent.IsShiftDown())
+			{
+				// If no value is set for WheelBigStep, we use the DefaultStep multiplied by 10
+				Step *= 10.0;
+			}
+
+			const double Sign = (MouseEvent.GetWheelDelta() > 0) ? 1.0 : -1.0;
+			const double NewValue = InternalValue + (Sign * Step);
+			const NumericType RoundedNewValue = RoundIfIntegerValue(NewValue);
+
+			// First SetText is to update the value before calling CommitValue. Otherwise, when the text lose
+			// focus from the CommitValue, it will override the value we just committed.
+			// The second SetText is to update the text to the InternalValue since it could have been clamped.
+			EditableText->SetText(FText::FromString(Interface->ToString(NewValue)));
+			CommitValue(RoundedNewValue, NewValue, CommittedViaSpin, ETextCommit::OnEnter);
+			EditableText->SetText(FText::FromString(Interface->ToString(InternalValue)));
+
+			return FReply::Handled();
+		}
+
+		return FReply::Unhandled();
+	}
+
 	virtual FCursorReply OnCursorQuery(const FGeometry& MyGeometry, const FPointerEvent& CursorEvent) const override
 	{
 		const bool bEnableSlider = GetEnableSlider();
@@ -1076,6 +1127,7 @@ private:
 	TAttribute<bool> EnableSlider;
 	TAttribute<bool> SupportDynamicSliderMaxValue;
 	TAttribute<bool> SupportDynamicSliderMinValue;
+	TAttribute< TOptional<NumericType> > WheelStep;
 	FOnDynamicSliderMinMaxValueChanged OnDynamicSliderMaxValueChanged;
 	FOnDynamicSliderMinMaxValueChanged OnDynamicSliderMinValueChanged;
 
@@ -1151,6 +1203,9 @@ private:
 	 * When true, the viewport will be updated with every single change to the value during dragging
 	 */
 	bool bPreventThrottling;
+
+	/** Does this spin box have the mouse wheel feature enabled? */
+	bool bEnableWheel = true;
 };
 
 template<typename NumericType>
