FAsyncTaskNotification: Fix a crash when you tried to update a notification from another thread

diff --git a/Engine/Source/Runtime/Slate/Private/Framework/Notifications/SlateAsyncTaskNotificationImpl.cpp b/Engine/Source/Runtime/Slate/Private/Framework/Notifications/SlateAsyncTaskNotificationImpl.cpp
--- a/Engine/Source/Runtime/Slate/Private/Framework/Notifications/SlateAsyncTaskNotificationImpl.cpp
+++ b/Engine/Source/Runtime/Slate/Private/Framework/Notifications/SlateAsyncTaskNotificationImpl.cpp
@@ -371,9 +371,8 @@ void FSlateAsyncTaskNotificationImpl::UpdateNotification()
 		// Update the notification UI only if the state hasn't changed (i.e this notification will not be deleted)
 		if(OwningNotification && State == PreviousCompletionState)
 		{
-			OwningNotification->SetText(TitleText);
-			OwningNotification->SetSubText(ProgressText);
-			OwningNotification->SetHyperlink(Hyperlink, HyperlinkText);
+			/* Slate requries the notification to be updated from the main thread, so we add a one frame ticker for it */
+			FTSTicker::GetCoreTicker().AddTicker(FTickerDelegate::CreateStatic(&FSlateAsyncTaskNotificationImpl::UpdateNotificationDeferred, OwningNotification, TitleText, ProgressText, Hyperlink, HyperlinkText));
 		}
 
 		// Set the Pending Completion State in case the notification has to change
@@ -381,6 +380,16 @@ void FSlateAsyncTaskNotificationImpl::UpdateNotification()
 	}
 }
 
+bool FSlateAsyncTaskNotificationImpl::UpdateNotificationDeferred(float InDeltaTime, TSharedPtr<SNotificationItem> OwningNotification, FText TitleText, FText ProgressText, FSimpleDelegate Hyperlink, FText HyperlinkText)
+{
+	OwningNotification->SetText(TitleText);
+	OwningNotification->SetSubText(ProgressText);
+	OwningNotification->SetHyperlink(Hyperlink, HyperlinkText);
+
+	// We only want this function to tick once
+	return false;
+}
+
 EAsyncTaskNotificationPromptAction FSlateAsyncTaskNotificationImpl::GetPromptAction() const
 {
 	if(NotificationConfig.bIsHeadless)
diff --git a/Engine/Source/Runtime/Slate/Private/Framework/Notifications/SlateAsyncTaskNotificationImpl.h b/Engine/Source/Runtime/Slate/Private/Framework/Notifications/SlateAsyncTaskNotificationImpl.h
--- a/Engine/Source/Runtime/Slate/Private/Framework/Notifications/SlateAsyncTaskNotificationImpl.h
+++ b/Engine/Source/Runtime/Slate/Private/Framework/Notifications/SlateAsyncTaskNotificationImpl.h
@@ -96,6 +96,9 @@ private:
 	/* Cleanly destroy the current notification */
 	void DestroyCurrentNotification();
 
+	/* Static function to update the notification from the main thread */
+	static bool UpdateNotificationDeferred(float InDeltaTime, TSharedPtr<SNotificationItem> OwningNotification, FText TitleText, FText ProgressText, FSimpleDelegate Hyperlink, FText HyperlinkText);
+
 private:
 
 	/** The Config used for all notifications */
