Update SlateApplication and the message handler to take in FInputDeviceId and FPlatformId instead of the old int32 controllerid for controller input

diff --git a/Engine/Source/Runtime/ApplicationCore/Public/GenericPlatform/GenericApplicationMessageHandler.h b/Engine/Source/Runtime/ApplicationCore/Public/GenericPlatform/GenericApplicationMessageHandler.h
--- a/Engine/Source/Runtime/ApplicationCore/Public/GenericPlatform/GenericApplicationMessageHandler.h
+++ b/Engine/Source/Runtime/ApplicationCore/Public/GenericPlatform/GenericApplicationMessageHandler.h
@@ -300,10 +300,10 @@ public:
 	 */
 	virtual bool ShouldUsePlatformUserId() const
 	{
-		return false;
+		return true;
 	}
 
-	virtual bool OnControllerAnalog(FGamepadKeyNames::Type KeyName, FPlatformUserId PlatformUserId, float AnalogValue)
+	virtual bool OnControllerAnalog(FGamepadKeyNames::Type KeyName, FPlatformUserId PlatformUserId, FInputDeviceId InputDeviceId, float AnalogValue)
 	{
 		if (!ShouldUsePlatformUserId())
 		{
@@ -312,7 +312,7 @@ public:
 		return false;
 	}
 
-	virtual bool OnControllerButtonPressed(FGamepadKeyNames::Type KeyName, FPlatformUserId PlatformUserId, bool IsRepeat)
+	virtual bool OnControllerButtonPressed(FGamepadKeyNames::Type KeyName, FPlatformUserId PlatformUserId, FInputDeviceId InputDeviceId, bool IsRepeat)
 	{
 		if (!ShouldUsePlatformUserId())
 		{
@@ -321,7 +321,7 @@ public:
 		return false;
 	}
 
-	virtual bool OnControllerButtonReleased(FGamepadKeyNames::Type KeyName, FPlatformUserId PlatformUserId, bool IsRepeat)
+	virtual bool OnControllerButtonReleased(FGamepadKeyNames::Type KeyName, FPlatformUserId PlatformUserId, FInputDeviceId InputDeviceId, bool IsRepeat)
 	{
 		if (!ShouldUsePlatformUserId())
 		{
@@ -343,7 +343,7 @@ public:
     {
     }
 
-	virtual bool OnTouchStarted( const TSharedPtr< FGenericWindow >& Window, const FVector2D& Location, float Force, int32 TouchIndex, FPlatformUserId PlatformUserId )
+	virtual bool OnTouchStarted( const TSharedPtr< FGenericWindow >& Window, const FVector2D& Location, float Force, int32 TouchIndex, FPlatformUserId PlatformUserId, FInputDeviceId DeviceId )
 	{
 		if (!ShouldUsePlatformUserId())
 		{
@@ -352,7 +352,7 @@ public:
 		return false;
 	}
 
-	virtual bool OnTouchMoved( const FVector2D& Location, float Force, int32 TouchIndex, FPlatformUserId PlatformUserId )
+	virtual bool OnTouchMoved( const FVector2D& Location, float Force, int32 TouchIndex, FPlatformUserId PlatformUserId, FInputDeviceId DeviceID )
 	{
 		if (!ShouldUsePlatformUserId())
 		{
@@ -361,7 +361,7 @@ public:
 		return false;
 	}
 
-	virtual bool OnTouchEnded( const FVector2D& Location, int32 TouchIndex, FPlatformUserId PlatformUserId )
+	virtual bool OnTouchEnded( const FVector2D& Location, int32 TouchIndex, FPlatformUserId PlatformUserId, FInputDeviceId DeviceID )
 	{
 		if (!ShouldUsePlatformUserId())
 		{
@@ -370,7 +370,7 @@ public:
 		return false;
 	}
 
-	virtual bool OnTouchForceChanged(const FVector2D& Location, float Force, int32 TouchIndex, FPlatformUserId PlatformUserId)
+	virtual bool OnTouchForceChanged(const FVector2D& Location, float Force, int32 TouchIndex, FPlatformUserId PlatformUserId, FInputDeviceId DeviceID)
 	{
 		if (!ShouldUsePlatformUserId())
 		{
@@ -379,7 +379,7 @@ public:
 		return false;
 	}
 
-	virtual bool OnTouchFirstMove(const FVector2D& Location, float Force, int32 TouchIndex, FPlatformUserId PlatformUserId)
+	virtual bool OnTouchFirstMove(const FVector2D& Location, float Force, int32 TouchIndex, FPlatformUserId PlatformUserId, FInputDeviceId DeviceID)
 	{
 		if (!ShouldUsePlatformUserId())
 		{
@@ -393,7 +393,7 @@ public:
 
 	}
 
-	virtual bool OnMotionDetected( const FVector& Tilt, const FVector& RotationRate, const FVector& Gravity, const FVector& Acceleration, FPlatformUserId PlatformUserId )
+	virtual bool OnMotionDetected( const FVector& Tilt, const FVector& RotationRate, const FVector& Gravity, const FVector& Acceleration, FPlatformUserId PlatformUserId, FInputDeviceId InputDeviceId )
 	{
 		if (!ShouldUsePlatformUserId())
 		{
@@ -513,14 +513,16 @@ public:
 	}
 
 	// Deprecate these when engine code has been converted to handle platform user id
+	//UE_DEPRECATED(5.1, "This version of OnControllerAnalog has been deprecated, please use the one that takes an FPlatformUser and FInputDeviceId instead.")
 	virtual bool OnControllerAnalog(FGamepadKeyNames::Type KeyName, int32 ControllerId, float AnalogValue)
 	{
 		if (ShouldUsePlatformUserId())
 		{
-			return OnControllerAnalog(KeyName, FPlatformUserId(ControllerId), AnalogValue);
+			return OnControllerAnalog(KeyName, FPlatformUserId(ControllerId), INPUTDEVICEID_NONE, AnalogValue);
 		}
 		return false;
 	}
+	//UE_DEPRECATED(5.1, "This version of OnControllerButtonPressed has been deprecated, please use the one that takes an FPlatformUser and FInputDeviceId instead.")
 	virtual bool OnControllerButtonPressed(FGamepadKeyNames::Type KeyName, int32 ControllerId, bool IsRepeat)
 	{
 		if (ShouldUsePlatformUserId())
@@ -529,6 +531,7 @@ public:
 		}
 		return false;
 	}
+	//UE_DEPRECATED(5.1, "This version of OnControllerButtonReleased has been deprecated, please use the one that takes an FPlatformUser and FInputDeviceId instead.")
 	virtual bool OnControllerButtonReleased(FGamepadKeyNames::Type KeyName, int32 ControllerId, bool IsRepeat)
 	{
 		if (ShouldUsePlatformUserId())
@@ -541,7 +544,7 @@ public:
 	{
 		if (ShouldUsePlatformUserId())
 		{
-			return OnTouchStarted(Window, Location, Force, TouchIndex, FPlatformUserId(ControllerId));
+			return OnTouchStarted(Window, Location, Force, TouchIndex, FPlatformUserId(ControllerId), INPUTDEVICEID_NONE);
 		}
 		return false;
 	}
@@ -549,7 +552,7 @@ public:
 	{
 		if (ShouldUsePlatformUserId())
 		{
-			return OnTouchMoved(Location, Force, TouchIndex, FPlatformUserId(ControllerId));
+			return OnTouchMoved(Location, Force, TouchIndex, FPlatformUserId(ControllerId), INPUTDEVICEID_NONE);
 		}
 		return false;
 	}
@@ -557,7 +560,7 @@ public:
 	{
 		if (ShouldUsePlatformUserId())
 		{
-			return OnTouchEnded(Location, TouchIndex, FPlatformUserId(ControllerId));
+			return OnTouchEnded(Location, TouchIndex, FPlatformUserId(ControllerId), INPUTDEVICEID_NONE);
 		}
 		return false;
 	}
@@ -565,7 +568,7 @@ public:
 	{
 		if (ShouldUsePlatformUserId())
 		{
-			return OnTouchForceChanged(Location, Force, TouchIndex, FPlatformUserId(ControllerId));
+			return OnTouchForceChanged(Location, Force, TouchIndex, FPlatformUserId(ControllerId), INPUTDEVICEID_NONE);
 		}
 		return false;
 	}
@@ -573,7 +576,7 @@ public:
 	{
 		if (ShouldUsePlatformUserId())
 		{
-			return OnTouchFirstMove(Location, Force, TouchIndex, FPlatformUserId(ControllerId));
+			return OnTouchFirstMove(Location, Force, TouchIndex, FPlatformUserId(ControllerId), INPUTDEVICEID_NONE);
 		}
 		return false;
 	}
@@ -581,7 +584,7 @@ public:
 	{
 		if (ShouldUsePlatformUserId())
 		{
-			return OnMotionDetected(Tilt, RotationRate, Gravity, Acceleration, FPlatformUserId(ControllerId));
+			return OnMotionDetected(Tilt, RotationRate, Gravity, Acceleration, FPlatformUserId(ControllerId), INPUTDEVICEID_NONE);
 		}
 		return false;
 	}
diff --git a/Engine/Source/Runtime/Slate/Private/Framework/Application/SlateApplication.cpp b/Engine/Source/Runtime/Slate/Private/Framework/Application/SlateApplication.cpp
--- a/Engine/Source/Runtime/Slate/Private/Framework/Application/SlateApplication.cpp
+++ b/Engine/Source/Runtime/Slate/Private/Framework/Application/SlateApplication.cpp
@@ -2408,7 +2408,17 @@ int32 FSlateApplication::GetUserIndexForKeyboard() const
 {
 	return InputManager->GetUserIndexForKeyboard();
 }
- 
+
+int32 FSlateApplication::GetInputDeviceIdForMouse() const
+{
+	return InputManager->GetInputDeviceIdForMouse();
+}
+
+int32 FSlateApplication::GetInputDeviceIdForKeyboard() const
+{
+	return InputManager->GetInputDeviceIdForKeyboard();
+}
+
 TOptional<int32> FSlateApplication::GetUserIndexForController(int32 ControllerId, FKey InKey) const
 {
 	return InputManager->GetUserIndexForController(ControllerId, InKey);
@@ -2419,6 +2429,16 @@ int32 FSlateApplication::GetUserIndexForController(int32 ControllerId) const
 	return InputManager->GetUserIndexForController(ControllerId);
 }
 
+int32 FSlateApplication::GetUserIndexForInputDevice(FInputDeviceId InputDeviceId) const
+{
+	return InputManager->GetUserIndexForInputDevice(InputDeviceId);
+}
+
+int32 FSlateApplication::GetUserIndexForPlatformUser(FPlatformUserId PlatformUser) const
+{
+	return InputManager->GetUserIndexForPlatformUser(PlatformUser);
+}
+
 void FSlateApplication::SetInputManager(TSharedRef<ISlateInputManager> InInputManager)
 {
 	//@todo DanH/NickD: Should we be worried at all about potential invalidation of SlateUsers we've created but that no longer occupy an index that can be referenced?
@@ -3992,8 +4012,24 @@ TSharedRef<FSlateUser> FSlateApplication::GetOrCreateUser(int32 UserIndex)
 	return RegisterNewUser(UserIndex);
 }
 
+TSharedRef<FSlateUser> FSlateApplication::GetOrCreateUser(FPlatformUserId PlatformUserId)
+{
+	if (TSharedPtr<FSlateUser> FoundUser = GetUser(PlatformUserId))
+	{
+		return FoundUser.ToSharedRef();
+	}
+	return RegisterNewUser(PlatformUserId);
+}
+
 TSharedRef<FSlateUser> FSlateApplication::RegisterNewUser(int32 UserIndex, bool bIsVirtual)
 {
+	return RegisterNewUser(FGenericPlatformMisc::GetPlatformUserForUserIndex(UserIndex), bIsVirtual);
+}
+
+TSharedRef<FSlateUser> FSlateApplication::RegisterNewUser(FPlatformUserId PlatformUserId, bool bIsVirtual)
+{
+	int32 UserIndex = PlatformUserId.GetInternalId();
+	
 	// We tolerate no shenanigans with inappropriate arguments here
 	// New users must be registered at a valid non-negative index that is not already occupied by another user
 	check(UserIndex >= 0);
@@ -5912,42 +5948,74 @@ bool FSlateApplication::ExecuteNavigation(const FWidgetPath& NavigationSource, T
 	return bHandled;
 }
 
-bool FSlateApplication::OnControllerAnalog( FGamepadKeyNames::Type KeyName, int32 ControllerId, float AnalogValue )
+bool FSlateApplication::OnControllerAnalog(FGamepadKeyNames::Type KeyName, FPlatformUserId PlatformUserId, FInputDeviceId InputDeviceId, float AnalogValue)
 {
 	FKey Key(KeyName);
-	TOptional<int32> UserIndex = GetUserIndexForController(ControllerId, Key);
-	if(UserIndex.IsSet() && ensureMsgf(Key.IsValid(), TEXT("OnControllerAnalog(KeyName=%s,ControllerId=%d,AnalogValue=%f) key is invalid"), *KeyName.ToString(), ControllerId, AnalogValue))
+	TOptional<int32> UserIndex = GetUserIndexForInputDevice(InputDeviceId);
+	if (UserIndex.IsSet() && ensureMsgf(Key.IsValid(), TEXT("OnControllerAnalog(KeyName=%s,InputDeviceId=%d,AnalogValue=%f) key is invalid"), *KeyName.ToString(), InputDeviceId.GetId(), AnalogValue))
 	{
-		FAnalogInputEvent AnalogInputEvent(Key, PlatformApplication->GetModifierKeys(), UserIndex.GetValue(), false, 0, 0, AnalogValue);
+		FAnalogInputEvent AnalogInputEvent(Key, PlatformApplication->GetModifierKeys(), InputDeviceId, false, 0, 0, AnalogValue);
 		return ProcessAnalogInputEvent(AnalogInputEvent);
 	}
 	return false;
 }
 
-bool FSlateApplication::OnControllerButtonPressed(FGamepadKeyNames::Type KeyName, int32 ControllerId, bool IsRepeat)
+bool FSlateApplication::OnControllerButtonPressed(FGamepadKeyNames::Type KeyName, FPlatformUserId PlatformUserId, FInputDeviceId InputDeviceId, bool IsRepeat)
 {
 	FKey Key(KeyName);
-	TOptional<int32> UserIndex = GetUserIndexForController(ControllerId, Key);
-	if(UserIndex.IsSet() && ensureMsgf(Key.IsValid(), TEXT("OnControllerButtonPressed(KeyName=%s,ControllerId=%d,IsRepeat=%b) key is invalid"), *KeyName.ToString(), ControllerId, IsRepeat))
+	TOptional<int32> UserIndex = GetUserIndexForInputDevice(InputDeviceId);
+	if (UserIndex.IsSet() && ensureMsgf(Key.IsValid(), TEXT("OnControllerButtonPressed(KeyName=%s,InputDeviceId=%d,IsRepeat=%b) key is invalid"), *KeyName.ToString(), InputDeviceId.GetId(), IsRepeat))
 	{
-		FKeyEvent KeyEvent(Key, PlatformApplication->GetModifierKeys(), UserIndex.GetValue(), IsRepeat, 0, 0);
+		FKeyEvent KeyEvent(Key, PlatformApplication->GetModifierKeys(), InputDeviceId, IsRepeat, 0, 0);
 		return ProcessKeyDownEvent(KeyEvent);
 	}
 	return false;
 }
 
-bool FSlateApplication::OnControllerButtonReleased(FGamepadKeyNames::Type KeyName, int32 ControllerId, bool IsRepeat)
+bool FSlateApplication::OnControllerButtonReleased(FGamepadKeyNames::Type KeyName, FPlatformUserId PlatformUserId, FInputDeviceId InputDeviceId, bool IsRepeat)
 {
 	FKey Key(KeyName);
-	TOptional<int32> UserIndex = GetUserIndexForController(ControllerId, Key);
-	if(UserIndex.IsSet() && ensureMsgf(Key.IsValid(), TEXT("OnControllerButtonReleased(KeyName=%s,ControllerId=%d,IsRepeat=%b) key is invalid"), *KeyName.ToString(), ControllerId, IsRepeat))
+	TOptional<int32> UserIndex = GetUserIndexForInputDevice(InputDeviceId);
+	if (UserIndex.IsSet() && ensureMsgf(Key.IsValid(), TEXT("OnControllerButtonReleased(KeyName=%s,InputDeviceId=%d,IsRepeat=%b) key is invalid"), *KeyName.ToString(), InputDeviceId.GetId(), IsRepeat))
 	{
-		FKeyEvent KeyEvent(Key, PlatformApplication->GetModifierKeys(), UserIndex.GetValue(), IsRepeat, 0, 0);
+		FKeyEvent KeyEvent(Key, PlatformApplication->GetModifierKeys(), InputDeviceId, IsRepeat, 0, 0);
 		return ProcessKeyUpEvent(KeyEvent);
 	}
 	return false;
 }
 
+PRAGMA_DISABLE_DEPRECATION_WARNINGS
+bool FSlateApplication::OnControllerAnalog( FGamepadKeyNames::Type KeyName, int32 ControllerId, float AnalogValue )
+{
+	// Remap the old int32 ControlerId to the new Platform user for backwards compat
+	FPlatformUserId UserId = PLATFORMUSERID_NONE;
+	FInputDeviceId DeviceId = INPUTDEVICEID_NONE;
+	IPlatformInputDeviceMapper::Get().RemapControllerIdToPlatformUserAndDevice(ControllerId, UserId, DeviceId);
+
+	return OnControllerAnalog(KeyName, UserId, DeviceId, AnalogValue);
+}
+
+bool FSlateApplication::OnControllerButtonPressed(FGamepadKeyNames::Type KeyName, int32 ControllerId, bool IsRepeat)
+{
+	// Remap the old int32 ControlerId to the new Platform user for backwards compat
+	FPlatformUserId UserId = PLATFORMUSERID_NONE;
+	FInputDeviceId DeviceId = INPUTDEVICEID_NONE;
+	IPlatformInputDeviceMapper::Get().RemapControllerIdToPlatformUserAndDevice(ControllerId, UserId, DeviceId);
+
+	return OnControllerButtonPressed(KeyName, UserId, DeviceId, IsRepeat);
+}
+
+bool FSlateApplication::OnControllerButtonReleased(FGamepadKeyNames::Type KeyName, int32 ControllerId, bool IsRepeat)
+{
+	// Remap the old int32 ControlerId to the new Platform user for backwards compat
+	FPlatformUserId UserId = PLATFORMUSERID_NONE;
+	FInputDeviceId DeviceId = INPUTDEVICEID_NONE;
+	
+	IPlatformInputDeviceMapper::Get().RemapControllerIdToPlatformUserAndDevice(ControllerId, UserId, DeviceId);
+	return OnControllerButtonReleased(KeyName, UserId, DeviceId, IsRepeat);
+}
+PRAGMA_ENABLE_DEPRECATION_WARNINGS
+
 bool FSlateApplication::OnTouchGesture( EGestureEvent GestureType, const FVector2D &Delta, const float MouseWheelDelta, bool bIsDirectionInvertedFromDevice )
 {
 	const FVector2D CurrentCursorPosition = GetCursorPos();
diff --git a/Engine/Source/Runtime/Slate/Public/Framework/Application/SlateApplication.h b/Engine/Source/Runtime/Slate/Public/Framework/Application/SlateApplication.h
--- a/Engine/Source/Runtime/Slate/Public/Framework/Application/SlateApplication.h
+++ b/Engine/Source/Runtime/Slate/Public/Framework/Application/SlateApplication.h
@@ -13,6 +13,7 @@
 #include "Layout/SlateRect.h"
 #include "GenericPlatform/GenericApplicationMessageHandler.h"
 #include "GenericPlatform/GenericApplication.h"
+#include "GenericPlatform/GenericPlatformInputDeviceMapper.h"
 #include "Input/Events.h"
 #include "Input/DragAndDrop.h"
 #include "Input/Reply.h"
@@ -119,6 +120,22 @@ class SLATE_API ISlateInputManager
 public:
 	virtual int32 GetUserIndexForMouse() const = 0;
 	virtual int32 GetUserIndexForKeyboard() const = 0;
+
+	virtual int32 GetInputDeviceIdForMouse() const = 0;
+	virtual int32 GetInputDeviceIdForKeyboard() const = 0;
+	
+	virtual int32 GetUserIndexForInputDevice(FInputDeviceId InputDeviceId) const
+	{
+		// There is a 1:1 mapping of the platform user ID to a slate user index.
+		return GetUserIndexForPlatformUser(IPlatformInputDeviceMapper::Get().GetUserForInputDevice(InputDeviceId));
+	}
+
+	virtual int32 GetUserIndexForPlatformUser(FPlatformUserId PlatformUser) const
+	{
+		// There is a 1:1 mapping of the platform user ID to a slate user index.
+		return PlatformUser.GetInternalId();
+	}
+
 	virtual int32 GetUserIndexForController(int32 ControllerId) const { return ControllerId; }
 	virtual TOptional<int32> GetUserIndexForController(int32 ControllerId, FKey InKey) const = 0;
 };
@@ -128,6 +145,17 @@ class SLATE_API FSlateDefaultInputMapping : public ISlateInputManager
 public:
 	virtual int32 GetUserIndexForMouse() const override { return 0; }
 	virtual int32 GetUserIndexForKeyboard() const override { return 0; }
+	
+	virtual int32 GetInputDeviceIdForMouse() const override
+	{
+		return GetUserIndexForInputDevice(IPlatformInputDeviceMapper::Get().GetDefaultInputDevice());
+	};
+	
+	virtual int32 GetInputDeviceIdForKeyboard() const override
+	{
+		return GetUserIndexForInputDevice(IPlatformInputDeviceMapper::Get().GetDefaultInputDevice());
+	};
+	
 	virtual int32 GetUserIndexForController(int32 ControllerId) const override { return ControllerId; }
 	virtual TOptional<int32> GetUserIndexForController(int32 ControllerId, FKey InKey) const override { return GetUserIndexForController(ControllerId); }
 };
@@ -919,6 +947,19 @@ public:
 	{
 		return Users.IsValidIndex(UserIndex) ? Users[UserIndex] : nullptr;
 	}
+	FORCEINLINE TSharedPtr<FSlateUser> GetUser(FPlatformUserId PlatformUser)
+	{
+		int32 InternalId = 0;
+		if (PlatformUser.IsValid())
+		{
+			InternalId = PlatformUser.GetInternalId();
+		}
+		else
+		{
+			UE_LOG(LogSlate, Warning, TEXT("SlateApplication::GetUser called with an invalid platform user! Defaulting to 0"));
+		}
+		return Users.IsValidIndex(InternalId) ? Users[InternalId] : nullptr;
+	}
 	FORCEINLINE TSharedPtr<const FSlateUser> GetUser(const FInputEvent& InputEvent) const { return GetUser(InputEvent.GetUserIndex()); }
 	FORCEINLINE TSharedPtr<FSlateUser> GetUser(const FInputEvent& InputEvent) { return GetUser(InputEvent.GetUserIndex()); }
 	
@@ -940,6 +981,26 @@ public:
 		}
 		return nullptr;
 	}
+	
+	FORCEINLINE TSharedPtr<FSlateUser> GetUserFromPlatformUser(FPlatformUserId PlatformUser)
+	{
+		TOptional<int32> UserIndex = GetUserIndexForPlatformUser(PlatformUser);
+		if (UserIndex.IsSet())
+		{
+			return GetUser(UserIndex.GetValue());
+		}
+		return nullptr;
+	}
+	
+	FORCEINLINE TSharedPtr<const FSlateUser> GetUserFromPlatformUser(FPlatformUserId PlatformUser) const
+	{
+		TOptional<int32> UserIndex = GetUserIndexForPlatformUser(PlatformUser);
+		if (UserIndex.IsSet())
+		{
+			return GetUser(UserIndex.GetValue());
+		}
+		return nullptr;
+	}
 
 	/** Get the standard 'default' user (there's always guaranteed to be at least one). */
 	FORCEINLINE TSharedPtr<const FSlateUser> GetCursorUser() const
@@ -978,11 +1039,25 @@ protected:
 	 */
 	TSharedRef<FSlateUser> RegisterNewUser(int32 UserIndex, bool bIsVirtual = false);
 
+	/**
+	 * Register a new user with Slate.  Normally this is unnecessary as Slate automatically adds
+	 * a user entry if it gets input from a controller for that index.  Might happen if the user
+	 * allocates the virtual user.
+	 */
+	TSharedRef<FSlateUser> RegisterNewUser(FPlatformUserId PlatformUserId, bool bIsVirtual = false);
+
 	/**
 	 * Locates the SlateUser object corresponding to the index, creating a new one if it doesn't exist.
 	 * Asserts if given an invalid (ie, negative) index.
 	 */
 	TSharedRef<FSlateUser> GetOrCreateUser(int32 UserIndex);
+	
+	/**
+	 * Locates the SlateUser object corresponding to the index, creating a new one if it doesn't exist.
+	 * Asserts if given an invalid (ie, negative) index.
+	 */
+	TSharedRef<FSlateUser> GetOrCreateUser(FPlatformUserId PlatformUserId);
+	
 	FORCEINLINE TSharedRef<FSlateUser> GetOrCreateUser(const FInputEvent& InputEvent) { return GetOrCreateUser(InputEvent.GetUserIndex()); }
 
 	friend class FEventRouter;
@@ -1434,8 +1509,11 @@ public:
 	virtual bool OnMouseMove() override;
 	virtual bool OnRawMouseMove( const int32 X, const int32 Y ) override;
 	virtual bool OnCursorSet() override;
+	//UE_DEPRECATED(5.1, "This version of OnControllerAnalog has been deprecated, please use the one that takes an FPlatformUser and FInputDeviceId instead.")
 	virtual bool OnControllerAnalog( FGamepadKeyNames::Type KeyName, int32 ControllerId, float AnalogValue ) override;
+	//UE_DEPRECATED(5.1, "This version of OnControllerButtonPressed has been deprecated, please use the one that takes an FPlatformUser and FInputDeviceId instead.")
 	virtual bool OnControllerButtonPressed( FGamepadKeyNames::Type KeyName, int32 ControllerId, bool IsRepeat ) override;
+	//UE_DEPRECATED(5.1, "This version of OnControllerButtonReleased has been deprecated, please use the one that takes an FPlatformUser and FInputDeviceId instead.")
 	virtual bool OnControllerButtonReleased( FGamepadKeyNames::Type KeyName, int32 ControllerId, bool IsRepeat ) override;
 	virtual bool OnTouchGesture( EGestureEvent GestureType, const FVector2D& Delta, float WheelDelta, bool bIsDirectionInvertedFromDevice ) override;
 	virtual bool OnTouchStarted( const TSharedPtr< FGenericWindow >& PlatformWindow, const FVector2D& Location, float Force, int32 TouchIndex, int32 ControllerId ) override;
@@ -1445,6 +1523,11 @@ public:
 	virtual bool OnTouchFirstMove(const FVector2D& Location, float Force, int32 TouchIndex, int32 ControllerId) override;
 	virtual void ShouldSimulateGesture(EGestureEvent Gesture, bool bEnable) override;
 	virtual bool OnMotionDetected(const FVector& Tilt, const FVector& RotationRate, const FVector& Gravity, const FVector& Acceleration, int32 ControllerId) override;
+	
+	virtual bool OnControllerAnalog(FGamepadKeyNames::Type KeyName, FPlatformUserId PlatformUserId, FInputDeviceId InputDeviceId, float AnalogValue) override;
+	virtual bool OnControllerButtonPressed(FGamepadKeyNames::Type KeyName, FPlatformUserId PlatformUserId, FInputDeviceId InputDeviceId, bool IsRepeat) override;
+	virtual bool OnControllerButtonReleased(FGamepadKeyNames::Type KeyName, FPlatformUserId PlatformUserId, FInputDeviceId InputDeviceId, bool IsRepeat) override;
+
 	virtual bool OnSizeChanged( const TSharedRef< FGenericWindow >& PlatformWindow, const int32 Width, const int32 Height, bool bWasMinimized = false ) override;
 	virtual void OnOSPaint( const TSharedRef< FGenericWindow >& PlatformWindow ) override;
 	virtual FWindowSizeLimits GetSizeLimitsForWindow(const TSharedRef<FGenericWindow>& Window) const override;
@@ -1528,8 +1611,24 @@ public:
 	 */
 	int32 GetUserIndexForKeyboard() const;
 
+	/**
+	 * @return InputDeviceId that the mouse is mapped to
+	 */
+	int32 GetInputDeviceIdForMouse() const;
+
+	/**
+	 * @return InputDeviceId that the keyboard is mapped to
+	 */
+	int32 GetInputDeviceIdForKeyboard() const;
+
 	/** @return int user index that this controller is mapped to. */
 	int32 GetUserIndexForController(int32 ControllerId) const;
+
+	/** @return Gets the slate user index that this input device is mapped to */	
+	int32 GetUserIndexForInputDevice(FInputDeviceId InputDeviceId) const;
+	/** @return Gets the slate user index that this platform user id mapped to*/	
+	int32 GetUserIndexForPlatformUser(FPlatformUserId PlatformUser) const;
+	
 	TOptional<int32> GetUserIndexForController(int32 ControllerId, FKey InKey) const;
 
 	/** Establishes the input mapping object used to map input sources to SlateUser indices */
