Change the format of some slate functions to be not FORCEINLINE, and remove extra spacing based on feedback about changing the pattern in the file

#rnx

diff --git a/Engine/Source/Runtime/Slate/Private/Framework/Application/SlateApplication.cpp b/Engine/Source/Runtime/Slate/Private/Framework/Application/SlateApplication.cpp
--- a/Engine/Source/Runtime/Slate/Private/Framework/Application/SlateApplication.cpp
+++ b/Engine/Source/Runtime/Slate/Private/Framework/Application/SlateApplication.cpp
@@ -3973,6 +3973,40 @@ bool FSlateApplication::TakeScreenshot(const TSharedRef<SWidget>& Widget, const
 	return (OutSize.X != 0 && OutSize.Y != 0 && OutColorData.Num() >= OutSize.X * OutSize.Y);
 }
 
+TSharedPtr<FSlateUser> FSlateApplication::GetUser(FPlatformUserId PlatformUser)
+{
+	int32 InternalId = 0;
+	if (PlatformUser.IsValid())
+	{
+		InternalId = PlatformUser.GetInternalId();
+	}
+	else
+	{
+		UE_LOG(LogSlate, Warning, TEXT("SlateApplication::GetUser called with an invalid platform user! Defaulting to 0"));
+	}
+	return Users.IsValidIndex(InternalId) ? Users[InternalId] : nullptr;
+}
+
+TSharedPtr<FSlateUser> FSlateApplication::GetUserFromPlatformUser(FPlatformUserId PlatformUser)
+{
+	TOptional<int32> UserIndex = GetUserIndexForPlatformUser(PlatformUser);
+	if (UserIndex.IsSet())
+	{
+		return GetUser(UserIndex.GetValue());
+	}
+	return nullptr;
+}
+
+TSharedPtr<const FSlateUser> FSlateApplication::GetUserFromPlatformUser(FPlatformUserId PlatformUser) const
+{
+	TOptional<int32> UserIndex = GetUserIndexForPlatformUser(PlatformUser);
+	if (UserIndex.IsSet())
+	{
+		return GetUser(UserIndex.GetValue());
+	}
+	return nullptr;
+}
+
 TSharedRef<FSlateVirtualUserHandle> FSlateApplication::FindOrCreateVirtualUser(int32 VirtualUserIndex)
 {
 	// Ensure we have a large enough array to add the new virtual user
diff --git a/Engine/Source/Runtime/Slate/Public/Framework/Application/SlateApplication.h b/Engine/Source/Runtime/Slate/Public/Framework/Application/SlateApplication.h
--- a/Engine/Source/Runtime/Slate/Public/Framework/Application/SlateApplication.h
+++ b/Engine/Source/Runtime/Slate/Public/Framework/Application/SlateApplication.h
@@ -947,19 +947,7 @@ public:
 	{
 		return Users.IsValidIndex(UserIndex) ? Users[UserIndex] : nullptr;
 	}
-	FORCEINLINE TSharedPtr<FSlateUser> GetUser(FPlatformUserId PlatformUser)
-	{
-		int32 InternalId = 0;
-		if (PlatformUser.IsValid())
-		{
-			InternalId = PlatformUser.GetInternalId();
-		}
-		else
-		{
-			UE_LOG(LogSlate, Warning, TEXT("SlateApplication::GetUser called with an invalid platform user! Defaulting to 0"));
-		}
-		return Users.IsValidIndex(InternalId) ? Users[InternalId] : nullptr;
-	}
+	TSharedPtr<FSlateUser> GetUser(FPlatformUserId PlatformUser);
 	FORCEINLINE TSharedPtr<const FSlateUser> GetUser(const FInputEvent& InputEvent) const { return GetUser(InputEvent.GetUserIndex()); }
 	FORCEINLINE TSharedPtr<FSlateUser> GetUser(const FInputEvent& InputEvent) { return GetUser(InputEvent.GetUserIndex()); }
 	
@@ -982,25 +970,9 @@ public:
 		return nullptr;
 	}
 	
-	FORCEINLINE TSharedPtr<FSlateUser> GetUserFromPlatformUser(FPlatformUserId PlatformUser)
-	{
-		TOptional<int32> UserIndex = GetUserIndexForPlatformUser(PlatformUser);
-		if (UserIndex.IsSet())
-		{
-			return GetUser(UserIndex.GetValue());
-		}
-		return nullptr;
-	}
+	TSharedPtr<FSlateUser> GetUserFromPlatformUser(FPlatformUserId PlatformUser);
 	
-	FORCEINLINE TSharedPtr<const FSlateUser> GetUserFromPlatformUser(FPlatformUserId PlatformUser) const
-	{
-		TOptional<int32> UserIndex = GetUserIndexForPlatformUser(PlatformUser);
-		if (UserIndex.IsSet())
-		{
-			return GetUser(UserIndex.GetValue());
-		}
-		return nullptr;
-	}
+	TSharedPtr<const FSlateUser> GetUserFromPlatformUser(FPlatformUserId PlatformUser) const;
 
 	/** Get the standard 'default' user (there's always guaranteed to be at least one). */
 	FORCEINLINE TSharedPtr<const FSlateUser> GetCursorUser() const
@@ -1523,11 +1495,9 @@ public:
 	virtual bool OnTouchFirstMove(const FVector2D& Location, float Force, int32 TouchIndex, int32 ControllerId) override;
 	virtual void ShouldSimulateGesture(EGestureEvent Gesture, bool bEnable) override;
 	virtual bool OnMotionDetected(const FVector& Tilt, const FVector& RotationRate, const FVector& Gravity, const FVector& Acceleration, int32 ControllerId) override;
-	
 	virtual bool OnControllerAnalog(FGamepadKeyNames::Type KeyName, FPlatformUserId PlatformUserId, FInputDeviceId InputDeviceId, float AnalogValue) override;
 	virtual bool OnControllerButtonPressed(FGamepadKeyNames::Type KeyName, FPlatformUserId PlatformUserId, FInputDeviceId InputDeviceId, bool IsRepeat) override;
 	virtual bool OnControllerButtonReleased(FGamepadKeyNames::Type KeyName, FPlatformUserId PlatformUserId, FInputDeviceId InputDeviceId, bool IsRepeat) override;
-
 	virtual bool OnSizeChanged( const TSharedRef< FGenericWindow >& PlatformWindow, const int32 Width, const int32 Height, bool bWasMinimized = false ) override;
 	virtual void OnOSPaint( const TSharedRef< FGenericWindow >& PlatformWindow ) override;
 	virtual FWindowSizeLimits GetSizeLimitsForWindow(const TSharedRef<FGenericWindow>& Window) const override;
