Exporting SListPanel

diff --git a/Engine/Source/Runtime/Slate/Private/Widgets/Views/SListPanel.h b/Engine/Source/Runtime/Slate/Private/Widgets/Views/SListPanel.h
--- a/Engine/Source/Runtime/Slate/Private/Widgets/Views/SListPanel.h
+++ b/Engine/Source/Runtime/Slate/Private/Widgets/Views/SListPanel.h
@@ -18,7 +18,7 @@ class FArrangedChildren;
  * Items in this panel have a uniform height.
  * Also supports offsetting its items vertically.
  */
-class SListPanel : public SPanel
+class SLATE_API SListPanel : public SPanel
 {
 public:
 	/** A ListPanel slot is very simple - it just stores a widget. */
