[Engine/UI]  Allow navigation to escape a SListView if a navigable element cannot be found.

diff --git a/Engine/Source/Runtime/Slate/Public/Widgets/Views/SListView.h b/Engine/Source/Runtime/Slate/Public/Widgets/Views/SListView.h
--- a/Engine/Source/Runtime/Slate/Public/Widgets/Views/SListView.h
+++ b/Engine/Source/Runtime/Slate/Public/Widgets/Views/SListView.h
@@ -524,8 +524,12 @@ public:
 			// If it's valid we'll scroll it into view and return an explicit widget in the FNavigationReply
 			if (ItemsSourceRef.IsValidIndex(AttemptSelectIndex))
 			{
-				NavigationSelect(ItemsSourceRef[AttemptSelectIndex], InNavigationEvent);
-				return FNavigationReply::Explicit(nullptr);
+				TOptional<ItemType> ItemToSelect = Private_FindNextSelectableOrNavigableWithIndexAndDirection(ItemsSourceRef[AttemptSelectIndex], AttemptSelectIndex, AttemptSelectIndex >= CurSelectionIndex);
+				if (ItemToSelect.IsSet())
+				{
+					NavigationSelect(ItemToSelect.GetValue(), InNavigationEvent);
+					return FNavigationReply::Explicit(nullptr);
+				}
 			}
 		}
 
@@ -2314,36 +2318,19 @@ protected:
 
 protected:
 
-	TOptional<ItemType> Private_FindNextSelectableOrNavigable(const ItemType& InItemToSelect)
+	TOptional<ItemType> Private_FindNextSelectableOrNavigableWithIndexAndDirection(const ItemType& InItemToSelect, int32 SelectionIdx, bool bSelectForward)
 	{
 		ItemType ItemToSelect = InItemToSelect;
 
 		if (OnIsSelectableOrNavigable.IsBound())
 		{
-			NullableItemType LastSelectedItem = TListTypeTraits<ItemType>::MakeNullPtr();
-			if (SelectedItems.Num() == 1)
-			{
-				LastSelectedItem = *SelectedItems.CreateIterator();
-			}
-
-			bool bSelectNextItem = true;
+			// Walk through the list until we either find a navigable item or run out of entries.
 			while (!OnIsSelectableOrNavigable.Execute(ItemToSelect))
 			{
-				const int32 PendingItemIndex = ItemsSource->Find(ItemToSelect);
-				if (TListTypeTraits<ItemType>::IsPtrValid(LastSelectedItem) && PendingItemIndex > 0)
-				{
-					ItemType NonNullLastSelectedItem = TListTypeTraits<ItemType>::NullableItemTypeConvertToItemType(LastSelectedItem);
-					const int32 LastSelectedItemIdx = ItemsSource->Find(NonNullLastSelectedItem);
-
-					// If the previously selected item was before the header, assume we're navigating down and want to select the next item
-					// Otherwise, assume the opposite and navigate to the previous item
-					bSelectNextItem = LastSelectedItemIdx < PendingItemIndex;
-				}
-
-				const int32 NewSelectionIdx = PendingItemIndex + (bSelectNextItem ? 1 : -1);
-				if (ItemsSource->IsValidIndex(NewSelectionIdx))
+				SelectionIdx += (bSelectForward ? 1 : -1);
+				if (ItemsSource->IsValidIndex(SelectionIdx))
 				{
-					ItemToSelect = (*ItemsSource)[NewSelectionIdx];
+					ItemToSelect = (*ItemsSource)[SelectionIdx];
 				}
 				else
 				{
@@ -2356,6 +2343,51 @@ protected:
 		return TOptional<ItemType>(ItemToSelect);
 	}
 
+	TOptional<ItemType> Private_FindNextSelectableOrNavigable(const ItemType& InItemToSelect)
+	{
+		ItemType ItemToSelect = InItemToSelect;
+
+		if (OnIsSelectableOrNavigable.IsBound())
+		{
+			if (!OnIsSelectableOrNavigable.Execute(ItemToSelect))
+			{
+				int32 NewSelectionIdx = ItemsSource->Find(ItemToSelect);
+
+				// By default, we walk forward
+				bool bSelectNextItem = true;
+				if (SelectedItems.Num() == 1)
+				{
+					// If the last selected item is after the item to select, we'll want to walk backwards
+					NullableItemType LastSelectedItem = *SelectedItems.CreateIterator();
+					if (TListTypeTraits<ItemType>::IsPtrValid(LastSelectedItem))
+					{
+						ItemType NonNullLastSelectedItem = TListTypeTraits<ItemType>::NullableItemTypeConvertToItemType(LastSelectedItem);
+						const int32 LastSelectedItemIdx = ItemsSource->Find(NonNullLastSelectedItem);
+
+						bSelectNextItem = LastSelectedItemIdx < NewSelectionIdx;
+					}
+				}
+
+				// Walk through the list until we either find a navigable item or run out of entries.
+				do
+				{
+					NewSelectionIdx += (bSelectNextItem ? 1 : -1);
+					if (ItemsSource->IsValidIndex(NewSelectionIdx))
+					{
+						ItemToSelect = (*ItemsSource)[NewSelectionIdx];
+					}
+					else
+					{
+						// Failed to find a valid item to select
+						return TOptional<ItemType>();
+					}
+				} while (!OnIsSelectableOrNavigable.Execute(ItemToSelect));
+			}
+		}
+
+		return TOptional<ItemType>(ItemToSelect);
+	}
+
 	/**
 	 * Selects the specified item and scrolls it into view. If shift is held, it will be a range select.
 	 * 
