Fix an incorrect return type on the GetInputDeviceId functions on the slate app. It was returning an int32, but should have been an InputDeviceId.

#rnx

diff --git a/Engine/Source/Runtime/Slate/Private/Framework/Application/SlateApplication.cpp b/Engine/Source/Runtime/Slate/Private/Framework/Application/SlateApplication.cpp
--- a/Engine/Source/Runtime/Slate/Private/Framework/Application/SlateApplication.cpp
+++ b/Engine/Source/Runtime/Slate/Private/Framework/Application/SlateApplication.cpp
@@ -2409,12 +2409,12 @@ int32 FSlateApplication::GetUserIndexForKeyboard() const
 	return InputManager->GetUserIndexForKeyboard();
 }
 
-int32 FSlateApplication::GetInputDeviceIdForMouse() const
+FInputDeviceId FSlateApplication::GetInputDeviceIdForMouse() const
 {
 	return InputManager->GetInputDeviceIdForMouse();
 }
 
-int32 FSlateApplication::GetInputDeviceIdForKeyboard() const
+FInputDeviceId FSlateApplication::GetInputDeviceIdForKeyboard() const
 {
 	return InputManager->GetInputDeviceIdForKeyboard();
 }
diff --git a/Engine/Source/Runtime/Slate/Public/Framework/Application/SlateApplication.h b/Engine/Source/Runtime/Slate/Public/Framework/Application/SlateApplication.h
--- a/Engine/Source/Runtime/Slate/Public/Framework/Application/SlateApplication.h
+++ b/Engine/Source/Runtime/Slate/Public/Framework/Application/SlateApplication.h
@@ -121,8 +121,8 @@ public:
 	virtual int32 GetUserIndexForMouse() const = 0;
 	virtual int32 GetUserIndexForKeyboard() const = 0;
 
-	virtual int32 GetInputDeviceIdForMouse() const = 0;
-	virtual int32 GetInputDeviceIdForKeyboard() const = 0;
+	virtual FInputDeviceId GetInputDeviceIdForMouse() const = 0;
+	virtual FInputDeviceId GetInputDeviceIdForKeyboard() const = 0;
 	
 	virtual int32 GetUserIndexForInputDevice(FInputDeviceId InputDeviceId) const
 	{
@@ -146,14 +146,14 @@ public:
 	virtual int32 GetUserIndexForMouse() const override { return 0; }
 	virtual int32 GetUserIndexForKeyboard() const override { return 0; }
 	
-	virtual int32 GetInputDeviceIdForMouse() const override
+	virtual FInputDeviceId GetInputDeviceIdForMouse() const override
 	{
-		return GetUserIndexForInputDevice(IPlatformInputDeviceMapper::Get().GetDefaultInputDevice());
+		return IPlatformInputDeviceMapper::Get().GetDefaultInputDevice();
 	};
 	
-	virtual int32 GetInputDeviceIdForKeyboard() const override
+	virtual FInputDeviceId GetInputDeviceIdForKeyboard() const override
 	{
-		return GetUserIndexForInputDevice(IPlatformInputDeviceMapper::Get().GetDefaultInputDevice());
+		return IPlatformInputDeviceMapper::Get().GetDefaultInputDevice();
 	};
 	
 	virtual int32 GetUserIndexForController(int32 ControllerId) const override { return ControllerId; }
@@ -1584,12 +1584,12 @@ public:
 	/**
 	 * @return InputDeviceId that the mouse is mapped to
 	 */
-	int32 GetInputDeviceIdForMouse() const;
+	FInputDeviceId GetInputDeviceIdForMouse() const;
 
 	/**
 	 * @return InputDeviceId that the keyboard is mapped to
 	 */
-	int32 GetInputDeviceIdForKeyboard() const;
+	FInputDeviceId GetInputDeviceIdForKeyboard() const;
 
 	/** @return int user index that this controller is mapped to. */
 	int32 GetUserIndexForController(int32 ControllerId) const;
