Modified SListView.SetItemSelection to use `TConstArrayView` rather than `const TArray&`
It allows TConstArrayView use while still supporting TArrays.

diff --git a/Engine/Source/Runtime/Slate/Public/Widgets/Views/SListView.h b/Engine/Source/Runtime/Slate/Public/Widgets/Views/SListView.h
--- a/Engine/Source/Runtime/Slate/Public/Widgets/Views/SListView.h
+++ b/Engine/Source/Runtime/Slate/Public/Widgets/Views/SListView.h
@@ -1720,7 +1720,7 @@ public:
 	 * @param bSelected   true to select the items; false to unselect
 	 * @param SelectInfo  Provides context on how the selection changed
 	 */
-	void SetItemSelection( const TArray<ItemType>& InItems, bool bSelected, ESelectInfo::Type SelectInfo = ESelectInfo::Direct )
+	void SetItemSelection(TConstArrayView<ItemType> InItems, bool bSelected, ESelectInfo::Type SelectInfo = ESelectInfo::Direct)
 	{
 		if ( InItems.Num() == 0 || SelectionMode.Get() == ESelectionMode::None )
 		{
