STreeView: Fix a bug that prevented stacked hierarchy headers from showing up

diff --git a/Engine/Source/Runtime/Slate/Public/Widgets/Views/STreeView.h b/Engine/Source/Runtime/Slate/Public/Widgets/Views/STreeView.h
--- a/Engine/Source/Runtime/Slate/Public/Widgets/Views/STreeView.h
+++ b/Engine/Source/Runtime/Slate/Public/Widgets/Views/STreeView.h
@@ -710,7 +710,7 @@ public:
 		int32 CurrentItemIndex = StartIndex;
 
 		auto GetNonVisibleParents = [this, &InItemsSource, StartIndex](TArray<ItemType>& OutParents, int32 ItemIndex) {
-			if (DenseItemInfos.IsValidIndex(ItemIndex))
+			if (!DenseItemInfos.IsValidIndex(ItemIndex))
 			{
 				return;
 			}
