Initalize the PlatformUser in case the old FSlateUser constructor is used.

#rnx

diff --git a/Engine/Source/Runtime/Slate/Private/Framework/Application/SlateUser.cpp b/Engine/Source/Runtime/Slate/Private/Framework/Application/SlateUser.cpp
--- a/Engine/Source/Runtime/Slate/Private/Framework/Application/SlateUser.cpp
+++ b/Engine/Source/Runtime/Slate/Private/Framework/Application/SlateUser.cpp
@@ -107,7 +107,8 @@ FSlateUser::FSlateUser(int32 InUserIndex, TSharedPtr<ICursor> InCursor)
 	: UserIndex(InUserIndex)
 	, Cursor(InCursor)
 {
-	UE_LOG(LogSlate, Log, TEXT("New Slate User Created.  User Index %d, Is Virtual User: %d"), UserIndex, IsVirtualUser());
+	PlatformUser = FPlatformMisc::GetPlatformUserForUserIndex(InUserIndex);
+	UE_LOG(LogSlate, Log, TEXT("New Slate User Created. Platform User Id %d, User Index %d, Is Virtual User: %d"), PlatformUser.GetInternalId(), UserIndex, IsVirtualUser());
 
 	PointerPositionsByIndex.Add(FSlateApplication::CursorPointerIndex, FVector2D::ZeroVector);
 	PreviousPointerPositionsByIndex.Add(FSlateApplication::CursorPointerIndex, FVector2D::ZeroVector);
