Editor notifications now only start fading out once the editor starts ticking. When a notification was added during loading, it immediately disappear if loading took longer than the expiration time given.

This was occurring because the FCurveSequence animation the notification was using to drive its opacity was using the FSlateApplication::GetCurrentTime(), which is only updated after loading is completed.

diff --git a/Engine/Plugins/Editor/PluginBrowser/Source/PluginBrowser/Private/PluginBrowserModule.cpp b/Engine/Plugins/Editor/PluginBrowser/Source/PluginBrowser/Private/PluginBrowserModule.cpp
--- a/Engine/Plugins/Editor/PluginBrowser/Source/PluginBrowser/Private/PluginBrowserModule.cpp
+++ b/Engine/Plugins/Editor/PluginBrowser/Source/PluginBrowser/Private/PluginBrowserModule.cpp
@@ -203,8 +203,7 @@ void FPluginBrowserModule::OnMainFrameLoaded(TSharedPtr<SWindow> InRootWindow, b
 		Info.bFireAndForget = true;
 		Info.bUseLargeFont = true;
 		Info.bUseThrobber = false;
-		// Note this time is large to add padding because this notification is added before the editor UI is visible. 
-		Info.ExpireDuration = 25.0f;
+		Info.ExpireDuration = 10.0f;
 		Info.ButtonDetails.Add(FNotificationButtonInfo(LOCTEXT("NewPluginsPopupSettings", "Manage Plugins..."), LOCTEXT("NewPluginsPopupSettingsTT", "Open the plugin browser to enable plugins"), FSimpleDelegate::CreateRaw(this, &FPluginBrowserModule::OnNewPluginsPopupSettingsClicked)));
 		Info.ButtonDetails.Add(FNotificationButtonInfo(LOCTEXT("NewPluginsPopupDismiss", "Dismiss"), LOCTEXT("NewPluginsPopupDismissTT", "Dismiss this notification"), FSimpleDelegate::CreateRaw(this, &FPluginBrowserModule::OnNewPluginsPopupDismissClicked)));
 
diff --git a/Engine/Source/Runtime/Slate/Private/Widgets/Notifications/SNotificationList.cpp b/Engine/Source/Runtime/Slate/Private/Widgets/Notifications/SNotificationList.cpp
--- a/Engine/Source/Runtime/Slate/Private/Widgets/Notifications/SNotificationList.cpp
+++ b/Engine/Source/Runtime/Slate/Private/Widgets/Notifications/SNotificationList.cpp
@@ -69,12 +69,7 @@ public:
 
 	virtual void ExpireAndFadeout() override
 	{
-		FadeAnimation = FCurveSequence();
-		// Add some space for the expire time
-		FadeAnimation.AddCurve(FadeOutDuration.Get(), ExpireDuration.Get());
-		// Add the actual fade curve
-		FadeCurve = FadeAnimation.AddCurve(0.f, FadeOutDuration.Get());
-		FadeAnimation.PlayReverse(this->AsShared());
+		bAutoExpire = true;
 	}
 
 	/** Begins the fadein of this message */
@@ -147,17 +142,26 @@ public:
 
 	void Tick( const FGeometry& AllottedGeometry, const double InCurrentTime, const float InDeltaTime ) override
 	{
-		const bool bIsFadingOut = FadeAnimation.IsInReverse();
-		const bool bIsCurrentlyPlaying = FadeAnimation.IsPlaying();
-		const bool bIsIntroPlaying = IntroAnimation.IsPlaying();
+		InternalTime += InDeltaTime;
 
-		if ( !bIsCurrentlyPlaying && bIsFadingOut )
+		bool bIsFadingOut = FadeAnimation.IsInReverse();
+		if (bAutoExpire && 
+			!bIsFadingOut && 
+			InternalTime > (FadeInDuration.Get() + ExpireDuration.Get()))
+		{
+			Fadeout();
+			bIsFadingOut = true;
+		}
+
+		const bool bIsCurrentlyPlaying = FadeAnimation.IsPlaying();
+		if (bIsFadingOut && !bIsCurrentlyPlaying)
 		{
 			// Reset the Animation
 			FadeoutComplete();
 		}
-
-		if ( !bIsIntroPlaying && ThrottleHandle.IsValid() )
+		
+		const bool bIsIntroPlaying = IntroAnimation.IsPlaying();
+		if (!bIsIntroPlaying && ThrottleHandle.IsValid())
 		{
 			// Leave responsive mode once the intro finishes playing
 			FSlateThrottleManager::Get().LeaveResponsiveMode( ThrottleHandle );
@@ -327,6 +331,9 @@ protected:
 
 	/** Handle to a throttle request made to ensure the intro animation is smooth in low FPS situations */
 	FThrottleRequest ThrottleHandle;
+
+	float InternalTime = 0.0f;
+	bool bAutoExpire = false;
 };
 
 
