Enable OnController_X function deprecations for the Slate App and the Generic Message handler.

#rnx

diff --git a/Engine/Source/Runtime/ApplicationCore/Public/GenericPlatform/GenericApplicationMessageHandler.h b/Engine/Source/Runtime/ApplicationCore/Public/GenericPlatform/GenericApplicationMessageHandler.h
--- a/Engine/Source/Runtime/ApplicationCore/Public/GenericPlatform/GenericApplicationMessageHandler.h
+++ b/Engine/Source/Runtime/ApplicationCore/Public/GenericPlatform/GenericApplicationMessageHandler.h
@@ -8,6 +8,7 @@
 #include "Math/Vector2D.h"
 #include "Templates/SharedPointer.h"
 #include "Misc/Optional.h"
+#include "GenericPlatform/GenericPlatformInputDeviceMapper.h"
 
 class FGenericWindow;
 
@@ -305,28 +306,34 @@ public:
 
 	virtual bool OnControllerAnalog(FGamepadKeyNames::Type KeyName, FPlatformUserId PlatformUserId, FInputDeviceId InputDeviceId, float AnalogValue)
 	{
+		PRAGMA_DISABLE_DEPRECATION_WARNINGS
 		if (!ShouldUsePlatformUserId())
 		{
 			return OnControllerAnalog(KeyName, PlatformUserId.GetInternalId(), AnalogValue);
 		}
+		PRAGMA_ENABLE_DEPRECATION_WARNINGS
 		return false;
 	}
 
 	virtual bool OnControllerButtonPressed(FGamepadKeyNames::Type KeyName, FPlatformUserId PlatformUserId, FInputDeviceId InputDeviceId, bool IsRepeat)
 	{
+		PRAGMA_DISABLE_DEPRECATION_WARNINGS
 		if (!ShouldUsePlatformUserId())
 		{
 			return OnControllerButtonPressed(KeyName, PlatformUserId.GetInternalId(), IsRepeat);
 		}
+		PRAGMA_ENABLE_DEPRECATION_WARNINGS
 		return false;
 	}
 
 	virtual bool OnControllerButtonReleased(FGamepadKeyNames::Type KeyName, FPlatformUserId PlatformUserId, FInputDeviceId InputDeviceId, bool IsRepeat)
 	{
+		PRAGMA_DISABLE_DEPRECATION_WARNINGS
 		if (!ShouldUsePlatformUserId())
 		{
 			return OnControllerButtonReleased(KeyName, PlatformUserId.GetInternalId(), IsRepeat);
 		}
+		PRAGMA_ENABLE_DEPRECATION_WARNINGS
 		return false;
 	}
 
@@ -513,38 +520,54 @@ public:
 	}
 
 	// Deprecate these when engine code has been converted to handle platform user id
-	//UE_DEPRECATED(5.1, "This version of OnControllerAnalog has been deprecated, please use the one that takes an FPlatformUser and FInputDeviceId instead.")
+	UE_DEPRECATED(5.1, "This version of OnControllerAnalog has been deprecated, please use the one that takes an FPlatformUser and FInputDeviceId instead.")
 	virtual bool OnControllerAnalog(FGamepadKeyNames::Type KeyName, int32 ControllerId, float AnalogValue)
 	{
 		if (ShouldUsePlatformUserId())
 		{
-			return OnControllerAnalog(KeyName, FPlatformUserId(ControllerId), INPUTDEVICEID_NONE, AnalogValue);
+			// Remap the old int32 ControlerId to the new Platform user for backwards compat
+			FPlatformUserId UserId = FGenericPlatformMisc::GetPlatformUserForUserIndex(ControllerId);
+			FInputDeviceId DeviceId = INPUTDEVICEID_NONE;
+			IPlatformInputDeviceMapper::Get().RemapControllerIdToPlatformUserAndDevice(ControllerId, UserId, DeviceId);
+			return OnControllerAnalog(KeyName, UserId, DeviceId, AnalogValue);
 		}
 		return false;
 	}
-	//UE_DEPRECATED(5.1, "This version of OnControllerButtonPressed has been deprecated, please use the one that takes an FPlatformUser and FInputDeviceId instead.")
+	UE_DEPRECATED(5.1, "This version of OnControllerButtonPressed has been deprecated, please use the one that takes an FPlatformUser and FInputDeviceId instead.")
 	virtual bool OnControllerButtonPressed(FGamepadKeyNames::Type KeyName, int32 ControllerId, bool IsRepeat)
 	{
 		if (ShouldUsePlatformUserId())
 		{
-			return OnControllerButtonPressed(KeyName, FPlatformUserId(ControllerId), INPUTDEVICEID_NONE, IsRepeat);
+			// Remap the old int32 ControlerId to the new Platform user for backwards compat
+			FPlatformUserId UserId = FGenericPlatformMisc::GetPlatformUserForUserIndex(ControllerId);
+			FInputDeviceId DeviceId = INPUTDEVICEID_NONE;
+			IPlatformInputDeviceMapper::Get().RemapControllerIdToPlatformUserAndDevice(ControllerId, UserId, DeviceId);
+			return OnControllerButtonPressed(KeyName, UserId, DeviceId, IsRepeat);
 		}
 		return false;
 	}
-	//UE_DEPRECATED(5.1, "This version of OnControllerButtonReleased has been deprecated, please use the one that takes an FPlatformUser and FInputDeviceId instead.")
+	UE_DEPRECATED(5.1, "This version of OnControllerButtonReleased has been deprecated, please use the one that takes an FPlatformUser and FInputDeviceId instead.")
 	virtual bool OnControllerButtonReleased(FGamepadKeyNames::Type KeyName, int32 ControllerId, bool IsRepeat)
 	{
 		if (ShouldUsePlatformUserId())
 		{
-			return OnControllerButtonReleased(KeyName, FPlatformUserId(ControllerId), INPUTDEVICEID_NONE, IsRepeat);
+			// Remap the old int32 ControlerId to the new Platform user for backwards compat
+			FPlatformUserId UserId = FGenericPlatformMisc::GetPlatformUserForUserIndex(ControllerId);
+			FInputDeviceId DeviceId = INPUTDEVICEID_NONE;
+			IPlatformInputDeviceMapper::Get().RemapControllerIdToPlatformUserAndDevice(ControllerId, UserId, DeviceId);
+			return OnControllerButtonReleased(KeyName, UserId, DeviceId, IsRepeat);
 		}
 		return false;
 	}
 	virtual bool OnTouchStarted(const TSharedPtr< FGenericWindow >& Window, const FVector2D& Location, float Force, int32 TouchIndex, int32 ControllerId)
-	{
+	{ 
 		if (ShouldUsePlatformUserId())
 		{
-			return OnTouchStarted(Window, Location, Force, TouchIndex, FPlatformUserId(ControllerId), INPUTDEVICEID_NONE);
+			// Remap the old int32 ControlerId to the new Platform user for backwards compat
+			FPlatformUserId UserId = FGenericPlatformMisc::GetPlatformUserForUserIndex(ControllerId);
+			FInputDeviceId DeviceId = INPUTDEVICEID_NONE;
+			IPlatformInputDeviceMapper::Get().RemapControllerIdToPlatformUserAndDevice(ControllerId, UserId, DeviceId);
+			return OnTouchStarted(Window, Location, Force, TouchIndex, UserId, DeviceId);
 		}
 		return false;
 	}
@@ -552,7 +575,11 @@ public:
 	{
 		if (ShouldUsePlatformUserId())
 		{
-			return OnTouchMoved(Location, Force, TouchIndex, FPlatformUserId(ControllerId), INPUTDEVICEID_NONE);
+			// Remap the old int32 ControlerId to the new Platform user for backwards compat
+			FPlatformUserId UserId = FGenericPlatformMisc::GetPlatformUserForUserIndex(ControllerId);
+			FInputDeviceId DeviceId = INPUTDEVICEID_NONE;
+			IPlatformInputDeviceMapper::Get().RemapControllerIdToPlatformUserAndDevice(ControllerId, UserId, DeviceId);
+			return OnTouchMoved(Location, Force, TouchIndex, UserId, DeviceId);
 		}
 		return false;
 	}
@@ -560,7 +587,11 @@ public:
 	{
 		if (ShouldUsePlatformUserId())
 		{
-			return OnTouchEnded(Location, TouchIndex, FPlatformUserId(ControllerId), INPUTDEVICEID_NONE);
+			// Remap the old int32 ControlerId to the new Platform user for backwards compat
+			FPlatformUserId UserId = FGenericPlatformMisc::GetPlatformUserForUserIndex(ControllerId);
+			FInputDeviceId DeviceId = INPUTDEVICEID_NONE;
+			IPlatformInputDeviceMapper::Get().RemapControllerIdToPlatformUserAndDevice(ControllerId, UserId, DeviceId);
+			return OnTouchEnded(Location, TouchIndex, UserId, DeviceId);
 		}
 		return false;
 	}
@@ -568,7 +599,11 @@ public:
 	{
 		if (ShouldUsePlatformUserId())
 		{
-			return OnTouchForceChanged(Location, Force, TouchIndex, FPlatformUserId(ControllerId), INPUTDEVICEID_NONE);
+			// Remap the old int32 ControlerId to the new Platform user for backwards compat
+			FPlatformUserId UserId = FGenericPlatformMisc::GetPlatformUserForUserIndex(ControllerId);
+			FInputDeviceId DeviceId = INPUTDEVICEID_NONE;
+			IPlatformInputDeviceMapper::Get().RemapControllerIdToPlatformUserAndDevice(ControllerId, UserId, DeviceId);
+			return OnTouchForceChanged(Location, Force, TouchIndex, UserId, DeviceId);
 		}
 		return false;
 	}
@@ -576,7 +611,11 @@ public:
 	{
 		if (ShouldUsePlatformUserId())
 		{
-			return OnTouchFirstMove(Location, Force, TouchIndex, FPlatformUserId(ControllerId), INPUTDEVICEID_NONE);
+			// Remap the old int32 ControlerId to the new Platform user for backwards compat
+			FPlatformUserId UserId = FGenericPlatformMisc::GetPlatformUserForUserIndex(ControllerId);
+			FInputDeviceId DeviceId = INPUTDEVICEID_NONE;
+			IPlatformInputDeviceMapper::Get().RemapControllerIdToPlatformUserAndDevice(ControllerId, UserId, DeviceId);
+			return OnTouchFirstMove(Location, Force, TouchIndex, UserId, DeviceId);
 		}
 		return false;
 	}
@@ -584,7 +623,11 @@ public:
 	{
 		if (ShouldUsePlatformUserId())
 		{
-			return OnMotionDetected(Tilt, RotationRate, Gravity, Acceleration, FPlatformUserId(ControllerId), INPUTDEVICEID_NONE);
+			// Remap the old int32 ControlerId to the new Platform user for backwards compat
+			FPlatformUserId UserId = FGenericPlatformMisc::GetPlatformUserForUserIndex(ControllerId);
+			FInputDeviceId DeviceId = INPUTDEVICEID_NONE;
+			IPlatformInputDeviceMapper::Get().RemapControllerIdToPlatformUserAndDevice(ControllerId, UserId, DeviceId);
+			return OnMotionDetected(Tilt, RotationRate, Gravity, Acceleration, UserId, DeviceId);
 		}
 		return false;
 	}
diff --git a/Engine/Source/Runtime/Slate/Private/Framework/Application/SlateApplication.cpp b/Engine/Source/Runtime/Slate/Private/Framework/Application/SlateApplication.cpp
--- a/Engine/Source/Runtime/Slate/Private/Framework/Application/SlateApplication.cpp
+++ b/Engine/Source/Runtime/Slate/Private/Framework/Application/SlateApplication.cpp
@@ -4054,6 +4054,12 @@ TSharedRef<FSlateUser> FSlateApplication::GetOrCreateUser(FPlatformUserId Platfo
 	return RegisterNewUser(PlatformUserId);
 }
 
+TSharedRef<FSlateUser> FSlateApplication::GetOrCreateUser(FInputDeviceId DeviceId)
+{
+	// Get a user based on the owning platform user of this input device
+	return GetOrCreateUser(IPlatformInputDeviceMapper::Get().GetUserForInputDevice(DeviceId));
+}
+
 TSharedRef<FSlateUser> FSlateApplication::RegisterNewUser(int32 UserIndex, bool bIsVirtual)
 {
 	return RegisterNewUser(FGenericPlatformMisc::GetPlatformUserForUserIndex(UserIndex), bIsVirtual);
@@ -4709,7 +4715,7 @@ void FSlateApplication::SetGameIsFakingTouchEvents(const bool bIsFaking, FVector
 	{
 		if (bIsFakingTouched && !bIsFaking && bIsGameFakingTouch && !bIsFakingTouch)
 		{
-			OnTouchEnded((CursorLocation ? *CursorLocation : PlatformApplication->Cursor->GetPosition()), 0, 0);
+			OnTouchEnded((CursorLocation ? *CursorLocation : PlatformApplication->Cursor->GetPosition()), 0, FSlateApplicationBase::SlateAppPrimaryPlatformUser, IPlatformInputDeviceMapper::Get().GetDefaultInputDevice());
 		}
 
 		bIsGameFakingTouch = bIsFaking;
@@ -4732,7 +4738,7 @@ bool FSlateApplication::OnMouseDown( const TSharedPtr< FGenericWindow >& Platfor
 	if (IsFakingTouchEvents() && Button == EMouseButtons::Left)
 	{
 		bIsFakingTouched = true;
-		return OnTouchStarted( PlatformWindow, PlatformApplication->Cursor->GetPosition(), 1.0f, 0, 0 );
+		return OnTouchStarted( PlatformWindow, PlatformApplication->Cursor->GetPosition(), 1.0f, /* touch index */ 0, FSlateApplicationBase::SlateAppPrimaryPlatformUser, IPlatformInputDeviceMapper::Get().GetDefaultInputDevice() );
 	}
 
 	FKey Key = TranslateMouseButtonToKey( Button );
@@ -5423,7 +5429,7 @@ bool FSlateApplication::OnMouseDoubleClick( const TSharedPtr< FGenericWindow >&
 	if (IsFakingTouchEvents())
 	{
 		bIsFakingTouched = true;
-		return OnTouchStarted(PlatformWindow, PlatformApplication->Cursor->GetPosition(), 1.0f, 0, 0);
+		return OnTouchStarted(PlatformWindow, PlatformApplication->Cursor->GetPosition(), 1.0f, /* touch index */ 0, FSlateApplicationBase::SlateAppPrimaryPlatformUser, IPlatformInputDeviceMapper::Get().GetDefaultInputDevice());
 	}
 
 	FKey Key = TranslateMouseButtonToKey( Button );
@@ -5515,7 +5521,7 @@ bool FSlateApplication::OnMouseUp( const EMouseButtons::Type Button, const FVect
 	if (IsFakingTouchEvents() && Button == EMouseButtons::Left)
 	{
 		bIsFakingTouched = false;
-		return OnTouchEnded(PlatformApplication->Cursor->GetPosition(), 0, 0);
+		return OnTouchEnded(PlatformApplication->Cursor->GetPosition(), 0, FSlateApplicationBase::SlateAppPrimaryPlatformUser, IPlatformInputDeviceMapper::Get().GetDefaultInputDevice());
 	}
 
 	FKey Key = TranslateMouseButtonToKey( Button );
@@ -5696,7 +5702,7 @@ bool FSlateApplication::OnMouseMove()
 		// convert to touch event if we are faking it
 		if (bIsFakingTouched)
 		{
-			return OnTouchMoved(PlatformApplication->Cursor->GetPosition(), 1.0f, 0, 0);
+			return OnTouchMoved(PlatformApplication->Cursor->GetPosition(), 1.0f, 0, FSlateApplicationBase::SlateAppPrimaryPlatformUser, IPlatformInputDeviceMapper::Get().GetDefaultInputDevice());
 		}
 
 		// Throw out the mouse move event if we're faking touch events but the mouse button isn't down.
@@ -5756,7 +5762,7 @@ bool FSlateApplication::OnRawMouseMove( const int32 X, const int32 Y )
 		// convert to touch event if we are faking it
 		if (bIsFakingTouched)
 		{
-			return OnTouchMoved(GetCursorPos(), 1.0f, 0, 0);
+			return OnTouchMoved(GetCursorPos(), 1.0f, 0, FSlateApplicationBase::SlateAppPrimaryPlatformUser, IPlatformInputDeviceMapper::Get().GetDefaultInputDevice());
 		}
 
 		// Throw out the mouse move event if we're faking touch events but the mouse button isn't down.
@@ -6017,38 +6023,6 @@ bool FSlateApplication::OnControllerButtonReleased(FGamepadKeyNames::Type KeyNam
 	return false;
 }
 
-PRAGMA_DISABLE_DEPRECATION_WARNINGS
-bool FSlateApplication::OnControllerAnalog( FGamepadKeyNames::Type KeyName, int32 ControllerId, float AnalogValue )
-{
-	// Remap the old int32 ControlerId to the new Platform user for backwards compat
-	FPlatformUserId UserId = PLATFORMUSERID_NONE;
-	FInputDeviceId DeviceId = INPUTDEVICEID_NONE;
-	IPlatformInputDeviceMapper::Get().RemapControllerIdToPlatformUserAndDevice(ControllerId, UserId, DeviceId);
-
-	return OnControllerAnalog(KeyName, UserId, DeviceId, AnalogValue);
-}
-
-bool FSlateApplication::OnControllerButtonPressed(FGamepadKeyNames::Type KeyName, int32 ControllerId, bool IsRepeat)
-{
-	// Remap the old int32 ControlerId to the new Platform user for backwards compat
-	FPlatformUserId UserId = PLATFORMUSERID_NONE;
-	FInputDeviceId DeviceId = INPUTDEVICEID_NONE;
-	IPlatformInputDeviceMapper::Get().RemapControllerIdToPlatformUserAndDevice(ControllerId, UserId, DeviceId);
-
-	return OnControllerButtonPressed(KeyName, UserId, DeviceId, IsRepeat);
-}
-
-bool FSlateApplication::OnControllerButtonReleased(FGamepadKeyNames::Type KeyName, int32 ControllerId, bool IsRepeat)
-{
-	// Remap the old int32 ControlerId to the new Platform user for backwards compat
-	FPlatformUserId UserId = PLATFORMUSERID_NONE;
-	FInputDeviceId DeviceId = INPUTDEVICEID_NONE;
-	
-	IPlatformInputDeviceMapper::Get().RemapControllerIdToPlatformUserAndDevice(ControllerId, UserId, DeviceId);
-	return OnControllerButtonReleased(KeyName, UserId, DeviceId, IsRepeat);
-}
-PRAGMA_ENABLE_DEPRECATION_WARNINGS
-
 bool FSlateApplication::OnTouchGesture( EGestureEvent GestureType, const FVector2D &Delta, const float MouseWheelDelta, bool bIsDirectionInvertedFromDevice )
 {
 	const FVector2D CurrentCursorPosition = GetCursorPos();
@@ -6081,7 +6055,7 @@ bool ValidateTouchIndex(int32 TouchIndex)
 	return true;
 }
 
-bool FSlateApplication::OnTouchStarted( const TSharedPtr< FGenericWindow >& PlatformWindow, const FVector2D& Location, float Force, int32 TouchIndex, int32 ControllerId )
+bool FSlateApplication::OnTouchStarted( const TSharedPtr< FGenericWindow >& PlatformWindow, const FVector2D& Location, float Force, int32 TouchIndex, FPlatformUserId PlatformUserId, FInputDeviceId DeviceId )
 {
 	// Don't process touches that overlap or surpass with the cursor pointer index.
 	if (TouchIndex >= (int32)ETouchIndex::CursorPointerIndex)
@@ -6094,7 +6068,7 @@ bool FSlateApplication::OnTouchStarted( const TSharedPtr< FGenericWindow >& Plat
 	}
 
 	FPointerEvent PointerEvent(
-		ControllerId,
+		DeviceId,
 		TouchIndex,
 		Location,
 		Location,
@@ -6115,13 +6089,13 @@ void FSlateApplication::ProcessTouchStartedEvent( const TSharedPtr< FGenericWind
 	ProcessMouseButtonDownEvent(PlatformWindow, InTouchEvent);
 }
 
-bool FSlateApplication::OnTouchMoved( const FVector2D& Location, float Force, int32 TouchIndex, int32 ControllerId )
+bool FSlateApplication::OnTouchMoved( const FVector2D& Location, float Force, int32 TouchIndex, FPlatformUserId PlatformUserId, FInputDeviceId DeviceID )
 {
-	TSharedRef<FSlateUser> User = GetOrCreateUser(ControllerId);
+	TSharedRef<FSlateUser> User = GetOrCreateUser(DeviceID);
 	if (User->IsTouchPointerActive(TouchIndex))
 	{
 		FPointerEvent PointerEvent(
-			ControllerId,
+			DeviceID,
 			TouchIndex,
 			Location,
 			User->GetPreviousPointerPosition(TouchIndex),
@@ -6144,13 +6118,13 @@ void FSlateApplication::ProcessTouchMovedEvent( const FPointerEvent& PointerEven
 	ProcessMouseMoveEvent(PointerEvent);
 }
 
-bool FSlateApplication::OnTouchEnded( const FVector2D& Location, int32 TouchIndex, int32 ControllerId )
+bool FSlateApplication::OnTouchEnded( const FVector2D& Location, int32 TouchIndex, FPlatformUserId PlatformUserId, FInputDeviceId DeviceID )
 {
-	TSharedRef<FSlateUser> User = GetOrCreateUser(ControllerId);
+	TSharedRef<FSlateUser> User = GetOrCreateUser(DeviceID);
 	if (User->IsTouchPointerActive(TouchIndex))
 	{
 		FPointerEvent PointerEvent(
-			ControllerId,
+			DeviceID,
 			TouchIndex,
 			Location,
 			Location,
@@ -6177,13 +6151,13 @@ void FSlateApplication::ProcessTouchEndedEvent(const FPointerEvent& PointerEvent
 	ProcessMouseButtonUpEvent(PointerEvent);
 }
 
-bool FSlateApplication::OnTouchForceChanged(const FVector2D& Location, float Force, int32 TouchIndex, int32 ControllerId)
+bool FSlateApplication::OnTouchForceChanged(const FVector2D& Location, float Force, int32 TouchIndex, FPlatformUserId PlatformUserId, FInputDeviceId DeviceID)
 {
-	TSharedRef<FSlateUser> User = GetOrCreateUser(ControllerId);
+	TSharedRef<FSlateUser> User = GetOrCreateUser(DeviceID);
 	if (User->IsTouchPointerActive(TouchIndex))
 	{
 		FPointerEvent PointerEvent(
-			ControllerId,
+			DeviceID,
 			TouchIndex,
 			Location,
 			Location,
@@ -6199,13 +6173,13 @@ bool FSlateApplication::OnTouchForceChanged(const FVector2D& Location, float For
 	return false;
 }
 
-bool FSlateApplication::OnTouchFirstMove(const FVector2D& Location, float Force, int32 TouchIndex, int32 ControllerId)
+bool FSlateApplication::OnTouchFirstMove(const FVector2D& Location, float Force, int32 TouchIndex, FPlatformUserId PlatformUserId, FInputDeviceId DeviceID)
 {
-	TSharedRef<FSlateUser> User = GetOrCreateUser(ControllerId);
+	TSharedRef<FSlateUser> User = GetOrCreateUser(DeviceID);
 	if (User->IsTouchPointerActive(TouchIndex))
 	{
 		FPointerEvent PointerEvent(
-			ControllerId,
+			DeviceID,
 			TouchIndex,
 			Location,
 			User->GetPreviousPointerPosition(TouchIndex),
@@ -6228,10 +6202,10 @@ void FSlateApplication::ShouldSimulateGesture(EGestureEvent Gesture, bool bEnabl
 	SimulateGestures[(uint8)Gesture] = bEnable;
 }
 
-bool FSlateApplication::OnMotionDetected(const FVector& Tilt, const FVector& RotationRate, const FVector& Gravity, const FVector& Acceleration, int32 ControllerId)
+bool FSlateApplication::OnMotionDetected(const FVector& Tilt, const FVector& RotationRate, const FVector& Gravity, const FVector& Acceleration, FPlatformUserId PlatformUserId, FInputDeviceId InputDeviceId)
 {
 	FMotionEvent MotionEvent( 
-		ControllerId,
+		InputDeviceId,
 		Tilt,
 		RotationRate,
 		Gravity,
diff --git a/Engine/Source/Runtime/Slate/Public/Framework/Application/SlateApplication.h b/Engine/Source/Runtime/Slate/Public/Framework/Application/SlateApplication.h
--- a/Engine/Source/Runtime/Slate/Public/Framework/Application/SlateApplication.h
+++ b/Engine/Source/Runtime/Slate/Public/Framework/Application/SlateApplication.h
@@ -1029,6 +1029,13 @@ protected:
 	 * Asserts if given an invalid (ie, negative) index.
 	 */
 	TSharedRef<FSlateUser> GetOrCreateUser(FPlatformUserId PlatformUserId);
+
+	/**
+	 * Locates the SlateUser object corresponding to the input device id, creating a new one if it doesn't exist.
+	 * This will call GetOrCreateUser for the owning Platform User of the given input device.
+	 * Asserts if given an invalid (ie, negative) index.
+	 */
+	TSharedRef<FSlateUser> GetOrCreateUser(FInputDeviceId DeviceId);
 	
 	FORCEINLINE TSharedRef<FSlateUser> GetOrCreateUser(const FInputEvent& InputEvent) { return GetOrCreateUser(InputEvent.GetUserIndex()); }
 
@@ -1481,20 +1488,14 @@ public:
 	virtual bool OnMouseMove() override;
 	virtual bool OnRawMouseMove( const int32 X, const int32 Y ) override;
 	virtual bool OnCursorSet() override;
-	//UE_DEPRECATED(5.1, "This version of OnControllerAnalog has been deprecated, please use the one that takes an FPlatformUser and FInputDeviceId instead.")
-	virtual bool OnControllerAnalog( FGamepadKeyNames::Type KeyName, int32 ControllerId, float AnalogValue ) override;
-	//UE_DEPRECATED(5.1, "This version of OnControllerButtonPressed has been deprecated, please use the one that takes an FPlatformUser and FInputDeviceId instead.")
-	virtual bool OnControllerButtonPressed( FGamepadKeyNames::Type KeyName, int32 ControllerId, bool IsRepeat ) override;
-	//UE_DEPRECATED(5.1, "This version of OnControllerButtonReleased has been deprecated, please use the one that takes an FPlatformUser and FInputDeviceId instead.")
-	virtual bool OnControllerButtonReleased( FGamepadKeyNames::Type KeyName, int32 ControllerId, bool IsRepeat ) override;
 	virtual bool OnTouchGesture( EGestureEvent GestureType, const FVector2D& Delta, float WheelDelta, bool bIsDirectionInvertedFromDevice ) override;
-	virtual bool OnTouchStarted( const TSharedPtr< FGenericWindow >& PlatformWindow, const FVector2D& Location, float Force, int32 TouchIndex, int32 ControllerId ) override;
-	virtual bool OnTouchMoved( const FVector2D& Location, float Force, int32 TouchIndex, int32 ControllerId ) override;
-	virtual bool OnTouchEnded( const FVector2D& Location, int32 TouchIndex, int32 ControllerId ) override;
-	virtual bool OnTouchForceChanged(const FVector2D& Location, float Force, int32 TouchIndex, int32 ControllerId) override;
-	virtual bool OnTouchFirstMove(const FVector2D& Location, float Force, int32 TouchIndex, int32 ControllerId) override;
+	virtual bool OnTouchStarted( const TSharedPtr< FGenericWindow >& Window, const FVector2D& Location, float Force, int32 TouchIndex, FPlatformUserId PlatformUserId, FInputDeviceId DeviceId ) override;
+	virtual bool OnTouchMoved( const FVector2D& Location, float Force, int32 TouchIndex, FPlatformUserId PlatformUserId, FInputDeviceId DeviceID ) override;
+	virtual bool OnTouchEnded( const FVector2D& Location, int32 TouchIndex, FPlatformUserId PlatformUserId, FInputDeviceId DeviceID ) override;
+	virtual bool OnTouchForceChanged(const FVector2D& Location, float Force, int32 TouchIndex, FPlatformUserId PlatformUserId, FInputDeviceId DeviceID) override;
+	virtual bool OnTouchFirstMove(const FVector2D& Location, float Force, int32 TouchIndex, FPlatformUserId PlatformUserId, FInputDeviceId DeviceID) override;
 	virtual void ShouldSimulateGesture(EGestureEvent Gesture, bool bEnable) override;
-	virtual bool OnMotionDetected(const FVector& Tilt, const FVector& RotationRate, const FVector& Gravity, const FVector& Acceleration, int32 ControllerId) override;
+	virtual bool OnMotionDetected(const FVector& Tilt, const FVector& RotationRate, const FVector& Gravity, const FVector& Acceleration, FPlatformUserId PlatformUserId, FInputDeviceId InputDeviceId) override;
 	virtual bool OnControllerAnalog(FGamepadKeyNames::Type KeyName, FPlatformUserId PlatformUserId, FInputDeviceId InputDeviceId, float AnalogValue) override;
 	virtual bool OnControllerButtonPressed(FGamepadKeyNames::Type KeyName, FPlatformUserId PlatformUserId, FInputDeviceId InputDeviceId, bool IsRepeat) override;
 	virtual bool OnControllerButtonReleased(FGamepadKeyNames::Type KeyName, FPlatformUserId PlatformUserId, FInputDeviceId InputDeviceId, bool IsRepeat) override;
