Replaced #include CoreMiscDefines.h with #include CoreTypes.h because you can't actually include CoreMiscDefines.h unless you include CoreTypes.h first. Unfortunately adding the necessary includes into CoreMiscDefines.h breaks other things ATM.

diff --git a/Engine/Source/Runtime/Core/Public/Algo/RemoveIf.h b/Engine/Source/Runtime/Core/Public/Algo/RemoveIf.h
--- a/Engine/Source/Runtime/Core/Public/Algo/RemoveIf.h
+++ b/Engine/Source/Runtime/Core/Public/Algo/RemoveIf.h
@@ -2,8 +2,8 @@
 
 #pragma once
 
+#include "CoreTypes.h"
 #include "Templates/Invoke.h"
-#include "Misc/CoreMiscDefines.h"
 
 namespace Algo
 {
diff --git a/Engine/Source/Runtime/Core/Public/Misc/LowLevelTestAdapter.h b/Engine/Source/Runtime/Core/Public/Misc/LowLevelTestAdapter.h
--- a/Engine/Source/Runtime/Core/Public/Misc/LowLevelTestAdapter.h
+++ b/Engine/Source/Runtime/Core/Public/Misc/LowLevelTestAdapter.h
@@ -6,7 +6,6 @@
 
 #include "CoreTypes.h"
 #include "Misc/AssertionMacros.h"
-#include "CoreMiscDefines.h"
 
 #define IMPLEMENT_SIMPLE_AUTOMATION_TEST_PRIVATE_LLT( TClass, PrettyName, TFlags, FileName, LineNumber ) \
 		class TClass : public FAutomationTestBase \
diff --git a/Engine/Source/Runtime/Core/Public/Misc/OutputDevice.h b/Engine/Source/Runtime/Core/Public/Misc/OutputDevice.h
--- a/Engine/Source/Runtime/Core/Public/Misc/OutputDevice.h
+++ b/Engine/Source/Runtime/Core/Public/Misc/OutputDevice.h
@@ -5,7 +5,6 @@
 #include "CoreFwd.h"
 #include "CoreTypes.h"
 #include "Logging/LogVerbosity.h"
-#include "Misc/CoreMiscDefines.h"
 #include "Misc/VarArgs.h"
 #include "Templates/AndOrNot.h"
 #include "Templates/IsArrayOrRefOfTypeByPredicate.h"
diff --git a/Engine/Source/Runtime/Engine/Classes/Particles/ParticlePerfStats.h b/Engine/Source/Runtime/Engine/Classes/Particles/ParticlePerfStats.h
--- a/Engine/Source/Runtime/Engine/Classes/Particles/ParticlePerfStats.h
+++ b/Engine/Source/Runtime/Engine/Classes/Particles/ParticlePerfStats.h
@@ -4,8 +4,8 @@
 
 #include "Templates/Atomic.h"
 #include "CoreMinimal.h"
+#include "CoreTypes.h"
 #include "HAL/PlatformTime.h"
-#include "Misc/CoreMiscDefines.h"
 #include "Containers/Array.h"
 #include "ProfilingDebugging/CsvProfiler.h"
 
diff --git a/Engine/Source/Runtime/Engine/Public/Animation/AttributesRuntime.h b/Engine/Source/Runtime/Engine/Public/Animation/AttributesRuntime.h
--- a/Engine/Source/Runtime/Engine/Public/Animation/AttributesRuntime.h
+++ b/Engine/Source/Runtime/Engine/Public/Animation/AttributesRuntime.h
@@ -8,8 +8,8 @@
 #include "Containers/ContainersFwd.h"
 
 #include "Containers/Map.h"
+#include "CoreTypes.h"
 #include "Misc/AssertionMacros.h"
-#include "Misc/CoreMiscDefines.h"
 #include "UObject/StructOnScope.h"
 
 #include "Animation/IAttributeBlendOperator.h"
diff --git a/Engine/Source/Runtime/Experimental/Chaos/Public/Chaos/Defines.h b/Engine/Source/Runtime/Experimental/Chaos/Public/Chaos/Defines.h
--- a/Engine/Source/Runtime/Experimental/Chaos/Public/Chaos/Defines.h
+++ b/Engine/Source/Runtime/Experimental/Chaos/Public/Chaos/Defines.h
@@ -24,7 +24,6 @@
 #else
 	#include "CoreTypes.h"
 	#include "Logging/MessageLog.h"
-	#include "Misc/CoreMiscDefines.h"
 	#include "UObject/ExternalPhysicsCustomObjectVersion.h"
 #endif
 #include "Serializable.h"
diff --git a/Engine/Source/Runtime/Slate/Public/Widgets/Input/SEditableLabel.h b/Engine/Source/Runtime/Slate/Public/Widgets/Input/SEditableLabel.h
--- a/Engine/Source/Runtime/Slate/Public/Widgets/Input/SEditableLabel.h
+++ b/Engine/Source/Runtime/Slate/Public/Widgets/Input/SEditableLabel.h
@@ -3,6 +3,7 @@
 #pragma once
 
 #include "CoreMinimal.h"
+#include "CoreTypes.h"
 #include "Misc/Attribute.h"
 #include "Styling/SlateColor.h"
 #include "Fonts/SlateFontInfo.h"
@@ -13,7 +14,6 @@
 #include "Styling/SlateTypes.h"
 #include "Styling/CoreStyle.h"
 #include "Framework/SlateDelegates.h"
-#include "Misc/CoreMiscDefines.h"
 
 class SEditableText;
 class STextBlock;
