Fix slate losing focus when virutal keyboard is up on AppleTV

diff --git a/Engine/Source/Runtime/Slate/Private/Framework/Text/IOS/IOSPlatformTextField.cpp b/Engine/Source/Runtime/Slate/Private/Framework/Text/IOS/IOSPlatformTextField.cpp
--- a/Engine/Source/Runtime/Slate/Private/Framework/Text/IOS/IOSPlatformTextField.cpp
+++ b/Engine/Source/Runtime/Slate/Private/Framework/Text/IOS/IOSPlatformTextField.cpp
@@ -102,6 +102,7 @@ void FIOSPlatformTextField::ShowVirtualKeyboard(bool bShow, int32 UserIndex, TSh
 		}
         else
         {
+#if !PLATFORM_TVOS
 			if (TextField != nullptr && [TextField hasTextWidget])
 			{
                 UE_LOG(LogIOS, Log, TEXT("Hiding field: %p"), TextField);
@@ -111,9 +112,10 @@ void FIOSPlatformTextField::ShowVirtualKeyboard(bool bShow, int32 UserIndex, TSh
                     if (LocalTextField != nullptr)
                     {
                         [LocalTextField hide];
-                     }
+                    }
                 });
 			}
+#endif
         }
 	}
 }
