Slate: SScaleBox now clips when the stretch type is user specified.
#rnx

diff --git a/Engine/Source/Runtime/Slate/Private/Widgets/Layout/SScaleBox.cpp b/Engine/Source/Runtime/Slate/Private/Widgets/Layout/SScaleBox.cpp
--- a/Engine/Source/Runtime/Slate/Private/Widgets/Layout/SScaleBox.cpp
+++ b/Engine/Source/Runtime/Slate/Private/Widgets/Layout/SScaleBox.cpp
@@ -137,6 +137,7 @@ bool SScaleBox::DoesScaleRequireNormalizingPrepassOrLocalGeometry() const
 	case EStretch::Fill:
 	case EStretch::ScaleBySafeZone:
 	case EStretch::UserSpecified:
+	case EStretch::UserSpecifiedWithClipping:
 		return false;
 	default:
 		return true;
@@ -166,6 +167,7 @@ float SScaleBox::ComputeContentScale(const FGeometry& PaintGeometry) const
 	case EStretch::ScaleBySafeZone:
 		return SafeZoneScale;
 	case EStretch::UserSpecified:
+	case EStretch::UserSpecifiedWithClipping:
 		return UserSpecifiedScaleAttribute.Get();
 	}
 
@@ -310,6 +312,7 @@ int32 SScaleBox::OnPaint( const FPaintArgs& Args, const FGeometry& AllottedGeome
 		case EStretch::ScaleToFitX:
 		case EStretch::ScaleToFitY:
 		case EStretch::ScaleToFill:
+		case EStretch::UserSpecifiedWithClipping:
 			bClippingNeeded = true;
 			break;
 		}
diff --git a/Engine/Source/Runtime/Slate/Public/Widgets/Layout/SScaleBox.h b/Engine/Source/Runtime/Slate/Public/Widgets/Layout/SScaleBox.h
--- a/Engine/Source/Runtime/Slate/Public/Widgets/Layout/SScaleBox.h
+++ b/Engine/Source/Runtime/Slate/Public/Widgets/Layout/SScaleBox.h
@@ -62,7 +62,9 @@ namespace EStretch
 		/** Scales the content according to the size of the safe zone currently applied to the viewport. */
 		ScaleBySafeZone,
 		/** Scales the content by the scale specified by the user. */
-		UserSpecified
+		UserSpecified,
+		/** Scales the content by the scale specified by the user and also clips. */
+		UserSpecifiedWithClipping
 	};
 }
 
diff --git a/Engine/Source/Runtime/UMG/Private/Components/ScaleBox.cpp b/Engine/Source/Runtime/UMG/Private/Components/ScaleBox.cpp
--- a/Engine/Source/Runtime/UMG/Private/Components/ScaleBox.cpp
+++ b/Engine/Source/Runtime/UMG/Private/Components/ScaleBox.cpp
@@ -136,12 +136,13 @@ bool UScaleBox::CanEditChange(const FProperty* InProperty) const
 
 		if (PropertyName == GET_MEMBER_NAME_CHECKED(UScaleBox, StretchDirection))
 		{
-			return Stretch != EStretch::None && Stretch != EStretch::ScaleBySafeZone && Stretch != EStretch::UserSpecified;
+			return Stretch != EStretch::None && Stretch != EStretch::ScaleBySafeZone &&
+				Stretch != EStretch::UserSpecified && Stretch != EStretch::UserSpecifiedWithClipping;
 		}
 
 		if (PropertyName == GET_MEMBER_NAME_CHECKED(UScaleBox, UserSpecifiedScale))
 		{
-			return Stretch == EStretch::UserSpecified;
+			return Stretch == EStretch::UserSpecified || Stretch == EStretch::UserSpecifiedWithClipping;
 		}
 	}
 
