Fix layouts not saving in shipping configuration

diff --git a/Engine/Source/Runtime/Slate/Private/Framework/Docking/TabManager.cpp b/Engine/Source/Runtime/Slate/Private/Framework/Docking/TabManager.cpp
--- a/Engine/Source/Runtime/Slate/Private/Framework/Docking/TabManager.cpp
+++ b/Engine/Source/Runtime/Slate/Private/Framework/Docking/TabManager.cpp
@@ -387,7 +387,10 @@ FString FTabManager::FLayout::ToString() const
 
 	FString LayoutAsString;
 	TSharedRef< TJsonWriter<> > Writer = TJsonWriterFactory<>::Create( &LayoutAsString );
-	check( FJsonSerializer::Serialize( LayoutJson, Writer ) );
+	if (!FJsonSerializer::Serialize(LayoutJson, Writer))
+	{
+		UE_LOG(LogSlate, Error, TEXT("Failed save layout as Json string: %s"), *GetLayoutName().ToString());
+	}
 
 	return LayoutAsString;
 }
