Allow for keyboard input overrides for Spinbox

diff --git a/Engine/Source/Runtime/Slate/Public/Widgets/Input/SSpinBox.h b/Engine/Source/Runtime/Slate/Public/Widgets/Input/SSpinBox.h
--- a/Engine/Source/Runtime/Slate/Public/Widgets/Input/SSpinBox.h
+++ b/Engine/Source/Runtime/Slate/Public/Widgets/Input/SSpinBox.h
@@ -153,6 +153,8 @@ public:
 		SLATE_ATTRIBUTE(float, MinDesiredWidth)
 		/** How should the value be justified in the spinbox. */
 		SLATE_ATTRIBUTE(ETextJustify::Type, Justification)
+		/** What keyboard to display. */
+		SLATE_ATTRIBUTE(EKeyboardType, KeyboardType)
 		/** Provide custom type conversion functionality to this spin box */
 		SLATE_ARGUMENT(TSharedPtr< INumericTypeInterface<NumericType> >, TypeInterface)
 		/** If refresh requests for the viewport should happen for all value changes **/
@@ -286,7 +288,7 @@ public:
 				.ClearKeyboardFocusOnCommit(InArgs._ClearKeyboardFocusOnCommit)
 				.SelectAllTextOnCommit(InArgs._SelectAllTextOnCommit)
 				.MinDesiredWidth(this, &SSpinBox<NumericType>::GetTextMinDesiredWidth)
-				.VirtualKeyboardType(EKeyboardType::Keyboard_Number)
+				.VirtualKeyboardType(InArgs._KeyboardType)
 				.Justification(InArgs._Justification)
 				.VirtualKeyboardTrigger(EVirtualKeyboardTrigger::OnAllFocusEvents)
 				.ContextMenuExtender(InArgs._ContextMenuExtender)
diff --git a/Engine/Source/Runtime/UMG/Private/Components/SpinBox.cpp b/Engine/Source/Runtime/UMG/Private/Components/SpinBox.cpp
--- a/Engine/Source/Runtime/UMG/Private/Components/SpinBox.cpp
+++ b/Engine/Source/Runtime/UMG/Private/Components/SpinBox.cpp
@@ -39,6 +39,7 @@ USpinBox::USpinBox(const FObjectInitializer& ObjectInitializer)
 	MinDesiredWidth = 0;
 	ClearKeyboardFocusOnCommit = false;
 	SelectAllTextOnCommit = true;
+	KeyboardType = EVirtualKeyboardType::Number;
 
 	if (DefaultSpinBoxStyle == nullptr)
 	{
@@ -86,6 +87,7 @@ TSharedRef<SWidget> USpinBox::RebuildWidget()
 	.ClearKeyboardFocusOnCommit(ClearKeyboardFocusOnCommit)
 	.SelectAllTextOnCommit(SelectAllTextOnCommit)
 	.Justification(Justification)
+	.KeyboardType(EVirtualKeyboardType::AsKeyboardType(KeyboardType.GetValue()))
 	.OnValueChanged(BIND_UOBJECT_DELEGATE(FOnFloatValueChanged, HandleOnValueChanged))
 	.OnValueCommitted(BIND_UOBJECT_DELEGATE(FOnFloatValueCommitted, HandleOnValueCommitted))
 	.OnBeginSliderMovement(BIND_UOBJECT_DELEGATE(FSimpleDelegate, HandleOnBeginSliderMovement))
diff --git a/Engine/Source/Runtime/UMG/Public/Components/SpinBox.h b/Engine/Source/Runtime/UMG/Public/Components/SpinBox.h
--- a/Engine/Source/Runtime/UMG/Public/Components/SpinBox.h
+++ b/Engine/Source/Runtime/UMG/Public/Components/SpinBox.h
@@ -75,6 +75,10 @@ public:
 	/** The minimum width of the spin box */
 	UPROPERTY(EditAnywhere, BlueprintReadOnly, Category="Display", AdvancedDisplay, DisplayName = "Minimum Desired Width")
 	float MinDesiredWidth;
+	
+	/** If we're on a platform that requires a virtual keyboard, what kind of keyboard should this widget use? */
+	UPROPERTY(EditAnywhere, BlueprintReadOnly, Category = "Input", AdvancedDisplay)
+	TEnumAsByte<EVirtualKeyboardType::Type> KeyboardType;
 
 	/** Whether to remove the keyboard focus from the spin box when the value is committed */
 	UPROPERTY(EditAnywhere, BlueprintReadOnly, Category="Input", AdvancedDisplay)
