Make the GetOrCreateUser function on the slate app use the input device ID instead of the old User index.

#rnx

diff --git a/Engine/Source/Runtime/Slate/Public/Framework/Application/SlateApplication.h b/Engine/Source/Runtime/Slate/Public/Framework/Application/SlateApplication.h
--- a/Engine/Source/Runtime/Slate/Public/Framework/Application/SlateApplication.h
+++ b/Engine/Source/Runtime/Slate/Public/Framework/Application/SlateApplication.h
@@ -1037,7 +1037,7 @@ protected:
 	 */
 	TSharedRef<FSlateUser> GetOrCreateUser(FInputDeviceId DeviceId);
 	
-	FORCEINLINE TSharedRef<FSlateUser> GetOrCreateUser(const FInputEvent& InputEvent) { return GetOrCreateUser(InputEvent.GetUserIndex()); }
+	FORCEINLINE TSharedRef<FSlateUser> GetOrCreateUser(const FInputEvent& InputEvent) { return GetOrCreateUser(InputEvent.GetInputDeviceId()); }
 
 	friend class FEventRouter;
 
