PR #9383: Fix for multiple WidgetInteractionComponents focusing on the same widget (Contributed by [REDACTED_PERSON])

diff --git a/Engine/Source/Runtime/Slate/Private/Framework/Application/SlateApplication.cpp b/Engine/Source/Runtime/Slate/Private/Framework/Application/SlateApplication.cpp
--- a/Engine/Source/Runtime/Slate/Private/Framework/Application/SlateApplication.cpp
+++ b/Engine/Source/Runtime/Slate/Private/Framework/Application/SlateApplication.cpp
@@ -5220,7 +5220,28 @@ bool FSlateApplication::RoutePointerMoveEvent(const FWidgetPath& WidgetsUnderPoi
 				{
 					TOptional<FArrangedWidget> FoundWidget = WidgetsUnderPointer.FindArrangedWidget(SomeWidgetPreviouslyUnderCursor.ToSharedRef());
 					const bool bWidgetNoLongerUnderMouse = !FoundWidget.IsSet();
+					
+					bool bWidgetUnderOtherUsersPointer = false;
+					
+					// Verify if the Widget is under another user pointer
 					if (bWidgetNoLongerUnderMouse)
+					{
+						for (const TSharedPtr<FSlateUser>& user : Users)
+						{
+							if (user != nullptr && user->GetUserIndex() != PointerEvent.GetUserIndex())
+							{
+								if (user->IsWidgetDirectlyUnderAnyPointer(SomeWidgetPreviouslyUnderCursor))
+								{
+									bWidgetUnderOtherUsersPointer = true;
+
+									break;
+								}
+							}
+						}
+					}
+
+					// We consider the Widget is nolonger under a pointer if it's no longer under any pointers.
+					if (bWidgetNoLongerUnderMouse && !bWidgetUnderOtherUsersPointer)
 					{
 						// Widget is no longer under cursor, so send a MouseLeave.
 						// The widget might not even be in the hierarchy any more!
