Altered the ContentPadding of the WrapButton ComboButton so that the button appears as a circle, instead of the football shape that appeared previously.

diff --git a/Engine/Source/Runtime/Slate/Private/Framework/MultiBox/SClippingHorizontalBox.cpp b/Engine/Source/Runtime/Slate/Private/Framework/MultiBox/SClippingHorizontalBox.cpp
--- a/Engine/Source/Runtime/Slate/Private/Framework/MultiBox/SClippingHorizontalBox.cpp
+++ b/Engine/Source/Runtime/Slate/Private/Framework/MultiBox/SClippingHorizontalBox.cpp
@@ -101,7 +101,7 @@ void SClippingHorizontalBox::AddWrapButton()
 		SNew( SComboButton )
 		.HasDownArrow( false )
 		.ButtonStyle(&ToolBarStyle.ButtonStyle)
-		.ContentPadding( 0.f )
+		.ContentPadding( FMargin(4.f, 0.f) )
 		.ToolTipText( NSLOCTEXT("Slate", "ExpandToolbar", "Click to expand toolbar") )
 		.OnGetMenuContent( OnWrapButtonClicked )
 		.Cursor( EMouseCursor::Default )
