GetUserIndexForXXX functions in ISlateInputManager are TOptional<>

#rnx

diff --git a/Engine/Plugins/EnhancedInput/Source/EnhancedInput/Private/EnhancedInputSubsystemInterfaceDebug.cpp b/Engine/Plugins/EnhancedInput/Source/EnhancedInput/Private/EnhancedInputSubsystemInterfaceDebug.cpp
--- a/Engine/Plugins/EnhancedInput/Source/EnhancedInput/Private/EnhancedInputSubsystemInterfaceDebug.cpp
+++ b/Engine/Plugins/EnhancedInput/Source/EnhancedInput/Private/EnhancedInputSubsystemInterfaceDebug.cpp
@@ -314,8 +314,15 @@ void IEnhancedInputSubsystemInterface::ShowPlatformInputDebugInfo(class UCanvas*
 			FString SlateUserIdString = TEXT("Slate UserID is unknown (Application is not initialized)");
 			if (FSlateApplication::IsInitialized())
 			{
-				const int32 SlateUserId = FSlateApplication::Get().GetUserIndexForInputDevice(Device);
-				SlateUserIdString = FString::Printf(TEXT("Slate UserId ID: %d"), SlateUserId);				
+				TOptional<int32> SlateUserId = FSlateApplication::Get().GetUserIndexForInputDevice(Device);
+				if (SlateUserId.IsSet())
+				{
+					SlateUserIdString = FString::Printf(TEXT("Slate UserId ID: %d"), SlateUserId.GetValue());
+				}
+				else
+				{
+					SlateUserIdString = TEXT("Slate UserId ID: invalid");
+				}
 			}
 			
 			DisplayDebugManager.SetDrawColor(FColor::White);
diff --git a/Engine/Source/Runtime/Slate/Private/Framework/Application/SlateApplication.cpp b/Engine/Source/Runtime/Slate/Private/Framework/Application/SlateApplication.cpp
--- a/Engine/Source/Runtime/Slate/Private/Framework/Application/SlateApplication.cpp
+++ b/Engine/Source/Runtime/Slate/Private/Framework/Application/SlateApplication.cpp
@@ -2428,12 +2428,12 @@ int32 FSlateApplication::GetUserIndexForController(int32 ControllerId) const
 	return InputManager->GetUserIndexForController(ControllerId);
 }
 
-int32 FSlateApplication::GetUserIndexForInputDevice(FInputDeviceId InputDeviceId) const
+TOptional<int32> FSlateApplication::GetUserIndexForInputDevice(FInputDeviceId InputDeviceId) const
 {
 	return InputManager->GetUserIndexForInputDevice(InputDeviceId);
 }
 
-int32 FSlateApplication::GetUserIndexForPlatformUser(FPlatformUserId PlatformUser) const
+TOptional<int32> FSlateApplication::GetUserIndexForPlatformUser(FPlatformUserId PlatformUser) const
 {
 	return InputManager->GetUserIndexForPlatformUser(PlatformUser);
 }
diff --git a/Engine/Source/Runtime/Slate/Public/Framework/Application/SlateApplication.h b/Engine/Source/Runtime/Slate/Public/Framework/Application/SlateApplication.h
--- a/Engine/Source/Runtime/Slate/Public/Framework/Application/SlateApplication.h
+++ b/Engine/Source/Runtime/Slate/Public/Framework/Application/SlateApplication.h
@@ -124,13 +124,13 @@ public:
 	virtual FInputDeviceId GetInputDeviceIdForMouse() const = 0;
 	virtual FInputDeviceId GetInputDeviceIdForKeyboard() const = 0;
 	
-	virtual int32 GetUserIndexForInputDevice(FInputDeviceId InputDeviceId) const
+	virtual TOptional<int32> GetUserIndexForInputDevice(FInputDeviceId InputDeviceId) const
 	{
 		// There is a 1:1 mapping of the platform user ID to a slate user index.
 		return GetUserIndexForPlatformUser(IPlatformInputDeviceMapper::Get().GetUserForInputDevice(InputDeviceId));
 	}
 
-	virtual int32 GetUserIndexForPlatformUser(FPlatformUserId PlatformUser) const
+	virtual TOptional<int32> GetUserIndexForPlatformUser(FPlatformUserId PlatformUser) const
 	{
 		// There is a 1:1 mapping of the platform user ID to a slate user index.
 		return PlatformUser.GetInternalId();
@@ -1596,9 +1596,9 @@ public:
 	int32 GetUserIndexForController(int32 ControllerId) const;
 
 	/** @return Gets the slate user index that this input device is mapped to */	
-	int32 GetUserIndexForInputDevice(FInputDeviceId InputDeviceId) const;
+	TOptional<int32> GetUserIndexForInputDevice(FInputDeviceId InputDeviceId) const;
 	/** @return Gets the slate user index that this platform user id mapped to*/	
-	int32 GetUserIndexForPlatformUser(FPlatformUserId PlatformUser) const;
+	TOptional<int32> GetUserIndexForPlatformUser(FPlatformUserId PlatformUser) const;
 	
 	TOptional<int32> GetUserIndexForController(int32 ControllerId, FKey InKey) const;
 
