synthesized mouse input uses the platform user's default input device.

#rnx

diff --git a/Engine/Source/Runtime/Slate/Private/Framework/Application/SlateUser.cpp b/Engine/Source/Runtime/Slate/Private/Framework/Application/SlateUser.cpp
--- a/Engine/Source/Runtime/Slate/Private/Framework/Application/SlateUser.cpp
+++ b/Engine/Source/Runtime/Slate/Private/Framework/Application/SlateUser.cpp
@@ -6,6 +6,7 @@
 #include "Misc/App.h"
 #include "Widgets/SWindow.h"
 #include "Widgets/SWeakWidget.h"
+#include "GenericPlatform/GenericPlatformInputDeviceMapper.h"
 
 #if PLATFORM_MICROSOFT
 // Needed to be able to use RECT
@@ -668,10 +669,13 @@ bool FSlateUser::SynthesizeCursorMoveIfNeeded()
 
 		FSlateApplication& SlateApp = FSlateApplication::Get();
 		
+		FInputDeviceId InputDeviceId = IPlatformInputDeviceMapper::Get().GetPrimaryInputDeviceForUser( GetPlatformUserId() );
+		check(InputDeviceId.IsValid());
+
 		const bool bHasHardwareCursor = SlateApp.GetPlatformCursor() == Cursor;
 		const TSet<FKey> EmptySet;
 		FPointerEvent SyntheticCursorMoveEvent(
-			GetUserIndex(),
+			InputDeviceId,
 			FSlateApplication::CursorPointerIndex,
 			GetCursorPosition(),
 			GetPreviousCursorPosition(),
