Comment cleanup.

diff --git a/Engine/Source/Runtime/Slate/Private/Widgets/Input/SSpinBox.cpp b/Engine/Source/Runtime/Slate/Private/Widgets/Input/SSpinBox.cpp
--- a/Engine/Source/Runtime/Slate/Private/Widgets/Input/SSpinBox.cpp
+++ b/Engine/Source/Runtime/Slate/Private/Widgets/Input/SSpinBox.cpp
@@ -1,6 +1,5 @@
 // Copyright Epic Games, Inc. All Rights Reserved.
 #include "Widgets/Input/SSpinBox.h"
-//#include "SlatePrivatePCH.h"
 
 float SpinBoxComputeExponentSliderFraction(float FractionFilled, float StartFractionFilled, float SliderExponent)
 {
