Reverted UMG GridPanel behaviour changes with regards to row/column spans that go past the end of the grid.

Instead, we warn about using this behaviour and clamp the span to 10,000.

diff --git a/Engine/Source/Runtime/Slate/Private/Widgets/Layout/SGridPanel.cpp b/Engine/Source/Runtime/Slate/Private/Widgets/Layout/SGridPanel.cpp
--- a/Engine/Source/Runtime/Slate/Private/Widgets/Layout/SGridPanel.cpp
+++ b/Engine/Source/Runtime/Slate/Private/Widgets/Layout/SGridPanel.cpp
@@ -255,13 +255,9 @@ void SGridPanel::OnArrangeChildren( const FGeometry& AllottedGeometry, FArranged
 
 			// Figure out the size of this slot; takes row span into account.
 			// We use the properties of partial sums arrays to achieve this.
-			// Row/column span may be an arbitrarily large value above 1, so clamp it to the array bounds.
-			const int32 LastColumnIdx = FMath::Min(CurSlot.GetColumn() + CurSlot.GetColumnSpan(), FinalColumns.Num() - 1);
-			const int32 LastRowIdx = FMath::Min(CurSlot.GetRow() + CurSlot.GetRowSpan(), FinalRows.Num() - 1);
 			const FVector2D CellSize(
-				FinalColumns[LastColumnIdx] - ThisCellOffset.X,
-				FinalRows[LastRowIdx] - ThisCellOffset.Y
-			);
+				FinalColumns[CurSlot.GetColumn() + CurSlot.GetColumnSpan()] - ThisCellOffset.X,
+				FinalRows[CurSlot.GetRow() + CurSlot.GetRowSpan()] - ThisCellOffset.Y);
 
 			// Do the standard arrangement of elements within a slot
 			// Takes care of alignment and padding.
@@ -408,15 +404,14 @@ void SGridPanel::NotifySlotChanged(const FSlot* InSlot, bool bSlotLayerChanged /
 	// We need an extra cell at the end for easily figuring out the size across any number of cells
 	// by doing Columns[End] - Columns[Start] or Rows[End] - Rows[Start].
 	// The first Columns[]/Rows[] entry will be 0.
-	// For a grid with 2 rows (index 0 and 1), we need 3 slots, thus the + 2 here.
 
-	const int32 NumColumnsRequiredForThisSlot = InSlot->GetColumn() + 2;
+	const int32 NumColumnsRequiredForThisSlot = InSlot->GetColumn() + InSlot->GetColumnSpan() + 1;
 	if (NumColumnsRequiredForThisSlot > Columns.Num())
 	{
 		Columns.AddZeroed(NumColumnsRequiredForThisSlot - Columns.Num());
 	}
 
-	const int32 NumRowsRequiredForThisSlot = InSlot->GetRow() + 2;
+	const int32 NumRowsRequiredForThisSlot = InSlot->GetRow() + InSlot->GetColumnSpan() + 1;
 	if (NumRowsRequiredForThisSlot > Rows.Num())
 	{
 		Rows.AddZeroed(NumRowsRequiredForThisSlot - Rows.Num());
@@ -447,18 +442,12 @@ void SGridPanel::ComputeDesiredCellSizes( TArray<float>& OutColumns, TArray<floa
 			const FVector2D SlotDesiredSize = CurSlot.GetWidget()->GetDesiredSize() + CurSlot.GetPadding().GetDesiredSize();
 
 			// If the slot has a (colspan, rowspan) of (1,1) it will only affect that slot.
-			// For larger spans, the slots size will be evenly distributed across all the affected slots.
-			// For spans that are past the end of the slot array, only distribute to the size to existing columns, eg.
-			// | [0] | [1] | [2] |
-			// | S=1 |   Span=5  | <- size is only distributed between columns 1 and 2, until more slots exist
-			const int32 MaxColumns = FMath::Min(CurSlot.GetColumn() + CurSlot.GetColumnSpan(), OutColumns.Num());
-			const int32 MaxRows = FMath::Min(CurSlot.GetRow() + CurSlot.GetRowSpan(), OutRows.Num());
-
-			const FVector2D SizeContribution(SlotDesiredSize.X / (MaxColumns - CurSlot.GetColumn()), SlotDesiredSize.Y / (MaxRows - CurSlot.GetRow()));
+			// For larger spans, the slot's size will be evenly distributed across all the affected slots.
+			const FVector2D SizeContribution(SlotDesiredSize.X / CurSlot.GetColumnSpan(), SlotDesiredSize.Y / CurSlot.GetRowSpan());
 
 			// Distribute the size contributions over all the columns and rows that this slot spans
-			DistributeSizeContributions(SizeContribution.X, OutColumns, CurSlot.GetColumn(), MaxColumns);
-			DistributeSizeContributions(SizeContribution.Y, OutRows, CurSlot.GetRow(), MaxRows);
+			DistributeSizeContributions(SizeContribution.X, OutColumns, CurSlot.GetColumn(), CurSlot.GetColumn() + CurSlot.GetColumnSpan());
+			DistributeSizeContributions(SizeContribution.Y, OutRows, CurSlot.GetRow(), CurSlot.GetRow() + CurSlot.GetRowSpan());
 		}
 	}
 }
diff --git a/Engine/Source/Runtime/Slate/Public/Widgets/Layout/SGridPanel.h b/Engine/Source/Runtime/Slate/Public/Widgets/Layout/SGridPanel.h
--- a/Engine/Source/Runtime/Slate/Public/Widgets/Layout/SGridPanel.h
+++ b/Engine/Source/Runtime/Slate/Public/Widgets/Layout/SGridPanel.h
@@ -92,7 +92,8 @@ public:
 
 			void SetColumnSpan(int32 ColumnSpan)
 			{
-				ColumnSpan = FMath::Max(1, ColumnSpan);
+				// clamp span to a sensible size, otherwise computing slot sizes can slow down dramatically
+				ColumnSpan = FMath::Clamp(ColumnSpan, 1, 10000);
 				if (ColumnSpan != ColumnSpanParam)
 				{
 					ColumnSpanParam = ColumnSpan;
@@ -124,7 +125,8 @@ public:
 
 			void SetRowSpan(int32 RowSpan)
 			{
-				RowSpan = FMath::Max(1, RowSpan);
+				// clamp span to a sensible size, otherwise computing slots sizes can slow down dramatically
+				RowSpan = FMath::Clamp(RowSpan, 1, 10000);
 				if (RowSpan != RowSpanParam)
 				{
 					RowSpanParam = RowSpan;
diff --git a/Engine/Source/Runtime/UMG/Private/Components/GridPanel.cpp b/Engine/Source/Runtime/UMG/Private/Components/GridPanel.cpp
--- a/Engine/Source/Runtime/UMG/Private/Components/GridPanel.cpp
+++ b/Engine/Source/Runtime/UMG/Private/Components/GridPanel.cpp
@@ -4,6 +4,7 @@
 #include "Widgets/DeclarativeSyntaxSupport.h"
 #include "Widgets/Layout/SGridPanel.h"
 #include "Components/GridSlot.h"
+#include "Editor/WidgetCompilerLog.h"
 
 #define LOCTEXT_NAMESPACE "UMG"
 
@@ -132,6 +133,47 @@ const FText UGridPanel::GetPaletteCategory()
 	return LOCTEXT("Panel", "Panel");
 }
 
+void UGridPanel::ValidateCompiledDefaults(IWidgetCompilerLog& CompileLog) const
+{
+	int32 NumRows = 0;
+	int32 NumColumns = 0;
+
+	// find the maximum row & column
+	for (int32 Idx = 0; Idx < Slots.Num(); ++Idx)
+	{
+		if (const UGridSlot* ChildSlot = Cast<UGridSlot>(Slots[Idx].Get()))
+		{
+			NumRows = FMath::Max(NumRows, ChildSlot->GetRow());
+			NumColumns = FMath::Max(NumColumns, ChildSlot->GetColumn());
+		}
+	}
+
+	// Row span of 1 for the last row is valid
+	NumRows += 1;
+	NumColumns += 1;
+
+	// check if any row span goes past the end of the grid and warn
+	for (int32 Idx = 0; Idx < Slots.Num(); ++Idx)
+	{
+		if (const UGridSlot* ChildSlot = Cast<UGridSlot>(Slots[Idx].Get()))
+		{
+			if (ChildSlot->GetRow() + ChildSlot->GetRowSpan() > NumRows)
+			{
+				FText InfoText = FText::FormatOrdered(LOCTEXT("RowSpanPastEnd", "Slot at row {0}, column {1} has a row span value of {2}, which goes past the end of the grid. This behaviour will be deprecated in the future. Slots should not use row span to stretch themselves past the end of the grid."), ChildSlot->GetRow(), ChildSlot->GetColumn(), ChildSlot->GetRowSpan());
+				CompileLog.Warning(InfoText);
+			}
+
+			if (ChildSlot->GetColumn() + ChildSlot->GetColumnSpan() > NumColumns)
+			{
+				FText InfoText = FText::FormatOrdered(LOCTEXT("ColumnSpanPastEnd", "Slot at row {0}, column {1} has a column span value of {2}, which goes past the end of the grid. This behaviour will be deprecated in the future. Slots should not use column span to stretch themselves past the end of the grid."), ChildSlot->GetRow(), ChildSlot->GetColumn(), ChildSlot->GetColumnSpan());
+				CompileLog.Warning(InfoText);
+			}
+		}
+	}
+
+	UPanelWidget::ValidateCompiledDefaults(CompileLog);
+}
+
 #endif
 
 /////////////////////////////////////////////////////
diff --git a/Engine/Source/Runtime/UMG/Public/Components/GridPanel.h b/Engine/Source/Runtime/UMG/Public/Components/GridPanel.h
--- a/Engine/Source/Runtime/UMG/Public/Components/GridPanel.h
+++ b/Engine/Source/Runtime/UMG/Public/Components/GridPanel.h
@@ -8,6 +8,7 @@
 #include "Components/PanelWidget.h"
 #include "GridPanel.generated.h"
 
+class IWidgetCompilerLog;
 class SGridPanel;
 class UGridSlot;
 
@@ -53,6 +54,7 @@ public:
 
 #if WITH_EDITOR
 	virtual const FText GetPaletteCategory() override;
+	virtual void ValidateCompiledDefaults(IWidgetCompilerLog& CompileLog) const override;
 #endif
 
 protected:
