Fixed copy-paste bug in SGridPanel, GetColumnSpan() instead of GetRowSpan().

#rnx

diff --git a/Engine/Source/Runtime/Slate/Private/Widgets/Layout/SGridPanel.cpp b/Engine/Source/Runtime/Slate/Private/Widgets/Layout/SGridPanel.cpp
--- a/Engine/Source/Runtime/Slate/Private/Widgets/Layout/SGridPanel.cpp
+++ b/Engine/Source/Runtime/Slate/Private/Widgets/Layout/SGridPanel.cpp
@@ -411,7 +411,7 @@ void SGridPanel::NotifySlotChanged(const FSlot* InSlot, bool bSlotLayerChanged /
 		Columns.AddZeroed(NumColumnsRequiredForThisSlot - Columns.Num());
 	}
 
-	const int32 NumRowsRequiredForThisSlot = InSlot->GetRow() + InSlot->GetColumnSpan() + 1;
+	const int32 NumRowsRequiredForThisSlot = InSlot->GetRow() + InSlot->GetRowSpan() + 1;
 	if (NumRowsRequiredForThisSlot > Rows.Num())
 	{
 		Rows.AddZeroed(NumRowsRequiredForThisSlot - Rows.Num());
