Add a fallback to the default slate mouse input device if an FSlateUser has an invalid input device ID. You can get an invalid input device ID here on Xbox if you log in, log out, and log in with a different user attempting to use the same keyboard. This happens only on Xbox because the input devices are not remapped when the user is logged out.

#rnx

diff --git a/Engine/Source/Runtime/Slate/Private/Framework/Application/SlateUser.cpp b/Engine/Source/Runtime/Slate/Private/Framework/Application/SlateUser.cpp
--- a/Engine/Source/Runtime/Slate/Private/Framework/Application/SlateUser.cpp
+++ b/Engine/Source/Runtime/Slate/Private/Framework/Application/SlateUser.cpp
@@ -669,8 +669,13 @@ bool FSlateUser::SynthesizeCursorMoveIfNeeded()
 
 		FSlateApplication& SlateApp = FSlateApplication::Get();
 		
-		FInputDeviceId InputDeviceId = IPlatformInputDeviceMapper::Get().GetPrimaryInputDeviceForUser( GetPlatformUserId() );
-		check(InputDeviceId.IsValid());
+		FInputDeviceId InputDeviceId = IPlatformInputDeviceMapper::Get().GetPrimaryInputDeviceForUser(GetPlatformUserId());
+		
+		if (!InputDeviceId.IsValid())
+		{
+			UE_LOG(LogSlate, Warning, TEXT("SynthesizeCursorMoveIfNeeded had an invalid Input Device! Falling back to the default mouse input device ID."));
+			InputDeviceId = SlateApp.GetInputDeviceIdForMouse();
+		}		
 
 		const bool bHasHardwareCursor = SlateApp.GetPlatformCursor() == Cursor;
 		const TSet<FKey> EmptySet;
