Fix gamepad focus not set correctly due to default faux cursor position generating widgets under mouse when fullscreen.

diff --git a/Engine/Source/Runtime/Slate/Private/Framework/Application/SlateApplication.cpp b/Engine/Source/Runtime/Slate/Private/Framework/Application/SlateApplication.cpp
--- a/Engine/Source/Runtime/Slate/Private/Framework/Application/SlateApplication.cpp
+++ b/Engine/Source/Runtime/Slate/Private/Framework/Application/SlateApplication.cpp
@@ -96,7 +96,13 @@ namespace UE::Slate::Private
 class FFauxSlateCursor : public ICursor
 {
 public:
-	FFauxSlateCursor() {}
+	FFauxSlateCursor() 
+	{
+		// We don't support any concept of invalid or unset position for pointer events.
+		// To avoid collisions with fullscreen or windows or multi-monitor setups with monitors left of primary,
+		// we initialize the faux position to something that shouldn't generate overlap with any widgets.
+		CurrentPosition = FVector2D(std::numeric_limits<int32>::min(), std::numeric_limits<int32>::min());
+	}
 	virtual ~FFauxSlateCursor() {}
 	virtual void SetTypeShape(EMouseCursor::Type InCursorType, void* CursorHandle) override {}
 	
