Fix to Input chords on MAC to properly show Option instead of Alt.

diff --git a/Engine/Source/Runtime/Slate/Private/Framework/Commands/InputChord.cpp b/Engine/Source/Runtime/Slate/Private/Framework/Commands/InputChord.cpp
--- a/Engine/Source/Runtime/Slate/Private/Framework/Commands/InputChord.cpp
+++ b/Engine/Source/Runtime/Slate/Private/Framework/Commands/InputChord.cpp
@@ -37,6 +37,7 @@ FText FInputChord::GetModifierText(TOptional<FText> ModifierAppender) const
 #if PLATFORM_MAC
 	const FText CommandText = LOCTEXT("KeyName_Control", "Ctrl");
 	const FText ControlText = LOCTEXT("KeyName_Command", "Cmd");
+	const FText OptionText = LOCTEXT("KeyName_Option", "Option");
 #else
 	const FText ControlText = LOCTEXT("KeyName_Control", "Ctrl");
 	const FText CommandText = LOCTEXT("KeyName_Command", "Cmd");
@@ -62,7 +63,11 @@ FText FInputChord::GetModifierText(TOptional<FText> ModifierAppender) const
 
 	if (bAlt)
 	{
+#if PLATFORM_MAC
+		Args.Add(FString::Printf(TEXT("Mod%d"), ++ModCount), OptionText);
+#else
 		Args.Add(FString::Printf(TEXT("Mod%d"), ++ModCount), AltText);
+#endif
 	}
 
 	if (bShift)
