Ensuring that any docked tab stacks have the upper left most tab well exposed in a floating window.

diff --git a/Engine/Source/Runtime/Slate/Private/Framework/Docking/SDockingArea.cpp b/Engine/Source/Runtime/Slate/Private/Framework/Docking/SDockingArea.cpp
--- a/Engine/Source/Runtime/Slate/Private/Framework/Docking/SDockingArea.cpp
+++ b/Engine/Source/Runtime/Slate/Private/Framework/Docking/SDockingArea.cpp
@@ -660,6 +660,8 @@ void SDockingArea::UpdateWindowChromeAndSidebar()
 			}
 		}
 
+		AdjustDockedTabsIfNeeded();
+
 		// Call the delegate for when the window is activated, to update the global menu bar without the user having to click on it
 		OnOwningWindowActivated();
 	}
diff --git a/Engine/Source/Runtime/Slate/Private/Framework/Docking/SDockingSplitter.cpp b/Engine/Source/Runtime/Slate/Private/Framework/Docking/SDockingSplitter.cpp
--- a/Engine/Source/Runtime/Slate/Private/Framework/Docking/SDockingSplitter.cpp
+++ b/Engine/Source/Runtime/Slate/Private/Framework/Docking/SDockingSplitter.cpp
@@ -1,6 +1,8 @@
 // Copyright Epic Games, Inc. All Rights Reserved.
 
 #include "Framework/Docking/SDockingSplitter.h"
+
+#include "SDockingArea.h"
 #include "Widgets/Docking/SDockTab.h"
 #include "Framework/Docking/SDockingTabStack.h"
 
@@ -85,6 +87,44 @@ SDockingNode::ECleanupRetVal SDockingSplitter::MostResponsibility( SDockingNode:
 	return FMath::Min(A, B);
 }
 
+void SDockingSplitter::AdjustDockedTabsIfNeeded()
+{
+	TSharedPtr<SDockingArea> DockingArea = nullptr;
+
+	if (ParentNodePtr != nullptr && ParentNodePtr.IsValid())
+	{
+		DockingArea = ParentNodePtr.Pin()->GetDockArea().ToSharedRef();
+	}
+	else if (GetNodeType() == DockArea)
+	{
+		DockingArea = StaticCastSharedRef<SDockingArea>(AsShared());
+	}
+
+	// if is in floating window, set the first docking tab stack to have it's tab well unhidden.
+	if (DockingArea != nullptr && DockingArea->GetParentWindow().IsValid())
+	{
+		for (const TSharedRef<SDockingNode>& ChildNode : Children)
+		{
+			/*
+			 * if the node type is a tab stack AND it is visible, this must be the first tab stack and we should not
+			 * hide the tabwell. Unhide it and break out of the method, no need to go further
+			 */
+			if (ChildNode->GetNodeType() == DockTabStack && ChildNode->GetVisibility() == EVisibility::Visible)
+			{
+				TSharedRef<SDockingTabStack> TabStack = StaticCastSharedRef<SDockingTabStack>(ChildNode);
+				TabStack->SetTabWellHidden(false);
+				break;
+			}
+			// else if  node type is splitter, the first tab stack might be in there... check its children to see
+			else if (ChildNode->GetNodeType() == DockSplitter)
+			{
+				TSharedRef<SDockingSplitter> ChildSplitter = StaticCastSharedRef<SDockingSplitter>(ChildNode);
+				ChildSplitter->AdjustDockedTabsIfNeeded();
+			}
+		}
+	}
+}
+
 SDockingNode::ECleanupRetVal SDockingSplitter::CleanUpNodes()
 {
 	ECleanupRetVal ThisNodePurpose = NoTabsUnderNode;
diff --git a/Engine/Source/Runtime/Slate/Private/Framework/Docking/SDockingSplitter.h b/Engine/Source/Runtime/Slate/Private/Framework/Docking/SDockingSplitter.h
--- a/Engine/Source/Runtime/Slate/Private/Framework/Docking/SDockingSplitter.h
+++ b/Engine/Source/Runtime/Slate/Private/Framework/Docking/SDockingSplitter.h
@@ -124,6 +124,12 @@ protected:
 	
 	/** The DockNode children. All these children are kept in sync with the SSplitter's children via the use of the public interface for adding, removing and replacing children. */
 	TArray< TSharedRef<class SDockingNode> > Children;
+
+	/**
+	 *   As the Dock area is updating after clean up of any added, removed, or newly shown Dock areas or tabs,
+	 *   the tabs/tab well may need adjusting. Use this for post clean up docked tab adjustments.
+	 */
+	void AdjustDockedTabsIfNeeded();
 };
 
 
