Tab Manager: Return true from IsValidTabForSpawning for nomad tabs that have specified OnFindTabToReuse, which is more consistent with other handling of that delegate.

diff --git a/Engine/Source/Runtime/Slate/Private/Framework/Docking/TabManager.cpp b/Engine/Source/Runtime/Slate/Private/Framework/Docking/TabManager.cpp
--- a/Engine/Source/Runtime/Slate/Private/Framework/Docking/TabManager.cpp
+++ b/Engine/Source/Runtime/Slate/Private/Framework/Docking/TabManager.cpp
@@ -1877,7 +1877,7 @@ bool FTabManager::IsValidTabForSpawning( const FTab& SomeTab ) const
 
 	// Nomad tabs being restored from layouts should not be spawned if the nomad tab is already spawned.
 	TSharedRef<FTabSpawnerEntry>* NomadSpawner = NomadTabSpawner->Find( SomeTab.TabId.TabType );
-	return ( !NomadSpawner || !NomadSpawner->Get().IsSoleTabInstanceSpawned() );
+	return ( !NomadSpawner || !NomadSpawner->Get().IsSoleTabInstanceSpawned() || NomadSpawner->Get().OnFindTabToReuse.IsBound() );
 }
 
 bool FTabManager::IsAllowedTab(const FTabId& TabId) const
