SBorder now can be constructed with the .Padding(1,2,3,4) declarative style, just like SBox and widget slots.

diff --git a/Engine/Source/Runtime/Slate/Public/Widgets/Layout/SBorder.h b/Engine/Source/Runtime/Slate/Public/Widgets/Layout/SBorder.h
--- a/Engine/Source/Runtime/Slate/Public/Widgets/Layout/SBorder.h
+++ b/Engine/Source/Runtime/Slate/Public/Widgets/Layout/SBorder.h
@@ -74,6 +74,25 @@ public:
 
 		/** Flips the background image if the localization's flow direction is RightToLeft */
 		SLATE_ARGUMENT(bool, FlipForRightToLeftFlowDirection)
+
+		FArguments& Padding(float Uniform)
+		{
+			_Padding = FMargin(Uniform);
+			return *this;
+		}
+
+		FArguments& Padding(float Horizontal, float Vertical)
+		{
+			_Padding = FMargin(Horizontal, Vertical);
+			return *this;
+		}
+
+		FArguments& Padding(float Left, float Top, float Right, float Bottom)
+		{
+			_Padding = FMargin(Left, Top, Right, Bottom);
+			return *this;
+		}
+
 	SLATE_END_ARGS()
 
 	/**
