Fixup include path in MacMenu

diff --git a/Engine/Source/Runtime/Slate/Private/Framework/MultiBox/Mac/MacMenu.h b/Engine/Source/Runtime/Slate/Private/Framework/MultiBox/Mac/MacMenu.h
--- a/Engine/Source/Runtime/Slate/Private/Framework/MultiBox/Mac/MacMenu.h
+++ b/Engine/Source/Runtime/Slate/Private/Framework/MultiBox/Mac/MacMenu.h
@@ -3,7 +3,7 @@
 #pragma once
 
 #include "Framework/MultiBox/MultiBox.h"
-#include "SMenuEntryBlock.h"
+#include "Framework/MultiBox/SMenuEntryBlock.h"
 #include "Mac/CocoaMenu.h"
 
 @interface FMacMenu : FCocoaMenu <NSMenuDelegate>
