Call the correct "OnController" function from the SVirtualJoystick. Before it was only calling OnControllerAnalog, which will result in normal FKey's never having their "Released" input be triggered. This fix calls OnControllerButtonPressed and OnControllerButtonReleased for non-analog keys, which will give the correct behavior.

diff --git a/Engine/Source/Runtime/Slate/Private/Widgets/Input/SVirtualJoystick.cpp b/Engine/Source/Runtime/Slate/Private/Widgets/Input/SVirtualJoystick.cpp
--- a/Engine/Source/Runtime/Slate/Private/Widgets/Input/SVirtualJoystick.cpp
+++ b/Engine/Source/Runtime/Slate/Private/Widgets/Input/SVirtualJoystick.cpp
@@ -417,9 +417,28 @@ void SVirtualJoystick::Tick(const FGeometry& AllottedGeometry, const double InCu
 			FSlateApplication::Get().SetAllUserFocusToGameViewport();
 			
 			FInputDeviceId PrimaryInputDevice = IPlatformInputDeviceMapper::Get().GetPrimaryInputDeviceForUser(FSlateApplicationBase::SlateAppPrimaryPlatformUser);
-			FSlateApplication::Get().OnControllerAnalog(XAxis, FSlateApplicationBase::SlateAppPrimaryPlatformUser, PrimaryInputDevice, NormalizedOffset.X);
-			FSlateApplication::Get().OnControllerAnalog(YAxis, FSlateApplicationBase::SlateAppPrimaryPlatformUser, PrimaryInputDevice, -NormalizedOffset.Y);
+
+			auto ApplyInput = [PrimaryInputDevice](const FGamepadKeyNames::Type KeyName, float Delta)
+			{
+				FKey Key(KeyName);
+				if (Key.IsAnalog())
+				{
+					FSlateApplication::Get().OnControllerAnalog(KeyName, FSlateApplicationBase::SlateAppPrimaryPlatformUser, PrimaryInputDevice, Delta);
+				}
+				else if (Delta != 0.0f)
+				{
+					FSlateApplication::Get().OnControllerButtonPressed(KeyName, FSlateApplicationBase::SlateAppPrimaryPlatformUser, PrimaryInputDevice, false);
+				}
+				else
+				{
+					FSlateApplication::Get().OnControllerButtonReleased(KeyName, FSlateApplicationBase::SlateAppPrimaryPlatformUser, PrimaryInputDevice, false);
+				}
+			};
+
+			ApplyInput(XAxis, NormalizedOffset.X);
+			ApplyInput(YAxis, -NormalizedOffset.Y);
 		}
+		
 
 		// is this active?
 		if (Control.CapturedPointerIndex != -1)
