Allow user lists to support scrolling when list items are varying heights.

#rnx
#Tests Ran game in a single client in windowed mode with a controller.  Navigated to event panel and resized window to be flush against the last visible user list item before the state button.  Scrolled down and observed scrolling working as intended.  Tested by QA on a Preflight build to confirm there were no abnormal side effects.

- SListView wasn't trusting the user list objects to report their proper size and would determine a max size assuming all objects in the list were the same size.  This simply removes the arbitrary max size test.

diff --git a/Engine/Source/Runtime/Slate/Public/Widgets/Views/SListView.h b/Engine/Source/Runtime/Slate/Public/Widgets/Views/SListView.h
--- a/Engine/Source/Runtime/Slate/Public/Widgets/Views/SListView.h
+++ b/Engine/Source/Runtime/Slate/Public/Widgets/Views/SListView.h
@@ -2105,8 +2105,7 @@ protected:
 							}
 						}
 
-						const double MaxScrollOffset = FMath::Max(0.0, static_cast<double>(ItemsSource->Num()) - NumLiveWidgets);
-						SetScrollOffset((float)FMath::Min(NewScrollOffset, MaxScrollOffset));
+						SetScrollOffset((float)NewScrollOffset);
 					}
 				}
 
