Fix keyboard focus issue on static drop down menus (menu wouldn't detect any keyboard hit after being closed by pressing Esc)

diff --git a/Engine/Source/Runtime/Slate/Private/Framework/Application/MenuStack.cpp b/Engine/Source/Runtime/Slate/Private/Framework/Application/MenuStack.cpp
--- a/Engine/Source/Runtime/Slate/Private/Framework/Application/MenuStack.cpp
+++ b/Engine/Source/Runtime/Slate/Private/Framework/Application/MenuStack.cpp
@@ -172,7 +172,7 @@ namespace MenuStackInternal
 		if (Key == EKeys::Escape)
 		{
 			FSlateApplication::Get().DismissAllMenus();
-			return FReply::Handled();
+			return FReply::Handled().ClearUserFocus();
 		}
 
 		return FReply::Unhandled();
