Little fix for recently add TListTypeTraits::IsPtrValid

diff --git a/Engine/Source/Runtime/Slate/Public/Framework/Views/TableViewTypeTraits.h b/Engine/Source/Runtime/Slate/Public/Framework/Views/TableViewTypeTraits.h
--- a/Engine/Source/Runtime/Slate/Public/Framework/Views/TableViewTypeTraits.h
+++ b/Engine/Source/Runtime/Slate/Public/Framework/Views/TableViewTypeTraits.h
@@ -484,7 +484,7 @@ public:
 
 	static bool IsPtrValid( const TObjectPtr<T>& InPtr )
 	{
-		return InPtr.IsNull();
+		return !InPtr.IsNull();
 	}
 
 	static void ResetPtr(TObjectPtr<T>& InPtr )
