Content Browser: Fix an issue causing browse to asset to not scroll to the correct location sometimes

diff --git a/Engine/Source/Runtime/Slate/Private/Widgets/Views/SListPanel.cpp b/Engine/Source/Runtime/Slate/Private/Widgets/Views/SListPanel.cpp
--- a/Engine/Source/Runtime/Slate/Private/Widgets/Views/SListPanel.cpp
+++ b/Engine/Source/Runtime/Slate/Private/Widgets/Views/SListPanel.cpp
@@ -238,7 +238,8 @@ float SListPanel::GetItemPadding(const FGeometry& AllottedGeometry, const EListI
 		// Only add padding between items if we have more total items that we can fit on a single line.  Otherwise,
 		// the padding around items would continue to change proportionately with the (ample) free space along the line axis
 		const int32 NumItemsPerLine = DesiredDimensions.LineAxis > 0.f ? FMath::FloorToInt(AllottedDimensions.LineAxis / DesiredDimensions.LineAxis) : 0;
-		if (NumItemsPerLine > 0 && Children.Num() > NumItemsPerLine)
+		
+		if (NumItemsPerLine > 0 && Children.Num() >= NumItemsPerLine)
 		{
 			return (AllottedDimensions.LineAxis - FloatingPointPrecisionOffset - (NumItemsPerLine * DesiredDimensions.LineAxis)) / NumItemsPerLine;
 		}
diff --git a/Engine/Source/Runtime/Slate/Private/Widgets/Views/STableViewBase.cpp b/Engine/Source/Runtime/Slate/Private/Widgets/Views/STableViewBase.cpp
--- a/Engine/Source/Runtime/Slate/Private/Widgets/Views/STableViewBase.cpp
+++ b/Engine/Source/Runtime/Slate/Private/Widgets/Views/STableViewBase.cpp
@@ -279,13 +279,25 @@ void STableViewBase::Tick( const FGeometry& AllottedGeometry, const double InCur
 	if (ItemsPanel.IsValid())
 	{
 		FGeometry PanelGeometry = FindChildGeometry( AllottedGeometry, ItemsPanel.ToSharedRef() );
-		if ( bItemsNeedRefresh || PanelGeometryLastTick.GetLocalSize() != PanelGeometry.GetLocalSize())
+
+		bool bPanelGeometryChanged = PanelGeometryLastTick.GetLocalSize() != PanelGeometry.GetLocalSize();
+		
+		if ( bItemsNeedRefresh || bPanelGeometryChanged)
 		{
 			PanelGeometryLastTick = PanelGeometry;
 			
 			const int32 NumItemsPerLine = GetNumItemsPerLine();
-			const EScrollIntoViewResult ScrollIntoViewResult = ScrollIntoView(PanelGeometry);
 
+			EScrollIntoViewResult ScrollIntoViewResult = EScrollIntoViewResult::Deferred;
+
+			/* Don't try to scroll if our geometry is changing (eg: In the middle of refreshing items) because the
+			 * scrolled item could end up outside the view
+			 */
+			if(!bPanelGeometryChanged)
+			{
+				ScrollIntoViewResult = ScrollIntoView(PanelGeometry);
+			}
+			
 			double TargetScrollOffset = GetTargetScrollOffset();
 
 			if (bEnableAnimatedScrolling)
