Fix menu closing unexpectedly, by preventing to clear the 'WidgetToFocusOnActivate', from a SWindow that was triggering indirectly this unexpected close, when we tried to update the focused menus..

#9423

Tests: Found an exact 100% repro (open Build menu, click on Lighting info, move on lightmap density, then lightmight resolution, then lightmap density again), it no longer occurs. Played with a lot menus in the editor, and steps from the jira, no regression found. No more menu closing unexpectedly.

diff --git a/Engine/Source/Runtime/Slate/Private/Framework/Application/MenuStack.cpp b/Engine/Source/Runtime/Slate/Private/Framework/Application/MenuStack.cpp
--- a/Engine/Source/Runtime/Slate/Private/Framework/Application/MenuStack.cpp
+++ b/Engine/Source/Runtime/Slate/Private/Framework/Application/MenuStack.cpp
@@ -183,8 +183,13 @@ namespace MenuStackInternal
 
 TSharedRef<IMenu> FMenuStack::Push(const FWidgetPath& InOwnerPath, const TSharedRef<SWidget>& InContent, const FVector2D& SummonLocation, const FPopupTransitionEffect& TransitionEffect, const bool bFocusImmediately, const FVector2D& SummonLocationSize, TOptional<EPopupMethod> InMethod, const bool bIsCollapsedByParent, const bool bEnablePerPixelTransparency)
 {
-	// We want to ensure that when the window is restored to restore the current keyboard focus
-	InOwnerPath.GetWindow()->SetWidgetToFocusOnActivate(FSlateApplication::Get().GetKeyboardFocusedWidget());
+	// We want to ensure that when the window is restored, we restore the current keyboard focus, 
+	// but only if it is valid, otherwise we could end up clearing a previously valid path.
+	TSharedPtr<SWidget> FocusedWidget = FSlateApplication::Get().GetKeyboardFocusedWidget();
+	if (FocusedWidget.IsValid())
+	{
+		InOwnerPath.GetWindow()->SetWidgetToFocusOnActivate(FocusedWidget);
+	}
 
 	FSlateRect Anchor(SummonLocation, SummonLocation + SummonLocationSize);
 	TSharedPtr<IMenu> ParentMenu;
diff --git a/Engine/Source/Runtime/Slate/Private/Widgets/Input/SMenuAnchor.cpp b/Engine/Source/Runtime/Slate/Private/Widgets/Input/SMenuAnchor.cpp
--- a/Engine/Source/Runtime/Slate/Private/Widgets/Input/SMenuAnchor.cpp
+++ b/Engine/Source/Runtime/Slate/Private/Widgets/Input/SMenuAnchor.cpp
@@ -163,7 +163,7 @@ void SMenuAnchor::Tick( const FGeometry& AllottedGeometry, const double InCurren
 	{
 		// Ideally, do this in OnArrangeChildren(); currently not possible because OnArrangeChildren()
 		// can be called in DesktopSpace or WindowSpace, and we will not know which version of the Window
-		// geometry to use. Tick() is always in DesktopSpace, so cache the solution here and just use
+		// geometry to use. Tick() is always in DesktopSpace, so cache the position here and just use
 		// it in OnArrangeChildren().
 		const FPopupPlacement LocalPlacement(AllottedGeometry, Children[1].GetWidget()->GetDesiredSize(), Placement.Get());
 
