Add an optional SlateUserIndex argument to Key Events that will handle situations where the slate user is not 1:1 with platform user Id

diff --git a/Engine/Source/Runtime/Slate/Private/Framework/Application/SlateApplication.cpp b/Engine/Source/Runtime/Slate/Private/Framework/Application/SlateApplication.cpp
--- a/Engine/Source/Runtime/Slate/Private/Framework/Application/SlateApplication.cpp
+++ b/Engine/Source/Runtime/Slate/Private/Framework/Application/SlateApplication.cpp
@@ -6019,7 +6019,8 @@ bool FSlateApplication::OnControllerAnalog(FGamepadKeyNames::Type KeyName, FPlat
 	TOptional<int32> UserIndex = GetUserIndexForInputDevice(InputDeviceId);
 	if (UserIndex.IsSet() && ensureMsgf(Key.IsValid(), TEXT("OnControllerAnalog(KeyName=%s,InputDeviceId=%d,AnalogValue=%f) key is invalid"), *KeyName.ToString(), InputDeviceId.GetId(), AnalogValue))
 	{
-		FAnalogInputEvent AnalogInputEvent(Key, PlatformApplication->GetModifierKeys(), InputDeviceId, false, 0, 0, AnalogValue);
+		FAnalogInputEvent AnalogInputEvent(Key, PlatformApplication->GetModifierKeys(), InputDeviceId, false, 0, 0, AnalogValue, UserIndex);
+		
 		return ProcessAnalogInputEvent(AnalogInputEvent);
 	}
 	return false;
@@ -6031,7 +6032,8 @@ bool FSlateApplication::OnControllerButtonPressed(FGamepadKeyNames::Type KeyName
 	TOptional<int32> UserIndex = GetUserIndexForInputDevice(InputDeviceId);
 	if (UserIndex.IsSet() && ensureMsgf(Key.IsValid(), TEXT("OnControllerButtonPressed(KeyName=%s,InputDeviceId=%d,IsRepeat=%b) key is invalid"), *KeyName.ToString(), InputDeviceId.GetId(), IsRepeat))
 	{
-		FKeyEvent KeyEvent(Key, PlatformApplication->GetModifierKeys(), InputDeviceId, IsRepeat, 0, 0);
+		FKeyEvent KeyEvent(Key, PlatformApplication->GetModifierKeys(), InputDeviceId, IsRepeat, 0, 0, UserIndex);
+		
 		return ProcessKeyDownEvent(KeyEvent);
 	}
 	return false;
@@ -6043,7 +6045,8 @@ bool FSlateApplication::OnControllerButtonReleased(FGamepadKeyNames::Type KeyNam
 	TOptional<int32> UserIndex = GetUserIndexForInputDevice(InputDeviceId);
 	if (UserIndex.IsSet() && ensureMsgf(Key.IsValid(), TEXT("OnControllerButtonReleased(KeyName=%s,InputDeviceId=%d,IsRepeat=%b) key is invalid"), *KeyName.ToString(), InputDeviceId.GetId(), IsRepeat))
 	{
-		FKeyEvent KeyEvent(Key, PlatformApplication->GetModifierKeys(), InputDeviceId, IsRepeat, 0, 0);
+		FKeyEvent KeyEvent(Key, PlatformApplication->GetModifierKeys(), InputDeviceId, IsRepeat, 0, 0, UserIndex);
+		
 		return ProcessKeyUpEvent(KeyEvent);
 	}
 	return false;
diff --git a/Engine/Source/Runtime/Slate/Public/Framework/Application/SlateApplication.h b/Engine/Source/Runtime/Slate/Public/Framework/Application/SlateApplication.h
--- a/Engine/Source/Runtime/Slate/Public/Framework/Application/SlateApplication.h
+++ b/Engine/Source/Runtime/Slate/Public/Framework/Application/SlateApplication.h
@@ -1037,7 +1037,7 @@ protected:
 	 */
 	TSharedRef<FSlateUser> GetOrCreateUser(FInputDeviceId DeviceId);
 	
-	FORCEINLINE TSharedRef<FSlateUser> GetOrCreateUser(const FInputEvent& InputEvent) { return GetOrCreateUser(InputEvent.GetInputDeviceId()); }
+	FORCEINLINE TSharedRef<FSlateUser> GetOrCreateUser(const FInputEvent& InputEvent) { return GetOrCreateUser(InputEvent.GetUserIndex()); }
 
 	friend class FEventRouter;
 
diff --git a/Engine/Source/Runtime/SlateCore/Public/Input/Events.h b/Engine/Source/Runtime/SlateCore/Public/Input/Events.h
--- a/Engine/Source/Runtime/SlateCore/Public/Input/Events.h
+++ b/Engine/Source/Runtime/SlateCore/Public/Input/Events.h
@@ -443,13 +443,19 @@ public:
 			const FInputDeviceId InDeviceId,
 			const bool bInIsRepeat,
 			const uint32 InCharacterCode,
-			const uint32 InKeyCode
+			const uint32 InKeyCode,
+			const TOptional<int32> InOptionalSlateUserIndex = TOptional<int32>()
 	)
 		: FInputEvent(InModifierKeys, InDeviceId, bInIsRepeat)
 		, Key(InKey)
 		, CharacterCode(InCharacterCode)
 		, KeyCode(InKeyCode)
-	{ }
+	{
+		if (InOptionalSlateUserIndex.IsSet())
+		{
+			UserIndex = InOptionalSlateUserIndex.GetValue();
+		}
+	}
 
 	/**
 	 * Returns the name of the key for this event
@@ -551,9 +557,10 @@ public:
 		const bool bInIsRepeat,
 		const uint32 InCharacterCode,
 		const uint32 InKeyCode,
-		const float InAnalogValue
+		const float InAnalogValue,
+		const TOptional<int32> InOptionalSlateUserIndex = TOptional<int32>()
 	)
-		: FKeyEvent(InKey, InModifierKeys, InDeviceId, bInIsRepeat, InCharacterCode, InKeyCode)
+		: FKeyEvent(InKey, InModifierKeys, InDeviceId, bInIsRepeat, InCharacterCode, InKeyCode, InOptionalSlateUserIndex)
 		, AnalogValue(InAnalogValue)
 	{ }
 
