Added OnVerifyTextChanged to SMultiLineEditableTextBox (based on CL# 11063519 for SEditableTextBox)

diff --git a/Engine/Source/Runtime/Slate/Private/Widgets/Input/SMultiLineEditableTextBox.cpp b/Engine/Source/Runtime/Slate/Private/Widgets/Input/SMultiLineEditableTextBox.cpp
--- a/Engine/Source/Runtime/Slate/Private/Widgets/Input/SMultiLineEditableTextBox.cpp
+++ b/Engine/Source/Runtime/Slate/Private/Widgets/Input/SMultiLineEditableTextBox.cpp
@@ -45,6 +45,10 @@ void SMultiLineEditableTextBox::Construct( const FArguments& InArgs )
 	FocusedForegroundColorOverride = InArgs._FocusedForegroundColor;
 	bSelectWordOnMouseDoubleClick = InArgs._SelectWordOnMouseDoubleClick;
 
+	OnTextChanged = InArgs._OnTextChanged;
+	OnVerifyTextChanged = InArgs._OnVerifyTextChanged;
+	OnTextCommitted = InArgs._OnTextCommitted;
+
 	bHasExternalHScrollBar = InArgs._HScrollBar.IsValid();
 	HScrollBar = InArgs._HScrollBar;
 	if (!HScrollBar.IsValid())
@@ -99,8 +103,8 @@ void SMultiLineEditableTextBox::Construct( const FArguments& InArgs )
 					.AllowMultiLine( InArgs._AllowMultiLine )
 					.OnContextMenuOpening( InArgs._OnContextMenuOpening )
 					.OnIsTypedCharValid( InArgs._OnIsTypedCharValid )
-					.OnTextChanged( InArgs._OnTextChanged )
-					.OnTextCommitted( InArgs._OnTextCommitted )
+					.OnTextChanged(this, &SMultiLineEditableTextBox::OnEditableTextChanged)
+					.OnTextCommitted(this, &SMultiLineEditableTextBox::OnEditableTextCommitted)
 					.OnCursorMoved( InArgs._OnCursorMoved )
 					.ContextMenuExtender( InArgs._ContextMenuExtender )
 					.CreateSlateTextLayout( InArgs._CreateSlateTextLayout )
@@ -599,4 +603,50 @@ void SMultiLineEditableTextBox::ForceScroll(int32 UserIndex, float ScrollAxisMag
 	EditableText->ForceScroll(UserIndex, ScrollAxisMagnitude);
 }
 
+void SMultiLineEditableTextBox::OnEditableTextChanged(const FText& InText)
+{
+	OnTextChanged.ExecuteIfBound(InText);
+
+	if (OnVerifyTextChanged.IsBound())
+	{
+		FText OutErrorMessage;
+		if (!OnVerifyTextChanged.Execute(InText, OutErrorMessage))
+		{
+			// Display as an error.
+			SetError(OutErrorMessage);
+		}
+		else
+		{
+			SetError(FText::GetEmpty());
+		}
+	}
+}
+
+void SMultiLineEditableTextBox::OnEditableTextCommitted(const FText& InText, ETextCommit::Type InCommitType)
+{
+	if (OnVerifyTextChanged.IsBound())
+	{
+		FText OutErrorMessage;
+		if (!OnVerifyTextChanged.Execute(InText, OutErrorMessage))
+		{
+			// Display as an error.
+			if (InCommitType == ETextCommit::OnEnter)
+			{
+				SetError(OutErrorMessage);
+			}
+			return;
+		}
+		else
+		{
+			if (InCommitType == ETextCommit::OnEnter)
+			{
+				SetError(FText::GetEmpty());
+			}
+
+		}
+	}
+
+	OnTextCommitted.ExecuteIfBound(InText, InCommitType);
+}
+
 #endif //WITH_FANCY_TEXT
diff --git a/Engine/Source/Runtime/Slate/Public/Widgets/Input/SMultiLineEditableTextBox.h b/Engine/Source/Runtime/Slate/Public/Widgets/Input/SMultiLineEditableTextBox.h
--- a/Engine/Source/Runtime/Slate/Public/Widgets/Input/SMultiLineEditableTextBox.h
+++ b/Engine/Source/Runtime/Slate/Public/Widgets/Input/SMultiLineEditableTextBox.h
@@ -167,6 +167,9 @@ public:
 		/** Called whenever the text is committed.  This happens when the user presses enter or the text box loses focus. */
 		SLATE_EVENT( FOnTextCommitted, OnTextCommitted )
 
+		/** Called whenever the text is changed programmatically or interactively by the user */
+		SLATE_EVENT( FOnVerifyTextChanged, OnVerifyTextChanged )
+
 		/** Called whenever the horizontal scrollbar is moved by the user */
 		SLATE_EVENT( FOnUserScrolled, OnHScrollBarUserScrolled )
 
@@ -472,6 +475,13 @@ public:
 	 */
 	void ForceScroll(int32 UserIndex, float ScrollAxisMagnitude);
 
+protected:
+	/** Callback for the editable text's OnTextChanged event */
+	void OnEditableTextChanged(const FText& InText);
+
+	/** Callback when the editable text is committed. */
+	void OnEditableTextCommitted(const FText& InText, ETextCommit::Type InCommitType);
+
 protected:
 
 	/** Editable text widget */
@@ -510,7 +520,6 @@ protected:
 	/** Allows for inserting additional widgets that extend the functionality of the text box */
 	TSharedPtr<SHorizontalBox> Box;
 
-
 	/** Whether we have an externally supplied horizontal scrollbar or one created internally */
 	bool bHasExternalHScrollBar;
 
@@ -532,6 +541,15 @@ protected:
 	/** SomeWidget reporting */
 	TSharedPtr<class IErrorReportingWidget> ErrorReporting;
 
+	/** Called when the text is changed interactively */
+	FOnTextChanged OnTextChanged;
+
+	/** Called when the user commits their change to the editable text control */
+	FOnTextCommitted OnTextCommitted;
+
+	/** Callback to verify text when changed. Will return an error message to denote problems. */
+	FOnVerifyTextChanged OnVerifyTextChanged;
+
 	const FEditableTextBoxStyle* Style;
 
 private:
