Fix FInputChord bad hashing.

#rnx

diff --git a/Engine/Source/Runtime/Slate/Public/Framework/Commands/InputChord.h b/Engine/Source/Runtime/Slate/Public/Framework/Commands/InputChord.h
--- a/Engine/Source/Runtime/Slate/Public/Framework/Commands/InputChord.h
+++ b/Engine/Source/Runtime/Slate/Public/Framework/Commands/InputChord.h
@@ -227,7 +227,7 @@ public:
 	 */
 	friend uint32 GetTypeHash( const FInputChord& Chord )
 	{
-		return GetTypeHash(Chord.Key) ^ (Chord.bShift | Chord.bCtrl >> 1 | Chord.bAlt >> 2 | Chord.bCmd >> 3);
+		return GetTypeHash(Chord.Key) ^ (Chord.bShift | Chord.bCtrl << 1 | Chord.bAlt << 2 | Chord.bCmd << 3);
 	}
 };
 
