Fix for using the incorrect slate user when synthesizing the mouse cursor

#rnx

diff --git a/Engine/Source/Runtime/Slate/Private/Framework/Application/SlateUser.cpp b/Engine/Source/Runtime/Slate/Private/Framework/Application/SlateUser.cpp
--- a/Engine/Source/Runtime/Slate/Private/Framework/Application/SlateUser.cpp
+++ b/Engine/Source/Runtime/Slate/Private/Framework/Application/SlateUser.cpp
@@ -671,26 +671,26 @@ bool FSlateUser::SynthesizeCursorMoveIfNeeded()
 		
 		FInputDeviceId InputDeviceId = IPlatformInputDeviceMapper::Get().GetPrimaryInputDeviceForUser(GetPlatformUserId());
 		
-		if (!InputDeviceId.IsValid())
-		{
-			UE_LOG(LogSlate, Warning, TEXT("SynthesizeCursorMoveIfNeeded had an invalid Input Device! Falling back to the default mouse input device ID."));
-			InputDeviceId = SlateApp.GetInputDeviceIdForMouse();
-		}		
+		// The input device might be invalid if a split screen player has logged off but still has their controller plugged in
+		if (InputDeviceId.IsValid())
+		{			
+			const bool bHasHardwareCursor = SlateApp.GetPlatformCursor() == Cursor;
+			const TSet<FKey> EmptySet;
+			FPointerEvent SyntheticCursorMoveEvent(
+				InputDeviceId,
+				FSlateApplication::CursorPointerIndex,
+				GetCursorPosition(),
+				GetPreviousCursorPosition(),
+				bHasHardwareCursor ? SlateApp.GetPressedMouseButtons() : EmptySet,
+				EKeys::Invalid,
+				0,
+				bHasHardwareCursor ? SlateApp.GetPlatformApplication()->GetModifierKeys() : FModifierKeysState(),
+				UserIndex);
+
+			SlateApp.ProcessMouseMoveEvent(SyntheticCursorMoveEvent, true);
+			return true;
 
-		const bool bHasHardwareCursor = SlateApp.GetPlatformCursor() == Cursor;
-		const TSet<FKey> EmptySet;
-		FPointerEvent SyntheticCursorMoveEvent(
-			InputDeviceId,
-			FSlateApplication::CursorPointerIndex,
-			GetCursorPosition(),
-			GetPreviousCursorPosition(),
-			bHasHardwareCursor ? SlateApp.GetPressedMouseButtons() : EmptySet,
-			EKeys::Invalid,
-			0,
-			bHasHardwareCursor ? SlateApp.GetPlatformApplication()->GetModifierKeys() : FModifierKeysState());
-
-		SlateApp.ProcessMouseMoveEvent(SyntheticCursorMoveEvent, true);
-		return true;
+		}		
 	}
 	return false;
 }
diff --git a/Engine/Source/Runtime/SlateCore/Public/Input/Events.h b/Engine/Source/Runtime/SlateCore/Public/Input/Events.h
--- a/Engine/Source/Runtime/SlateCore/Public/Input/Events.h
+++ b/Engine/Source/Runtime/SlateCore/Public/Input/Events.h
@@ -768,7 +768,8 @@ public:
 		const TSet<FKey>& InPressedButtons,
 		FKey InEffectingButton,
 		float InWheelDelta,
-		const FModifierKeysState& InModifierKeys
+		const FModifierKeysState& InModifierKeys,
+		const TOptional<int32> InOptionalSlateUserIndex = TOptional<int32>()
 	)
 		: FInputEvent(InModifierKeys, InDeviceId, false)
 		, ScreenSpacePosition(InScreenSpacePosition)
@@ -785,7 +786,12 @@ public:
 		, bIsDirectionInvertedFromDevice(false)
 		, bIsTouchForceChanged(false)
 		, bIsTouchFirstMove(false)
-	{ }
+	{
+		if (InOptionalSlateUserIndex.IsSet())
+		{
+			UserIndex = InOptionalSlateUserIndex.GetValue();
+		}
+	}
 
 	FPointerEvent(
 		uint32 InUserIndex,
