Fixed cached string value of the SpinBox not getting reconsidered when localized options (like Fractional Digits) are changed.

diff --git a/Engine/Source/Runtime/Slate/Public/Widgets/Input/SSpinBox.h b/Engine/Source/Runtime/Slate/Public/Widgets/Input/SSpinBox.h
--- a/Engine/Source/Runtime/Slate/Public/Widgets/Input/SSpinBox.h
+++ b/Engine/Source/Runtime/Slate/Public/Widgets/Input/SSpinBox.h
@@ -214,6 +214,7 @@ public:
 
 		CachedExternalValue = ValueAttribute.Get();
 		CachedValueString = Interface->ToString(CachedExternalValue);
+		bCachedValueStringDirty = false;
 
 		InternalValue = (double)CachedExternalValue;
 
@@ -837,6 +838,7 @@ public:
 	void SetMinFractionalDigits(const TAttribute<TOptional<int32>>& InMinFractionalDigits)
 	{
 		Interface->SetMinFractionalDigits((InMinFractionalDigits.Get().IsSet()) ? InMinFractionalDigits.Get() : MinFractionalDigits);
+		bCachedValueStringDirty = true;
 	}
 
 	/** See the MaxFractionalDigits attribute */
@@ -844,6 +846,7 @@ public:
 	void SetMaxFractionalDigits(const TAttribute<TOptional<int32>>& InMaxFractionalDigits)
 	{
 		Interface->SetMaxFractionalDigits((InMaxFractionalDigits.Get().IsSet()) ? InMaxFractionalDigits.Get() : MaxFractionalDigits);
+		bCachedValueStringDirty = true;
 	}
 
 	/** See the AlwaysUsesDeltaSnap attribute */
@@ -885,11 +888,12 @@ protected:
 	FString GetValueAsString() const
 	{
 		NumericType CurrentValue = ValueAttribute.Get();
-		if (CurrentValue == CachedExternalValue)
+		if (!bCachedValueStringDirty && CurrentValue == CachedExternalValue)
 		{
 			return CachedValueString;
 		}
 
+		bCachedValueStringDirty = false;
 		return Interface->ToString(CurrentValue);
 	}
 
@@ -1040,10 +1044,11 @@ protected:
 
 		// Update the cache of the external value to what the user believes the value is now.
 		const NumericType CurrentValue = ValueAttribute.Get();
-		if (CachedExternalValue != CurrentValue)
+		if (CachedExternalValue != CurrentValue || bCachedValueStringDirty)
 		{
 			CachedExternalValue = ValueAttribute.Get();
 			CachedValueString = Interface->ToString(CachedExternalValue);
+			bCachedValueStringDirty = false;
 		}
 
 		// This ensures that dragging is cleared if focus has been removed from this widget in one of the delegate calls, such as when spawning a modal dialog.
@@ -1169,9 +1174,6 @@ private:
 		}
 	}
 
-	/** Whether the user is dragging the slider */
-	bool bDragging;
-
 	/** Tracks which cursor is currently dragging the slider (e.g., the mouse cursor or a specific finger) */
 	int32 PointerDraggingSliderIndex;
 
@@ -1197,6 +1199,12 @@ private:
 	/** Used to prevent per-frame re-conversion of the cached numeric value to a string. */
 	FString CachedValueString;
 
+	/** Whetever the interfaced setting changed and the CachedValueString needs to be recomputed. */
+	mutable bool bCachedValueStringDirty;
+
+	/** Whether the user is dragging the slider */
+	bool bDragging;
+
 	/** Re-entrant guard for the text changed handler */
 	bool bIsTextChanging;
 
