Slate:

-Expose text highlighting in SHyperlink

diff --git a/Engine/Source/Runtime/Slate/Public/Widgets/Input/SHyperlink.h b/Engine/Source/Runtime/Slate/Public/Widgets/Input/SHyperlink.h
--- a/Engine/Source/Runtime/Slate/Public/Widgets/Input/SHyperlink.h
+++ b/Engine/Source/Runtime/Slate/Public/Widgets/Input/SHyperlink.h
@@ -38,6 +38,9 @@ public:
 		, _OnNavigate()
 		, _TextShapingMethod()
 		, _TextFlowDirection()
+		, _HighlightColor()
+		, _HighlightShape()
+		, _HighlightText()
 		{}
 
 		SLATE_ATTRIBUTE( FText, Text )
@@ -48,6 +51,9 @@ public:
 		SLATE_EVENT( FSimpleDelegate, OnNavigate )
 		SLATE_ARGUMENT( TOptional<ETextShapingMethod>, TextShapingMethod )
 		SLATE_ARGUMENT( TOptional<ETextFlowDirection>, TextFlowDirection )
+		SLATE_ATTRIBUTE( FLinearColor, HighlightColor )
+		SLATE_ATTRIBUTE( const FSlateBrush*, HighlightShape )
+		SLATE_ATTRIBUTE( FText, HighlightText )
 	SLATE_END_ARGS()
 
 	/**
@@ -76,6 +82,9 @@ public:
 				SNew(STextBlock)
 				.TextStyle(TextStyle)
 				.Text(InArgs._Text)
+				.HighlightColor(InArgs._HighlightColor)
+				.HighlightShape(InArgs._HighlightShape)
+				.HighlightText(InArgs._HighlightText)
 			]
 		);
 	}
