Adding a cvar to adjust the software cursor's scale.

 - Can be used with 4k device profiles to make the cursor more visible.
 - Requires a cursor widget that is scale-aware.
 - ```Slate.SoftwareCursorScale```

diff --git a/Engine/Source/Runtime/Slate/Private/Framework/Application/SlateUser.cpp b/Engine/Source/Runtime/Slate/Private/Framework/Application/SlateUser.cpp
--- a/Engine/Source/Runtime/Slate/Private/Framework/Application/SlateUser.cpp
+++ b/Engine/Source/Runtime/Slate/Private/Framework/Application/SlateUser.cpp
@@ -38,6 +38,12 @@ FAutoConsoleVariableRef CVarEnableCursorQueries(
 	bEnableCursorQueries,
 	TEXT(""));
 
+static float SoftwareCursorScale = 1.0f;
+FAutoConsoleVariableRef CVarSoftwareCursorScale(
+	TEXT("Slate.SoftwareCursorScale"),
+	SoftwareCursorScale,
+	TEXT("Scale factor applied to the software cursor. Requires the cursor widget to be scale-aware."));
+
 static float TooltipSummonDelay = 0.15f;
 FAutoConsoleVariableRef CVarTooltipSummonDelay(
 	TEXT("Slate.TooltipSummonDelay"),
@@ -623,8 +629,8 @@ void FSlateUser::DrawCursor(const TSharedRef<SWindow>& WindowToDraw, FSlateWindo
 
 			FVector2D CursorInScreen = GetCursorPosition();
 			FVector2D CursorPosInWindowSpace = WindowToDraw->GetWindowGeometryInScreen().AbsoluteToLocal(CursorInScreen) * WindowRootScale;
-			CursorPosInWindowSpace += (CursorWidget->GetDesiredSize() * -0.5);
-			const FGeometry CursorGeometry = FGeometry::MakeRoot(CursorWidget->GetDesiredSize(), FSlateLayoutTransform(CursorPosInWindowSpace));
+			CursorPosInWindowSpace += (CursorWidget->GetDesiredSize() * SoftwareCursorScale * -0.5);
+			const FGeometry CursorGeometry = FGeometry::MakeRoot(CursorWidget->GetDesiredSize() * SoftwareCursorScale, FSlateLayoutTransform(CursorPosInWindowSpace));
 
 			CursorWidget->Paint(
 				FPaintArgs(&WindowToDraw.Get(), WindowToDraw->GetHittestGrid(), WindowToDraw->GetPositionInScreen(), SlateApp.GetCurrentTime(), SlateApp.GetDeltaTime()),
