Remove ensure in slate when OnTouchEnded is called without a OnTouchStarted.

diff --git a/Engine/Source/Runtime/Slate/Private/Framework/Application/SlateApplication.cpp b/Engine/Source/Runtime/Slate/Private/Framework/Application/SlateApplication.cpp
--- a/Engine/Source/Runtime/Slate/Private/Framework/Application/SlateApplication.cpp
+++ b/Engine/Source/Runtime/Slate/Private/Framework/Application/SlateApplication.cpp
@@ -6195,10 +6195,6 @@ bool FSlateApplication::OnTouchEnded( const FVector2D& Location, int32 TouchInde
 			true);
 		ProcessTouchEndedEvent(PointerEvent);
 
-#if WITH_SLATE_DEBUGGING
-		ensure(!User->IsTouchPointerActive(TouchIndex));
-#endif
-
 		return true;
 	}
 
