Read-only textboxes won't try to commit text when losing focus anymore.

Reported on UDN.

diff --git a/Engine/Source/Runtime/Slate/Private/Widgets/Text/SlateEditableTextLayout.cpp b/Engine/Source/Runtime/Slate/Private/Widgets/Text/SlateEditableTextLayout.cpp
--- a/Engine/Source/Runtime/Slate/Private/Widgets/Text/SlateEditableTextLayout.cpp
+++ b/Engine/Source/Runtime/Slate/Private/Widgets/Text/SlateEditableTextLayout.cpp
@@ -898,31 +898,34 @@ bool FSlateEditableTextLayout::HandleFocusLost(const FFocusEvent& InFocusEvent)
 		ClearSelection();
 	}
 
-	// When focus is lost let anyone who is interested that text was committed
-	// See if user explicitly tabbed away or moved focus
-	ETextCommit::Type TextAction;
-	switch (InFocusEvent.GetCause())
-	{
-	case EFocusCause::Navigation:
-	case EFocusCause::Mouse:
-		TextAction = ETextCommit::OnUserMovedFocus;
-		break;
-
-	case EFocusCause::Cleared:
-		TextAction = ETextCommit::OnCleared;
-		break;
-
-	default:
-		TextAction = ETextCommit::Default;
-		break;
-	}
+	if (!OwnerWidget->IsTextReadOnly())
+	{
+		// When focus is lost let anyone who is interested that text was committed
+		// See if user explicitly tabbed away or moved focus
+		ETextCommit::Type TextAction;
+		switch (InFocusEvent.GetCause())
+		{
+		case EFocusCause::Navigation:
+		case EFocusCause::Mouse:
+			TextAction = ETextCommit::OnUserMovedFocus;
+			break;
+
+		case EFocusCause::Cleared:
+			TextAction = ETextCommit::OnCleared;
+			break;
+
+		default:
+			TextAction = ETextCommit::Default;
+			break;
+		}
 
-	// Always clear the local undo chain on commit
-	ClearUndoStates();
+		// Always clear the local undo chain on commit
+		ClearUndoStates();
 
-	const FText EditedText = GetEditableText();
+		const FText EditedText = GetEditableText();
 
-	OwnerWidget->OnTextCommitted(EditedText, TextAction);
+		OwnerWidget->OnTextCommitted(EditedText, TextAction);
+	}
 
 	// Reload underlying value now it is committed  (commit may alter the value) 
 	// so it can be re-displayed in the edit box
