Add the ability to provide a tooltip for menu items that use widgets

diff --git a/Engine/Source/Developer/ToolMenus/Private/ToolMenuEntry.cpp b/Engine/Source/Developer/ToolMenus/Private/ToolMenuEntry.cpp
--- a/Engine/Source/Developer/ToolMenus/Private/ToolMenuEntry.cpp
+++ b/Engine/Source/Developer/ToolMenus/Private/ToolMenuEntry.cpp
@@ -227,10 +227,11 @@ FToolMenuEntry FToolMenuEntry::InitSeparator(const FName InName)
 }
 
 
-FToolMenuEntry FToolMenuEntry::InitWidget(const FName InName, const TSharedRef<SWidget>& InWidget, const FText& Label, bool bNoIndent, bool bSearchable, bool bNoPadding)
+FToolMenuEntry FToolMenuEntry::InitWidget(const FName InName, const TSharedRef<SWidget>& InWidget, const FText& Label, bool bNoIndent, bool bSearchable, bool bNoPadding, const FText& InToolTipText)
 {
 	FToolMenuEntry Entry(UToolMenus::Get()->CurrentOwner(), InName, EMultiBlockType::Widget);
 	Entry.Label = Label;
+	Entry.ToolTip = InToolTipText;
 	Entry.MakeCustomWidget.BindLambda([=](const FToolMenuContext&, const FToolMenuCustomWidgetContext&) { return InWidget; });
 	Entry.WidgetData.bNoIndent = bNoIndent;
 	Entry.WidgetData.bSearchable = bSearchable;
diff --git a/Engine/Source/Developer/ToolMenus/Private/ToolMenus.cpp b/Engine/Source/Developer/ToolMenus/Private/ToolMenus.cpp
--- a/Engine/Source/Developer/ToolMenus/Private/ToolMenus.cpp
+++ b/Engine/Source/Developer/ToolMenus/Private/ToolMenus.cpp
@@ -307,7 +307,7 @@ public:
 			}
 			else
 			{
-				MenuBuilder.AddWidget(Widget.ToSharedRef(), LabelToDisplay.Get(), Block.WidgetData.bNoIndent, Block.WidgetData.bSearchable);
+				MenuBuilder.AddWidget(Widget.ToSharedRef(), LabelToDisplay.Get(), Block.WidgetData.bNoIndent, Block.WidgetData.bSearchable, Block.ToolTip.Get());
 			}
 		}
 		else
diff --git a/Engine/Source/Developer/ToolMenus/Public/ToolMenuEntry.h b/Engine/Source/Developer/ToolMenus/Public/ToolMenuEntry.h
--- a/Engine/Source/Developer/ToolMenus/Public/ToolMenuEntry.h
+++ b/Engine/Source/Developer/ToolMenus/Public/ToolMenuEntry.h
@@ -129,7 +129,7 @@ struct TOOLMENUS_API FToolMenuEntry
 
 	static FToolMenuEntry InitSeparator(const FName InName);
 
-	static FToolMenuEntry InitWidget(const FName InName, const TSharedRef<SWidget>& InWidget, const FText& Label, bool bNoIndent = false, bool bSearchable = true, bool bNoPadding = false);
+	static FToolMenuEntry InitWidget(const FName InName, const TSharedRef<SWidget>& InWidget, const FText& Label, bool bNoIndent = false, bool bSearchable = true, bool bNoPadding = false, const FText& InToolTipText = FText());
 
 	bool IsSubMenu() const { return SubMenuData.bIsSubMenu; }
 
diff --git a/Engine/Source/Runtime/Slate/Private/Framework/MultiBox/MultiBoxBuilder.cpp b/Engine/Source/Runtime/Slate/Private/Framework/MultiBox/MultiBoxBuilder.cpp
--- a/Engine/Source/Runtime/Slate/Private/Framework/MultiBox/MultiBoxBuilder.cpp
+++ b/Engine/Source/Runtime/Slate/Private/Framework/MultiBox/MultiBoxBuilder.cpp
@@ -332,11 +332,11 @@ void FMenuBuilder::AddWrapperSubMenu( const FText& InMenuLabel, const FText& InT
 	MultiBox->AddMultiBlock( NewMenuEntryBlock );
 }
 
-void FMenuBuilder::AddWidget( TSharedRef<SWidget> InWidget, const FText& Label, bool bNoIndent, bool bInSearchable )
+void FMenuBuilder::AddWidget( TSharedRef<SWidget> InWidget, const FText& Label, bool bNoIndent, bool bInSearchable, const TAttribute<FText>& InToolTipText )
 {
 	ApplySectionBeginning();
 
-	TSharedRef< FWidgetBlock > NewWidgetBlock(new FWidgetBlock( InWidget, Label, bNoIndent ));
+	TSharedRef< FWidgetBlock > NewWidgetBlock(new FWidgetBlock( InWidget, Label, bNoIndent, EHorizontalAlignment::HAlign_Fill, InToolTipText ));
 	NewWidgetBlock->SetSearchable( bInSearchable );
 
 	MultiBox->AddMultiBlock( NewWidgetBlock );
diff --git a/Engine/Source/Runtime/Slate/Private/Framework/MultiBox/SMenuEntryBlock.cpp b/Engine/Source/Runtime/Slate/Private/Framework/MultiBox/SMenuEntryBlock.cpp
--- a/Engine/Source/Runtime/Slate/Private/Framework/MultiBox/SMenuEntryBlock.cpp
+++ b/Engine/Source/Runtime/Slate/Private/Framework/MultiBox/SMenuEntryBlock.cpp
@@ -890,7 +890,7 @@ TSharedRef< SWidget> SMenuEntryBlock::BuildSubMenuWidget( const FMenuEntryBuildP
 			// Create a button
 			SNew( SSubMenuButton )
 			// Pass along the block's tool-tip string
-			.ToolTipText( EntryToolTip )
+			.ToolTip( FMultiBoxSettings::ToolTipConstructor.Execute(EntryToolTip, nullptr, UICommand ) )
 			// Style to use
 			.ButtonStyle( &StyleSet->GetWidgetStyle<FButtonStyle>( ISlateStyle::Join( StyleName, ".Button" ) ) )
 			.ContentPadding(StyleSet->GetMargin(StyleName, ".Block.Padding"))
diff --git a/Engine/Source/Runtime/Slate/Private/Framework/MultiBox/SWidgetBlock.cpp b/Engine/Source/Runtime/Slate/Private/Framework/MultiBox/SWidgetBlock.cpp
--- a/Engine/Source/Runtime/Slate/Private/Framework/MultiBox/SWidgetBlock.cpp
+++ b/Engine/Source/Runtime/Slate/Private/Framework/MultiBox/SWidgetBlock.cpp
@@ -4,6 +4,7 @@
 #include "Widgets/SBoxPanel.h"
 #include "Framework/MultiBox/MultiBoxBuilder.h"
 #include "Styling/ToolBarStyle.h"
+#include "Widgets/SToolTip.h"
 
 
 /**
@@ -11,10 +12,11 @@
  *
  * @param	InHeadingText	Heading text
  */
-FWidgetBlock::FWidgetBlock( TSharedRef<SWidget> InContent, const FText& InLabel, bool bInNoIndent, EHorizontalAlignment InHorizontalAlignment)
+FWidgetBlock::FWidgetBlock( TSharedRef<SWidget> InContent, const FText& InLabel, bool bInNoIndent, EHorizontalAlignment InHorizontalAlignment, const TAttribute<FText>& InToolTipText)
 	: FMultiBlock( nullptr, nullptr, NAME_None, EMultiBlockType::Widget )
 	, ContentWidget( InContent )
 	, Label( InLabel )
+	, ToolTipText(InToolTipText)
 	, bNoIndent( bInNoIndent )
 	, HorizontalAlignment(InHorizontalAlignment)
 {
@@ -150,6 +152,7 @@ void SWidgetBlock::BuildMultiBlockWidget(const ISlateStyle* StyleSet, const FNam
 	.Padding( Padding )	// Large left margin mimics the indent of normal menu items when bNoIndent is false
 	[
 		SNew(SHorizontalBox)
+		.ToolTip(FMultiBoxSettings::ToolTipConstructor.Execute(WidgetBlock->ToolTipText, nullptr, nullptr ))
 		+ SHorizontalBox::Slot()
 		.AutoWidth()
 		[
diff --git a/Engine/Source/Runtime/Slate/Private/Framework/MultiBox/SWidgetBlock.h b/Engine/Source/Runtime/Slate/Private/Framework/MultiBox/SWidgetBlock.h
--- a/Engine/Source/Runtime/Slate/Private/Framework/MultiBox/SWidgetBlock.h
+++ b/Engine/Source/Runtime/Slate/Private/Framework/MultiBox/SWidgetBlock.h
@@ -21,11 +21,12 @@ public:
 	/**
 	 * Constructor
 	 *
-	 * @param	InContent	The widget to place in the block
-	 * @param	InLabel		Optional label text to be added to the left of the content
-	 * @param	bInNoIndent	If true, removes the padding from the left of the widget that lines it up with other menu items
+	 * @param	InContent		The widget to place in the block
+	 * @param	InLabel			Optional label text to be added to the left of the content
+	 * @param	bInNoIndent		If true, removes the padding from the left of the widget that lines it up with other menu items
+	 * @param	InToolTipText	Optional tooltip text to be added to the widget and label
 	 */
-	FWidgetBlock(TSharedRef<SWidget> InContent, const FText& InLabel, bool bInNoIndent, EHorizontalAlignment InHorizontalAlignment = HAlign_Fill);
+	FWidgetBlock(TSharedRef<SWidget> InContent, const FText& InLabel, bool bInNoIndent, EHorizontalAlignment InHorizontalAlignment = HAlign_Fill, const TAttribute<FText>& InToolTipText = FText());
 
 	/** FMultiBlock interface */
 	virtual void CreateMenuEntry(class FMenuBuilder& MenuBuilder) const override;
@@ -49,6 +50,9 @@ private:
 	/** Optional label text */
 	FText Label;
 
+	/** Optional ToolTip text */
+	TAttribute<FText> ToolTipText;
+
 	/** Remove the padding from the left of the widget that lines it up with other menu items? */
 	bool bNoIndent;
 
diff --git a/Engine/Source/Runtime/Slate/Public/Framework/MultiBox/MultiBoxBuilder.h b/Engine/Source/Runtime/Slate/Public/Framework/MultiBox/MultiBoxBuilder.h
--- a/Engine/Source/Runtime/Slate/Public/Framework/MultiBox/MultiBoxBuilder.h
+++ b/Engine/Source/Runtime/Slate/Public/Framework/MultiBox/MultiBoxBuilder.h
@@ -369,8 +369,9 @@ public:
 	 * @param	InLabel				Optional label text to be added to the left of the content
 	 * @param	bNoIndent			If true, removes the padding from the left of the widget that lines it up with other menu items (default == false)
 	 * @param	bSearchable			If true, widget will be searchable (default == true)
+	 * @param	InToolTipText	Optional tooltip text to be added to the widget and label
 	 */
-	void AddWidget( TSharedRef<SWidget> InWidget, const FText& Label, bool bNoIndent = false, bool bSearchable = true );
+	void AddWidget( TSharedRef<SWidget> InWidget, const FText& Label, bool bNoIndent = false, bool bSearchable = true, const TAttribute<FText>& InToolTipText = FText());
 
 	/**
 	* Adds the widget the multibox will use for searching
