TreeView / Listview: Drag + Drop Auto Scroll / Expand

Enables auto scroll and auto expand by default for all list and tree views.

diff --git a/Engine/Plugins/Animation/ControlRig/Source/ControlRigEditor/Private/Editor/SRigHierarchy.cpp b/Engine/Plugins/Animation/ControlRig/Source/ControlRigEditor/Private/Editor/SRigHierarchy.cpp
--- a/Engine/Plugins/Animation/ControlRig/Source/ControlRigEditor/Private/Editor/SRigHierarchy.cpp
+++ b/Engine/Plugins/Animation/ControlRig/Source/ControlRigEditor/Private/Editor/SRigHierarchy.cpp
@@ -229,7 +229,6 @@ void SRigHierarchy::Construct(const FArguments& InArgs, TSharedRef<FControlRigEd
 				[
 					SAssignNew(TreeView, SRigHierarchyTreeView)
 					.RigTreeDelegates(Delegates)
-					.AutoScrollEnabled(true)
 				]
 			]
 		]
diff --git a/Engine/Plugins/Animation/ControlRig/Source/ControlRigEditor/Private/Editor/SRigHierarchyTreeView.cpp b/Engine/Plugins/Animation/ControlRig/Source/ControlRigEditor/Private/Editor/SRigHierarchyTreeView.cpp
--- a/Engine/Plugins/Animation/ControlRig/Source/ControlRigEditor/Private/Editor/SRigHierarchyTreeView.cpp
+++ b/Engine/Plugins/Animation/ControlRig/Source/ControlRigEditor/Private/Editor/SRigHierarchyTreeView.cpp
@@ -206,7 +206,6 @@ FText SRigHierarchyItem::GetName() const
 void SRigHierarchyTreeView::Construct(const FArguments& InArgs)
 {
 	Delegates = InArgs._RigTreeDelegates;
-	bAutoScrollEnabled = InArgs._AutoScrollEnabled;
 
 	STreeView<TSharedPtr<FRigTreeElement>>::FArguments SuperArgs;
 	SuperArgs.TreeItemsSource(&RootElements);
@@ -232,63 +231,6 @@ void SRigHierarchyTreeView::Construct(const FArguments& InArgs)
 	});
 
 	STreeView<TSharedPtr<FRigTreeElement>>::Construct(SuperArgs);
-
-	LastMousePosition = FVector2D::ZeroVector;
-	TimeAtMousePosition = 0.0;
-}
-
-void SRigHierarchyTreeView::Tick(const FGeometry& AllottedGeometry, const double InCurrentTime, const float InDeltaTime)
-{
-	STreeView<TSharedPtr<FRigTreeElement, ESPMode::ThreadSafe>>::Tick(AllottedGeometry, InCurrentTime, InDeltaTime);
-
-	const FGeometry PaintGeometry = GetPaintSpaceGeometry();
-	const FVector2D MousePosition = FSlateApplication::Get().GetCursorPos();
-
-	if(PaintGeometry.IsUnderLocation(MousePosition))
-	{
-		const FVector2D WidgetPosition = PaintGeometry.AbsoluteToLocal(MousePosition);
-
-		static constexpr float SteadyMousePositionTolerance = 5.f;
-
-		if(LastMousePosition.Equals(MousePosition, SteadyMousePositionTolerance))
-		{
-			TimeAtMousePosition += InDeltaTime;
-		}
-		else
-		{
-			LastMousePosition = MousePosition;
-			TimeAtMousePosition = 0.0;
-		}
-
-		static constexpr float AutoScrollStartDuration = 0.5f; // in seconds
-		static constexpr float AutoScrollDistance = 24.f; // in pixels
-		static constexpr float AutoScrollSpeed = 150.f;
-
-		if(TimeAtMousePosition > AutoScrollStartDuration && FSlateApplication::Get().IsDragDropping())
-		{
-			if((WidgetPosition.Y < AutoScrollDistance) || (WidgetPosition.Y > PaintGeometry.Size.Y - AutoScrollDistance))
-			{
-				if(bAutoScrollEnabled)
-				{
-					const bool bScrollUp = (WidgetPosition.Y < AutoScrollDistance);
-
-					const float DeltaInSlateUnits = (bScrollUp ? -InDeltaTime : InDeltaTime) * AutoScrollSpeed; 
-					ScrollBy(GetCachedGeometry(), DeltaInSlateUnits, EAllowOverscroll::No);
-				}
-			}
-			else
-			{
-				const TSharedPtr<FRigTreeElement> Item = FindItemAtPosition(MousePosition);
-				if(Item.IsValid())
-				{
-					if(!IsItemExpanded(Item))
-					{
-						SetItemExpansion(Item, true);
-					}
-				}
-			}
-		}
-	}
 }
 
 TSharedPtr<FRigTreeElement> SRigHierarchyTreeView::FindElement(const FRigElementKey& InElementKey, TSharedPtr<FRigTreeElement> CurrentItem)
diff --git a/Engine/Plugins/Animation/ControlRig/Source/ControlRigEditor/Private/Editor/SRigHierarchyTreeView.h b/Engine/Plugins/Animation/ControlRig/Source/ControlRigEditor/Private/Editor/SRigHierarchyTreeView.h
--- a/Engine/Plugins/Animation/ControlRig/Source/ControlRigEditor/Private/Editor/SRigHierarchyTreeView.h
+++ b/Engine/Plugins/Animation/ControlRig/Source/ControlRigEditor/Private/Editor/SRigHierarchyTreeView.h
@@ -221,18 +221,13 @@ class SRigHierarchyTreeView : public STreeView<TSharedPtr<FRigTreeElement>>
 public:
 
 	SLATE_BEGIN_ARGS(SRigHierarchyTreeView)
-		: _AutoScrollEnabled(false)
 	{}
 		SLATE_ARGUMENT(FRigTreeDelegates, RigTreeDelegates)
-		SLATE_ARGUMENT(bool, AutoScrollEnabled)
 	SLATE_END_ARGS()
 
 	void Construct(const FArguments& InArgs);
 	virtual ~SRigHierarchyTreeView() {}
 
-	/** Performs auto scroll */
-	virtual void Tick(const FGeometry& AllottedGeometry, const double InCurrentTime, const float InDeltaTime) override;
-
 	virtual FReply OnFocusReceived(const FGeometry& MyGeometry, const FFocusEvent& InFocusEvent) override
 	{
 		FReply Reply = STreeView<TSharedPtr<FRigTreeElement>>::OnFocusReceived(MyGeometry, InFocusEvent);
@@ -302,10 +297,6 @@ private:
 
 	FRigTreeDelegates Delegates;
 
-	bool bAutoScrollEnabled;
-	FVector2D LastMousePosition;
-	double TimeAtMousePosition;
-
 	friend class SRigHierarchy;
 };
 
diff --git a/Engine/Source/Runtime/Slate/Private/Widgets/Views/STableViewBase.cpp b/Engine/Source/Runtime/Slate/Private/Widgets/Views/STableViewBase.cpp
--- a/Engine/Source/Runtime/Slate/Private/Widgets/Views/STableViewBase.cpp
+++ b/Engine/Source/Runtime/Slate/Private/Widgets/Views/STableViewBase.cpp
@@ -1129,6 +1129,21 @@ EVisibility STableViewBase::GetPinnedItemsVisiblity() const
 	return PinnedItemsPanel->GetChildren()->Num() != 0 ? EVisibility::Visible : EVisibility::Collapsed;
 }
 
+TSharedPtr<SWidget> STableViewBase::FindItemWidgetAtPosition(const FVector2D& InScreenSpacePosition) const
+{
+	const FGeometry MyGeometry = ItemsPanel->GetCachedGeometry();
+	FArrangedChildren ArrangedChildren(EVisibility::Visible);
+	ItemsPanel->ArrangeChildren(MyGeometry, ArrangedChildren, true);
+
+	const int32 Index = ItemsPanel->FindChildUnderPosition(ArrangedChildren, InScreenSpacePosition); 
+	if(ArrangedChildren.IsValidIndex(Index))
+	{
+		return ArrangedChildren[Index].Widget->AsShared();
+	}
+
+	return TSharedPtr<SWidget>();
+}
+
 static const TBitArray<> EmptyBitArray = TBitArray<>();
 
 const TBitArray<>& TableViewHelpers::GetEmptyBitArray()
diff --git a/Engine/Source/Runtime/Slate/Public/Widgets/Views/SListView.h b/Engine/Source/Runtime/Slate/Public/Widgets/Views/SListView.h
--- a/Engine/Source/Runtime/Slate/Public/Widgets/Views/SListView.h
+++ b/Engine/Source/Runtime/Slate/Public/Widgets/Views/SListView.h
@@ -1678,6 +1678,27 @@ public:
 		return Private_ItemFromWidget( WidgetToFind );
 	}
 
+	/**
+ 	 * Given a screen space position, find the corresponding data item.
+ 	 * Note: This is relying on its CachedGeometry - so the result may
+ 	 * be one frame behind.
+ 	 * 
+ 	 * @param InScreenSpacePosition The position in screen space to search at
+ 	 *
+ 	 * @return the data item residing at the provided screen space position (or nullptr)
+ 	 */
+	const ItemType* ItemAtScreenSpacePosition( const FVector2D& InScreenSpacePosition ) const
+	{
+		const TSharedPtr<SWidget> ItemWidget = this->FindItemWidgetAtPosition(InScreenSpacePosition);
+    	if(ItemWidget.IsValid())
+    	{
+    		TSharedRef<STableRow<ItemType>> TableRow = StaticCastSharedRef<STableRow<ItemType>>(ItemWidget.ToSharedRef());
+    		return this->ItemFromWidget(&TableRow.Get());
+    	}
+
+		return nullptr;
+	}
+
 	/**
 	 * Test if the current item is selected.
 	 *
@@ -2435,7 +2456,67 @@ protected:
 		}
 	}
 
+public:
+
+	virtual void OnDragEnter(const FGeometry& MyGeometry, const FDragDropEvent& DragDropEvent) override
+	{
+		STableViewBase::OnDragEnter(MyGeometry, DragDropEvent);
+
+		// start a timer to deal with auto scroll (and auto expand for treeviews)
+		ActiveMouseMoveDragDropTimer = RegisterActiveTimer(
+			0,
+			FWidgetActiveTimerDelegate::CreateSP(this, &SListView<ItemType>::OnMouseMoveDuringDragDropTimer)
+		);
+	}
+
+	virtual void OnDragLeave(const FDragDropEvent& DragDropEvent) override
+	{
+		STableViewBase::OnDragLeave(DragDropEvent);
+
+		if(ActiveMouseMoveDragDropTimer.IsValid())
+		{
+			UnRegisterActiveTimer(ActiveMouseMoveDragDropTimer.ToSharedRef());
+			ActiveMouseMoveDragDropTimer.Reset();
+		}
+	}
+
 protected:
+
+	EActiveTimerReturnType OnMouseMoveDuringDragDropTimer(const double InCurrentTime, const float InDeltaTime)
+	{
+		HandleMouseMoveDuringDragDrop(InCurrentTime, InDeltaTime);
+		return EActiveTimerReturnType::Continue;
+	}
+
+	virtual void HandleMouseMoveDuringDragDrop(const double InCurrentTime, const float InDeltaTime)
+	{
+		const FVector2D MousePosition = FSlateApplicationBase::Get().GetCursorPos();
+		const FVector2f MousePositionf(MousePosition);
+		
+		if(LastMousePositionDuringDragDrop.Equals(MousePositionf, MousePositionToleranceDuringDragDrop))
+		{
+			DurationAtMousePositionDuringDragDrop += InDeltaTime;
+		}
+		else
+		{
+			LastMousePositionDuringDragDrop = MousePositionf;
+			DurationAtMousePositionDuringDragDrop = 0.0;
+		}
+
+		if(bEnableAutoScroll && DurationAtMousePositionDuringDragDrop > AutoScrollStartDuration)
+		{
+			const FGeometry PaintGeometry = GetPaintSpaceGeometry();
+			const FVector2D WidgetPosition = PaintGeometry.AbsoluteToLocal(MousePosition);
+			if((WidgetPosition.Y < AutoScrollDistance) || (WidgetPosition.Y > PaintGeometry.Size.Y - AutoScrollDistance))
+			{
+				const bool bScrollUp = (WidgetPosition.Y < AutoScrollDistance);
+
+				const float DeltaInSlateUnits = (bScrollUp ? -InDeltaTime : InDeltaTime) * AutoScrollSpeed; 
+				ScrollBy(PaintGeometry, DeltaInSlateUnits, EAllowOverscroll::No);
+			}
+		}
+	}
+	
 	/** A widget generator component */
 	FWidgetGenerator WidgetGenerator;
 
@@ -2538,6 +2619,30 @@ protected:
 	/** If true, number of pinned items > MaxPinnedItems so some items are collapsed in the hierarchy */
 	bool bIsHierarchyCollapsed = false;
 
+	/** If enabled the view will auto scroll if the mouse reaches the top or bottom of the view (during drag and drop only) */
+	static constexpr bool bEnableAutoScroll = true;
+
+	/** The distance from top / bottom within which auto scroll should be enabled */
+	static constexpr float AutoScrollDistance = 24.f;
+
+	/** The duration in seconds it takes for auto scroll to start */
+	static constexpr float AutoScrollStartDuration = 0.5f;
+
+	/** The speed for auto scroll when scrolling up or down */
+	static constexpr float AutoScrollSpeed = 150.f;
+
+	/** The tolerance for a drag and drop mouse move to trigger auto scroll and other actions */
+	static constexpr float MousePositionToleranceDuringDragDrop = 5.f;
+
+	/** During drag and drop we track the time at a (semi) fixed mouse position */
+	double DurationAtMousePositionDuringDragDrop = 0.0;
+
+	/** During drag and drop this stores the last known location */
+	FVector2f LastMousePositionDuringDragDrop = FVector2f::ZeroVector;
+
+	/** The timer dealing with mouse move during drag and drop */
+	TSharedPtr<FActiveTimerHandle> ActiveMouseMoveDragDropTimer;
+
 private:
 	struct FGenerationPassGuard
 	{
diff --git a/Engine/Source/Runtime/Slate/Public/Widgets/Views/STableViewBase.h b/Engine/Source/Runtime/Slate/Public/Widgets/Views/STableViewBase.h
--- a/Engine/Source/Runtime/Slate/Public/Widgets/Views/STableViewBase.h
+++ b/Engine/Source/Runtime/Slate/Public/Widgets/Views/STableViewBase.h
@@ -353,6 +353,9 @@ protected:
 
 	EVisibility GetPinnedItemsVisiblity() const;
 
+	/** returns the item widget at a given screen space position (or an invalid ptr) */
+	TSharedPtr<SWidget> FindItemWidgetAtPosition(const FVector2D& InScreenSpacePosition) const; 
+
 	enum class EScrollIntoViewResult
 	{
 		/** The function scrolled an item (if set) into view (or the item was already in view) */
diff --git a/Engine/Source/Runtime/Slate/Public/Widgets/Views/STreeView.h b/Engine/Source/Runtime/Slate/Public/Widgets/Views/STreeView.h
--- a/Engine/Source/Runtime/Slate/Public/Widgets/Views/STreeView.h
+++ b/Engine/Source/Runtime/Slate/Public/Widgets/Views/STreeView.h
@@ -606,6 +606,23 @@ public:
 		// This may be due to TreeView requesting a refresh or because new items became visible due to resizing or scrolling.
 		SListView< ItemType >::Tick(AllottedGeometry, InCurrentTime, InDeltaTime);
 	}
+	
+	virtual void HandleMouseMoveDuringDragDrop(const double InCurrentTime, const float InDeltaTime) override
+	{
+		SListView< ItemType >::HandleMouseMoveDuringDragDrop(InCurrentTime, InDeltaTime);
+		
+		if(bEnableAutoExpand && this->DurationAtMousePositionDuringDragDrop > AutoExpandStartDuration)
+		{
+			if(const ItemType* ItemPtr = this->ItemAtScreenSpacePosition(FVector2D(this->LastMousePositionDuringDragDrop)))
+			{
+				const ItemType& Item = *ItemPtr;
+				if(!this->IsItemExpanded(Item))
+				{
+					this->SetItemExpansion(Item, true);
+				}
+			}
+		}
+	}
 		
 	/** 
 	 * Given: an array of items (ItemsSource) each of which potentially has a child.
@@ -971,6 +988,12 @@ protected:
 	/** Style resource for the tree */
 	const FTableViewStyle* Style;
 
+	/** If enabled the view will auto expand items if the mouse is steady over an item during drag & drop */
+	static constexpr bool bEnableAutoExpand = true;
+
+	/** The duration in seconds it takes for auto expand to run */
+	static constexpr double AutoExpandStartDuration = 0.5;
+
 private:		
 
 	/** true when the LinearizedItems need to be regenerated. */
