The "Packaging/Cooking complete" toasts now display for the correct duration.

SNotificationList::Tick() now only updates the InternalTime when bAutoExpire is set.

diff --git a/Engine/Source/Runtime/Slate/Private/Widgets/Notifications/SNotificationList.cpp b/Engine/Source/Runtime/Slate/Private/Widgets/Notifications/SNotificationList.cpp
--- a/Engine/Source/Runtime/Slate/Private/Widgets/Notifications/SNotificationList.cpp
+++ b/Engine/Source/Runtime/Slate/Private/Widgets/Notifications/SNotificationList.cpp
@@ -142,15 +142,18 @@ public:
 
 	void Tick( const FGeometry& AllottedGeometry, const double InCurrentTime, const float InDeltaTime ) override
 	{
-		InternalTime += InDeltaTime;
-
 		bool bIsFadingOut = FadeAnimation.IsInReverse();
-		if (bAutoExpire && 
-			!bIsFadingOut && 
-			InternalTime > (FadeInDuration.Get() + ExpireDuration.Get()))
+		if (bAutoExpire)
 		{
-			Fadeout();
-			bIsFadingOut = true;
+			InternalTime += InDeltaTime;
+
+			if (!bIsFadingOut && 
+				InternalTime > (FadeInDuration.Get() + ExpireDuration.Get()))
+			{
+				// if we've faded in and shown the notification for the expire duration, start fading out
+				Fadeout();
+				bIsFadingOut = true;
+			}
 		}
 
 		const bool bIsCurrentlyPlaying = FadeAnimation.IsPlaying();
@@ -159,7 +162,7 @@ public:
 			// Reset the Animation
 			FadeoutComplete();
 		}
-		
+
 		const bool bIsIntroPlaying = IntroAnimation.IsPlaying();
 		if (!bIsIntroPlaying && ThrottleHandle.IsValid())
 		{
