Fix 1 frame flickering when creating a sub text from a Rich text.
The scale was not taken into account at creation time (when creating a SlateWidgetRun), but it was updated in the next slate prepass.
Bonus: removed a bit of duplicated code.

test:
- got a 100% repro, in UMG designer, tested the fix in a simple case, but also with the rich text nested in a scaled vertical box, or in a scale box.
- played Lyra, no regression noticed.

diff --git a/Engine/Source/Runtime/Slate/Private/Framework/Text/SlateWidgetRun.cpp b/Engine/Source/Runtime/Slate/Private/Framework/Text/SlateWidgetRun.cpp
--- a/Engine/Source/Runtime/Slate/Private/Framework/Text/SlateWidgetRun.cpp
+++ b/Engine/Source/Runtime/Slate/Private/Framework/Text/SlateWidgetRun.cpp
@@ -169,17 +169,8 @@ ERunAttributes FSlateWidgetRun::GetRunAttributes() const
 }
 
 FSlateWidgetRun::FSlateWidgetRun(const TSharedRef<class FTextLayout>& InTextLayout, const FRunInfo& InRunInfo, const TSharedRef< const FString >& InText, const FWidgetRunInfo& InWidgetInfo)
-	: TextLayout(InTextLayout)
-	, RunInfo( InRunInfo )
-	, Text( InText )
-	, Range( 0, Text->Len() )
-	, Info( InWidgetInfo )
-	, Children()
-	, WidgetSize()
+	: FSlateWidgetRun(InTextLayout, InRunInfo, InText, InWidgetInfo, FTextRange(0, InText->Len()))
 {
-	Info.Widget->SlatePrepass();
-	WidgetSize = Info.Widget->GetDesiredSize();
-	Children.Add(Info.Widget);
 }
 
 FSlateWidgetRun::FSlateWidgetRun(const TSharedRef<class FTextLayout>& InTextLayout, const FRunInfo& InRunInfo, const TSharedRef< const FString >& InText, const FWidgetRunInfo& InWidgetInfo, const FTextRange& InRange)
@@ -191,7 +182,7 @@ FSlateWidgetRun::FSlateWidgetRun(const TSharedRef<class FTextLayout>& InTextLayo
 	, Children()
 	, WidgetSize()
 {
-	Info.Widget->SlatePrepass();
+	Info.Widget->SlatePrepass(InTextLayout.Get().GetScale());
 	WidgetSize = Info.Widget->GetDesiredSize();
 	Children.Add( Info.Widget );
 }
