UMG: Prevent the game widget to be painted/ticked when the BP is debugging. The SlateApplicatoin is enough to fix the jira but it can cause the same issue if the event was triggered by a mouse click instead of a tick. We need both until all events are registered. With GameLayerManager itself, it won't solve the jira in invalidation mode.

diff --git a/Engine/Source/Runtime/Engine/Private/Slate/SGameLayerManager.cpp b/Engine/Source/Runtime/Engine/Private/Slate/SGameLayerManager.cpp
--- a/Engine/Source/Runtime/Engine/Private/Slate/SGameLayerManager.cpp
+++ b/Engine/Source/Runtime/Engine/Private/Slate/SGameLayerManager.cpp
@@ -391,7 +391,17 @@ void SGameLayerManager::Tick(const FGeometry& AllottedGeometry, const double InC
 int32 SGameLayerManager::OnPaint(const FPaintArgs& Args, const FGeometry& AllottedGeometry, const FSlateRect& MyCullingRect, FSlateWindowElementList& OutDrawElements, int32 LayerId, const FWidgetStyle& InWidgetStyle, bool bParentEnabled) const
 {
 	SCOPED_NAMED_EVENT_TEXT("Paint: Game UI", FColor::Green);
-	
+#if WITH_EDITOR
+	if (GIntraFrameDebuggingGameThread)
+	{
+		// When BP debugging, do not paint the PIE game.
+		//It may trigger other BP code and while it's in a Debugging state the BP will not execute.
+		//It may also tick widgets that may already be ticking.
+		// ie. [A] SMyWidget::Tick() => BPEvent => BP Breakpoint => FApplication::EnterDebuggingMode() => while( SWindow::Paint() => SMyWidget::Tick() ) => End of the first [A] SMyWidget::Tick
+		return FMath::Max(GetPersistentState().OutgoingLayerId, LayerId);
+	}
+#endif
+
 	// By default game layer state is used to ignore disabled effects, don't set this flag if user wants disabled draw effects in-game
 	OutDrawElements.SetIsInGameLayer(!bApplyDisabledEffectOnWidgets);
 	const int32 ResultLayer = SCompoundWidget::OnPaint(Args, AllottedGeometry, MyCullingRect, OutDrawElements, LayerId, InWidgetStyle, bParentEnabled);
diff --git a/Engine/Source/Runtime/Slate/Private/Framework/Application/SlateApplication.cpp b/Engine/Source/Runtime/Slate/Private/Framework/Application/SlateApplication.cpp
--- a/Engine/Source/Runtime/Slate/Private/Framework/Application/SlateApplication.cpp
+++ b/Engine/Source/Runtime/Slate/Private/Framework/Application/SlateApplication.cpp
@@ -1090,6 +1090,15 @@ void FSlateApplication::DrawWindowAndChildren( const TSharedRef<SWindow>& Window
 	}
 #endif
 
+	// Skip Draw if we are debugging that window.
+	if (TSharedPtr<SWindow> CurrentDebuggingWindowPinned = CurrentDebuggingWindow.Pin())
+	{
+		if (CurrentDebuggingWindowPinned == WindowToDraw)
+		{
+			return;
+		}
+	}
+
 	// On Mac, where child windows can be on screen even if their parent is hidden or minimized, we want to always draw child windows.
 	// On other platforms we set bDrawChildWindows to true only if we draw the current window.
 	bool bDrawChildWindows = PLATFORM_MAC;
@@ -1097,6 +1106,7 @@ void FSlateApplication::DrawWindowAndChildren( const TSharedRef<SWindow>& Window
 	// Only draw visible windows or in off-screen rendering mode
 	if (bRenderOffScreen || (WindowToDraw->IsVisible() && (!WindowToDraw->IsWindowMinimized() || FApp::UseVRFocus())) )
 	{
+		TGuardValue TmpContext(CurrentDebugContextWidget, TWeakPtr<SWidget>(WindowToDraw));
 		// Switch to the appropriate world for drawing
 		FScopedSwitchWorldHack SwitchWorld( WindowToDraw );
 
@@ -1188,9 +1198,10 @@ static bool DoAnyWindowDescendantsNeedPrepass(TSharedRef<SWindow> WindowToPrepas
 	return false;
 }
 
-static void PrepassWindowAndChildren( TSharedRef<SWindow> WindowToPrepass )
+static void PrepassWindowAndChildren(TSharedRef<SWindow> WindowToPrepass, const TSharedPtr<SWindow>& DebuggingWindow, TWeakPtr<SWidget>& CurrentContext)
 {
-	if (IsRunningDedicatedServer())
+	// Skip Prepass if we are debugging that window or if we are on the server.
+	if (IsRunningDedicatedServer() || WindowToPrepass == DebuggingWindow)
 	{
 		return;
 	}
@@ -1199,6 +1210,7 @@ static void PrepassWindowAndChildren( TSharedRef<SWindow> WindowToPrepass )
 
 	if (bIsWindowVisible || DoAnyWindowDescendantsNeedPrepass(WindowToPrepass))
 	{
+		TGuardValue TmpContext(CurrentContext, TWeakPtr<SWidget>(WindowToPrepass));
 		FScopedSwitchWorldHack SwitchWorld(WindowToPrepass);
 		
 		{
@@ -1215,7 +1227,7 @@ static void PrepassWindowAndChildren( TSharedRef<SWindow> WindowToPrepass )
 		TArray<TSharedRef<SWindow>, FConcurrentLinearArrayAllocator> ChildWindows(WindowToPrepass->GetChildWindows());
 		for (const TSharedRef<SWindow>& ChildWindow : ChildWindows)
 		{
-			PrepassWindowAndChildren(ChildWindow);
+			PrepassWindowAndChildren(ChildWindow, DebuggingWindow, CurrentContext);
 		}
 	}
 }
@@ -1225,18 +1237,18 @@ void FSlateApplication::DrawPrepass( TSharedPtr<SWindow> DrawOnlyThisWindow )
 	SCOPED_NAMED_EVENT_TEXT("Slate::Prepass", FColor::Magenta);
 	CSV_SCOPED_TIMING_STAT(Slate, Prepass);
 
-	TSharedPtr<SWindow> ActiveModalWindow = GetActiveModalWindow();
+	TSharedPtr<SWindow> CurrentDebuggingWindowPinned = CurrentDebuggingWindow.Pin();
 
-	if (ActiveModalWindow.IsValid())
+	if (TSharedPtr<SWindow> ActiveModalWindow = GetActiveModalWindow())
 	{
-		PrepassWindowAndChildren( ActiveModalWindow.ToSharedRef() );
+		PrepassWindowAndChildren(ActiveModalWindow.ToSharedRef(), CurrentDebuggingWindowPinned, CurrentDebugContextWidget);
 
 		for (TArray< TSharedRef<SWindow> >::TConstIterator CurrentWindowIt(SlateWindows); CurrentWindowIt; ++CurrentWindowIt)
 		{
 			const TSharedRef<SWindow>& CurrentWindow = *CurrentWindowIt;
 			if (CurrentWindow->IsTopmostWindow())
 			{
-				PrepassWindowAndChildren( CurrentWindow );
+				PrepassWindowAndChildren(CurrentWindow, CurrentDebuggingWindowPinned, CurrentDebugContextWidget);
 			}
 		}
 
@@ -1244,19 +1256,19 @@ void FSlateApplication::DrawPrepass( TSharedPtr<SWindow> DrawOnlyThisWindow )
 		FSlateNotificationManager::Get().GetWindows(NotificationWindows);
 		for (auto CurrentWindowIt(NotificationWindows.CreateIterator()); CurrentWindowIt; ++CurrentWindowIt)
 		{
-			PrepassWindowAndChildren(*CurrentWindowIt );
+			PrepassWindowAndChildren(*CurrentWindowIt, CurrentDebuggingWindowPinned, CurrentDebugContextWidget);
 		}
 	}
 	else if (DrawOnlyThisWindow.IsValid())
 	{
-		PrepassWindowAndChildren(DrawOnlyThisWindow.ToSharedRef());
+		PrepassWindowAndChildren(DrawOnlyThisWindow.ToSharedRef(), CurrentDebuggingWindowPinned, CurrentDebugContextWidget);
 	}
 	else
 	{
 		// Draw all windows
 		for (const TSharedRef<SWindow>& CurrentWindow : SlateWindows)
 		{
-			PrepassWindowAndChildren(CurrentWindow);
+			PrepassWindowAndChildren(CurrentWindow, CurrentDebuggingWindowPinned, CurrentDebugContextWidget);
 		}
 	}
 }
@@ -2581,6 +2593,11 @@ void FSlateApplication::ClearKeyboardFocus(const EFocusCause ReasonFocusIsChangi
 	SetUserFocus(GetUserIndexForKeyboard(), FWidgetPath(), ReasonFocusIsChanging);
 }
 
+TSharedPtr<SWidget> FSlateApplication::GetCurrentDebugContextWidget() const
+{
+	return CurrentDebugContextWidget.Pin();
+}
+
 void FSlateApplication::ResetToDefaultInputSettings()
 {
 	ProcessReply(FWidgetPath(), FReply::Handled().ClearUserFocus(true), nullptr, nullptr);
@@ -3523,7 +3540,25 @@ void FSlateApplication::EnterDebuggingMode()
 		PreviousGameViewport->SetActive(false);
 		GameViewportWidget.Reset();
 	}
-	
+
+	// Find the SWindow that we should not tick while in DebuggingMode.
+	ensureMsgf(!CurrentDebuggingWindow.IsValid(), TEXT("Reentry of EnterDebuggingMode with a valid Debugging Window is not supported"));
+	CurrentDebuggingWindow.Reset();
+	if (TSharedPtr<SWidget> CurrentDebugContextWidgetPinned = CurrentDebugContextWidget.Pin())
+	{
+		// Only prevent Paint if there is more than one window.
+		//That is to prevent the user from getting stuck in the editor.
+		if (SlateWindows.Num() > 0)
+		{
+			CurrentDebuggingWindow = FindWidgetWindow(CurrentDebugContextWidgetPinned.ToSharedRef());
+		}
+		else
+		{
+			UE_LOG(LogSlate, Warning, TEXT("EnterDebuggingMode without blocking the window Paint. That may start a new Paint on the same Window while the previous Paint is not completed."));
+		}
+	}
+	TGuardValue TmpContext(CurrentDebugContextWidget, TWeakPtr<SWidget>());
+
 	Renderer->EndFrame();
 
 	Renderer->FlushCommands();
@@ -3570,6 +3605,7 @@ void FSlateApplication::EnterDebuggingMode()
 	Renderer->BeginFrame();
 	bRequestLeaveDebugMode = false;
 	
+	CurrentDebuggingWindow.Reset();
 	if ( PreviousGameViewport.IsValid() )
 	{
 		check(!GameViewportWidget.IsValid());
diff --git a/Engine/Source/Runtime/Slate/Public/Framework/Application/SlateApplication.h b/Engine/Source/Runtime/Slate/Public/Framework/Application/SlateApplication.h
--- a/Engine/Source/Runtime/Slate/Public/Framework/Application/SlateApplication.h
+++ b/Engine/Source/Runtime/Slate/Public/Framework/Application/SlateApplication.h
@@ -1479,6 +1479,7 @@ public:
 	virtual void SetAllUserFocus(const FWidgetPath& InFocusPath, const EFocusCause InCause) override;
 	virtual void SetAllUserFocusAllowingDescendantFocus(const FWidgetPath& InFocusPath, const EFocusCause InCause) override;
 	virtual TSharedPtr<SWidget> GetUserFocusedWidget(uint32 UserIndex) const override;
+	virtual TSharedPtr<SWidget> GetCurrentDebugContextWidget() const override;
 	virtual const TArray<TSharedRef<SWindow>> GetTopLevelWindows() const override { return SlateWindows; }
 
 	DECLARE_EVENT_OneParam(FSlateApplication, FApplicationActivationStateChangedEvent, const bool /*IsActive*/)
@@ -1773,12 +1774,14 @@ private:
 	/** Weak pointers to the allocated virtual users. */
 	TArray<TWeakPtr<FSlateVirtualUserHandle>> VirtualUsers;
 
-	/**
-	  * Last widget that was set for 'all users' focus and the cause.
-	  */
+	/** Last widget that was set for 'all users' focus and the cause. */
 	TWeakPtr<SWidget> LastAllUsersFocusWidget;
 	EFocusCause LastAllUsersFocusCause;
 
+	/** The painting SWindow. */
+	TWeakPtr<SWidget> CurrentDebugContextWidget;
+	TWeakPtr<SWindow> CurrentDebuggingWindow;
+
 	/**
 	 * Application throttling
 	 */
diff --git a/Engine/Source/Runtime/SlateCore/Public/Application/SlateApplicationBase.h b/Engine/Source/Runtime/SlateCore/Public/Application/SlateApplicationBase.h
--- a/Engine/Source/Runtime/SlateCore/Public/Application/SlateApplicationBase.h
+++ b/Engine/Source/Runtime/SlateCore/Public/Application/SlateApplicationBase.h
@@ -473,6 +473,11 @@ public:
 	 */
 	virtual TSharedPtr<SWidget> GetUserFocusedWidget(uint32 UserIndex) const = 0;
 
+	/**
+	 * @return the Widget that started a tick, paint, prepass or an event.
+	 */
+	virtual TSharedPtr<SWidget> GetCurrentDebugContextWidget() const = 0;
+
 	/**
 	 * Gets a delegate that is invoked when a global invalidate of all widgets should occur
 	 */
