Slate: Fix toolbar overflow menu sometimes not having the correct type of entries.

diff --git a/Engine/Source/Runtime/Slate/Private/Framework/MultiBox/SToolBarButtonBlock.cpp b/Engine/Source/Runtime/Slate/Private/Framework/MultiBox/SToolBarButtonBlock.cpp
--- a/Engine/Source/Runtime/Slate/Private/Framework/MultiBox/SToolBarButtonBlock.cpp
+++ b/Engine/Source/Runtime/Slate/Private/Framework/MultiBox/SToolBarButtonBlock.cpp
@@ -65,7 +65,7 @@ void FToolBarButtonBlock::CreateMenuEntry(FMenuBuilder& MenuBuilder) const
 	else if ( LabelOverride.IsSet() )
 	{
 		const FUIAction& DirectAction = GetDirectActions();
-		MenuBuilder.AddMenuEntry( LabelOverride.Get(), ToolTipOverride.Get(), IconOverride.Get(), DirectAction );
+		MenuBuilder.AddMenuEntry( LabelOverride.Get(), ToolTipOverride.Get(), IconOverride.Get(), DirectAction, NAME_None, UserInterfaceActionType);
 	}
 
 	if (bHasValidCommand) 
