PR #9835: Pass style foreground colour parameter to underlying text box in SNumericEntryBox. (Contributed by [REDACTED_PERSON])

diff --git a/Engine/Source/Runtime/Slate/Public/Widgets/Input/SNumericEntryBox.h b/Engine/Source/Runtime/Slate/Public/Widgets/Input/SNumericEntryBox.h
--- a/Engine/Source/Runtime/Slate/Public/Widgets/Input/SNumericEntryBox.h
+++ b/Engine/Source/Runtime/Slate/Public/Widgets/Input/SNumericEntryBox.h
@@ -254,6 +254,7 @@ public:
 		// Always create an editable text box.  In the case of an undetermined value being passed in, we cant use the spinbox.
 		SAssignNew(EditableText, SEditableText)
 			.Text(this, &SNumericEntryBox<NumericType>::OnGetValueForTextBox)
+			.ColorAndOpacity(InArgs._EditableTextBoxStyle->ForegroundColor)
 			.Visibility(bAllowSpin ? EVisibility::Collapsed : EVisibility::Visible)
 			.Font(InArgs._Font.IsSet() ? InArgs._Font : InArgs._EditableTextBoxStyle->TextStyle.Font)
 			.SelectAllTextWhenFocused(true)
