Fixed crash when removing favorite properties from details panel (exact repros unknown)

#rnx

diff --git a/Engine/Source/Runtime/Slate/Private/Widgets/Views/SListPanel.cpp b/Engine/Source/Runtime/Slate/Private/Widgets/Views/SListPanel.cpp
--- a/Engine/Source/Runtime/Slate/Private/Widgets/Views/SListPanel.cpp
+++ b/Engine/Source/Runtime/Slate/Private/Widgets/Views/SListPanel.cpp
@@ -97,7 +97,7 @@ void SListPanel::OnArrangeChildren( const FGeometry& AllottedGeometry, FArranged
 			}
 
 			const int32 FractionalWidgetIndex = NumWholeWidgetsOffset;
-			if (FractionalWidgetIndex < Children.Num())
+			if (Children.IsValidIndex(FractionalWidgetIndex))
 			{
 				const FTableViewDimensions OffsetFractionChildDimensions(Orientation, Children[FractionalWidgetIndex].GetWidget()->GetDesiredSize());
 				ScrollAxisOffset += FMath::Frac(FirstLineScrollOffset) * OffsetFractionChildDimensions.ScrollAxis;
