Expose OnVerifyTextChanged to SSearchBox

diff --git a/Engine/Source/Runtime/Slate/Private/Widgets/Input/SSearchBox.cpp b/Engine/Source/Runtime/Slate/Private/Widgets/Input/SSearchBox.cpp
--- a/Engine/Source/Runtime/Slate/Private/Widgets/Input/SSearchBox.cpp
+++ b/Engine/Source/Runtime/Slate/Private/Widgets/Input/SSearchBox.cpp
@@ -33,6 +33,7 @@ void SSearchBox::Construct( const FArguments& InArgs )
 		.ClearKeyboardFocusOnCommit( false )
 		.OnTextChanged( this, &SSearchBox::HandleTextChanged )
 		.OnTextCommitted( this, &SSearchBox::HandleTextCommitted )
+		.OnVerifyTextChanged( InArgs._OnVerifyTextChanged )
 		.MinDesiredWidth( InArgs._MinDesiredWidth )
 		.OnKeyDownHandler( InArgs._OnKeyDownHandler )
 	);
diff --git a/Engine/Source/Runtime/Slate/Public/Widgets/Input/SSearchBox.h b/Engine/Source/Runtime/Slate/Public/Widgets/Input/SSearchBox.h
--- a/Engine/Source/Runtime/Slate/Public/Widgets/Input/SSearchBox.h
+++ b/Engine/Source/Runtime/Slate/Public/Widgets/Input/SSearchBox.h
@@ -47,6 +47,7 @@ public:
 		, _IsSearching(false)
 		, _OnTextChanged()
 		, _OnTextCommitted()
+		, _OnVerifyTextChanged()
 		, _OnSearch()
 		, _SelectAllTextWhenFocused( true )
 		, _DelayChangeNotificationsWhileTyping( true )
@@ -74,6 +75,9 @@ public:
 		/** Invoked whenever the text is committed (e.g. user presses enter) */
 		SLATE_EVENT( FOnTextCommitted, OnTextCommitted )
 
+		/** Called whenever the text is changed programmatically or interactively by the user */
+		SLATE_EVENT( FOnVerifyTextChanged, OnVerifyTextChanged )
+	
 		/** This will add a next and previous button to your search box */
 		SLATE_EVENT( FOnSearch, OnSearch )
 
