Tab manager now saves layout as soon as you dock or undock a tab.

diff --git a/Engine/Source/Runtime/Slate/Private/Framework/Docking/SDockingNode.cpp b/Engine/Source/Runtime/Slate/Private/Framework/Docking/SDockingNode.cpp
--- a/Engine/Source/Runtime/Slate/Private/Framework/Docking/SDockingNode.cpp
+++ b/Engine/Source/Runtime/Slate/Private/Framework/Docking/SDockingNode.cpp
@@ -27,6 +27,8 @@ float SDockingNode::GetSizeCoefficient() const
 void SDockingNode::SetSizeCoefficient( float InSizeCoefficient )
 {
 	SizeCoefficient = InSizeCoefficient;
+
+	OnResized();
 }
 
 void SDockingNode::OnLiveTabAdded()
@@ -41,7 +43,6 @@ void SDockingNode::OnLiveTabAdded()
 }
 
 SDockingNode::SDockingNode()
-: SizeCoefficient(1.0f)
 {
 
 
diff --git a/Engine/Source/Runtime/Slate/Private/Framework/Docking/SDockingNode.h b/Engine/Source/Runtime/Slate/Private/Framework/Docking/SDockingNode.h
--- a/Engine/Source/Runtime/Slate/Private/Framework/Docking/SDockingNode.h
+++ b/Engine/Source/Runtime/Slate/Private/Framework/Docking/SDockingNode.h
@@ -67,7 +67,7 @@ public:
 	 *
 	 * @param InParent      Set this node's parent.
 	 */
-	virtual void SetParentNode( TSharedRef<class SDockingSplitter> InParent )
+	virtual void SetParentNode( TSharedRef<SDockingSplitter> InParent )
 	{
 		ParentNodePtr = InParent;
 	}
@@ -140,12 +140,14 @@ protected:
 	/** A live tab was added to this node or one of its descendants */
 	virtual void OnLiveTabAdded();
 
+	virtual void OnResized() {}
+
 	/**
 	 * Weak reference to the parent node. It is NULL until the node
 	 * is inserted into the hierarchy. Also NULL for root nodes (aka SDockingArea)
 	 */
-	TWeakPtr<class SDockingSplitter> ParentNodePtr;
+	TWeakPtr<SDockingSplitter> ParentNodePtr;
 
-	float SizeCoefficient;
+	float SizeCoefficient = 1.0f;
 };
 
diff --git a/Engine/Source/Runtime/Slate/Private/Framework/Docking/SDockingSplitter.cpp b/Engine/Source/Runtime/Slate/Private/Framework/Docking/SDockingSplitter.cpp
--- a/Engine/Source/Runtime/Slate/Private/Framework/Docking/SDockingSplitter.cpp
+++ b/Engine/Source/Runtime/Slate/Private/Framework/Docking/SDockingSplitter.cpp
@@ -43,6 +43,14 @@ void SDockingSplitter::AddChildNode( const TSharedRef<SDockingNode>& InChild, in
 	InChild->SetParentNode( SharedThis(this) );
 }
 
+void SDockingSplitter::OnResized()
+{
+	if (TSharedPtr<SDockingArea> Area = GetDockArea())
+	{
+		Area->GetTabManager()->RequestSavePersistentLayout();
+	}
+}
+
 void SDockingSplitter::ReplaceChild( const TSharedRef<SDockingNode>& InChildToReplace, const TSharedRef<SDockingNode>& Replacement )
 {
 	// We want to replace this placeholder with whatever is being dragged.
diff --git a/Engine/Source/Runtime/Slate/Private/Framework/Docking/SDockingSplitter.h b/Engine/Source/Runtime/Slate/Private/Framework/Docking/SDockingSplitter.h
--- a/Engine/Source/Runtime/Slate/Private/Framework/Docking/SDockingSplitter.h
+++ b/Engine/Source/Runtime/Slate/Private/Framework/Docking/SDockingSplitter.h
@@ -108,6 +108,8 @@ protected:
 
 	virtual SDockingNode::ECleanupRetVal CleanUpNodes() override;
 
+	virtual void OnResized() override;
+
 	float ComputeChildCoefficientTotal() const;
 
 	enum class ETabStackToFind
diff --git a/Engine/Source/Runtime/Slate/Private/Framework/Docking/SDockingTabStack.cpp b/Engine/Source/Runtime/Slate/Private/Framework/Docking/SDockingTabStack.cpp
--- a/Engine/Source/Runtime/Slate/Private/Framework/Docking/SDockingTabStack.cpp
+++ b/Engine/Source/Runtime/Slate/Private/Framework/Docking/SDockingTabStack.cpp
@@ -1343,8 +1343,7 @@ void SDockingTabStack::ExecuteCloseMinorTabCommand()
 
 bool SDockingTabStack::CanExecuteCloseMinorTabCommand()
 {
-	auto DockArea = GetDockArea();
-	if (DockArea.IsValid())
+	if (TSharedPtr<SDockingArea> DockArea = GetDockArea())
 	{
 		TSharedPtr<FGlobalTabmanager> GlobalTabManager = FGlobalTabmanager::Get();
 		TSharedPtr<SDockTab> ActiveTab = GlobalTabManager->GetActiveTab();
@@ -1360,4 +1359,12 @@ bool SDockingTabStack::CanExecuteCloseMinorTabCommand()
 	return false;
 }
 
+void SDockingTabStack::OnResized()
+{
+	if (TSharedPtr<SDockingArea> DockArea = GetDockArea())
+	{
+		DockArea->GetTabManager()->RequestSavePersistentLayout();
+	}
+}
+
 #undef LOCTEXT_NAMESPACE
diff --git a/Engine/Source/Runtime/Slate/Private/Framework/Docking/SDockingTabStack.h b/Engine/Source/Runtime/Slate/Private/Framework/Docking/SDockingTabStack.h
--- a/Engine/Source/Runtime/Slate/Private/Framework/Docking/SDockingTabStack.h
+++ b/Engine/Source/Runtime/Slate/Private/Framework/Docking/SDockingTabStack.h
@@ -307,6 +307,9 @@ private:
 	/** Attempts to find the active tab that can be closed by the CloseMinorTab command */
 	bool CanExecuteCloseMinorTabCommand();
 
+	/** Called when this stack is resized. */
+	virtual void OnResized() override;
+
 #if DEBUG_TAB_MANAGEMENT
 	FString ShowPersistentTabs() const;
 #endif
diff --git a/Engine/Source/Runtime/Slate/Private/Framework/Docking/TabManager.cpp b/Engine/Source/Runtime/Slate/Private/Framework/Docking/TabManager.cpp
--- a/Engine/Source/Runtime/Slate/Private/Framework/Docking/TabManager.cpp
+++ b/Engine/Source/Runtime/Slate/Private/Framework/Docking/TabManager.cpp
@@ -951,8 +951,34 @@ TSharedRef<FTabManager::FLayout> FTabManager::PersistLayout() const
 
 void FTabManager::SavePersistentLayout()
 {
+	if (PendingLayoutSaveHandle.IsValid())
+	{
+		FTSTicker::GetCoreTicker().RemoveTicker(PendingLayoutSaveHandle);
+		PendingLayoutSaveHandle.Reset();
+	}
+
 	const TSharedRef<FLayout> LayoutState = this->PersistLayout();
-	OnPersistLayout_Handler.ExecuteIfBound( LayoutState );
+	OnPersistLayout_Handler.ExecuteIfBound(LayoutState);
+}
+
+void FTabManager::RequestSavePersistentLayout()
+{
+	// if we already have a request pending, remove it and schedule a new one
+	// this is to avoid hitches when eg. resizing a docked tab
+	if (PendingLayoutSaveHandle.IsValid())
+	{
+		FTSTicker::GetCoreTicker().RemoveTicker(PendingLayoutSaveHandle);
+		PendingLayoutSaveHandle.Reset();
+	}
+
+	auto OnTick = [this](float FrameTime)
+	{
+		this->PendingLayoutSaveHandle.Reset();
+		this->SavePersistentLayout();
+		return false;
+	};
+
+	PendingLayoutSaveHandle = FTSTicker::GetCoreTicker().AddTicker(FTickerDelegate::CreateLambda(OnTick), 5.0f);
 }
 
 FTabSpawnerEntry& FTabManager::RegisterTabSpawner(const FName TabId, const FOnSpawnTab& OnSpawnTab, const FCanSpawnTab& CanSpawnTab)
@@ -1540,8 +1566,6 @@ FTabManager::FTabManager( const TSharedPtr<SDockTab>& InOwnerTab, const TSharedR
 , OwnerTabPtr( InOwnerTab )
 , PrivateApi( MakeShareable(new FPrivateApi(*this)) )
 , LastDocumentUID( 0 )
-, bIsSavingVisualState( false )
-, bCanDoDragOperation( true )
 , TabPermissionList( MakeShareable(new FNamePermissionList()) )
 {
 	LocalWorkspaceMenuRoot = FWorkspaceItem::NewGroup(LOCTEXT("LocalWorkspaceRoot", "Local Workspace Root"));
@@ -2227,6 +2251,13 @@ void FTabManager::OnTabRelocated( const TSharedRef<SDockTab>& RelocatedTab, cons
 	FGlobalTabmanager::Get()->UpdateMainMenu(RelocatedTab, true);
 
 	UpdateStats();
+
+	RequestSavePersistentLayout();
+
+	if (TSharedPtr<FTabManager> NewTabManager = RelocatedTab->GetTabManagerPtr())
+	{
+		NewTabManager->RequestSavePersistentLayout();
+	}
 }
 
 void FTabManager::OnTabOpening( const TSharedRef<SDockTab>& TabBeingOpened )
@@ -2264,12 +2295,12 @@ bool FTabManager::CanCloseManager( const TSet< TSharedRef<SDockTab> >& TabsToIgn
 	bool bCanCloseManager = true;
 
 	for (int32 DockAreaIndex=0; bCanCloseManager && DockAreaIndex < DockAreas.Num(); ++DockAreaIndex)
-	{
+		{
 		TSharedPtr<SDockingArea> SomeArea = DockAreas[DockAreaIndex].Pin();
 		TArray< TSharedRef<SDockTab> > AreasTabs = SomeArea.IsValid() ? SomeArea->GetAllChildTabs() : TArray< TSharedRef<SDockTab> >();
 		
 		for (int32 TabIndex=0; bCanCloseManager && TabIndex < AreasTabs.Num(); ++TabIndex)	
-		{
+			{
 			bCanCloseManager =
 				TabsToIgnore.Contains( AreasTabs[TabIndex] ) ||
 				AreasTabs[TabIndex]->GetTabRole() != ETabRole::MajorTab ||
@@ -2318,8 +2349,8 @@ TSharedPtr<FTabManager::FStack> FTabManager::FindTabUnderNode( const FTabMatcher
 			StackWithTab = FindTabUnderNode( Matcher, NodeAsSplitter->ChildNodes[ChildIndex] );
 		}
 
-		return StackWithTab;
-	}
+			return StackWithTab;
+		}
 }
 
 
@@ -2348,7 +2379,7 @@ const TSharedPtr<const FTabSpawnerEntry> FTabManager::FindTabSpawnerFor(FName Ta
 
 	return (Spawner != nullptr)
 		? TSharedPtr<const FTabSpawnerEntry>(*Spawner)
-		: TSharedPtr<FTabSpawnerEntry>();
+		: TSharedPtr<const FTabSpawnerEntry>();
 }
 
 int32 FTabManager::FindTabInCollapsedAreas( const FTabMatcher& Matcher )
diff --git a/Engine/Source/Runtime/Slate/Public/Framework/Docking/TabManager.h b/Engine/Source/Runtime/Slate/Public/Framework/Docking/TabManager.h
--- a/Engine/Source/Runtime/Slate/Public/Framework/Docking/TabManager.h
+++ b/Engine/Source/Runtime/Slate/Public/Framework/Docking/TabManager.h
@@ -4,6 +4,7 @@
 
 #include "CoreMinimal.h"
 #include "SlateFwd.h"
+#include "Containers/Ticker.h"
 #include "Misc/Attribute.h"
 #include "Textures/SlateIcon.h"
 #include "Widgets/SWindow.h"
@@ -831,6 +832,9 @@ class SLATE_API FTabManager : public TSharedFromThis<FTabManager>
 		/** Gather the persistent layout and execute the custom delegate for saving it to persistent storage (e.g. into config files) */
 		void SavePersistentLayout();
 
+		/** Request a deferred save of the layout. */
+		void RequestSavePersistentLayout();
+
 		/**
 		 * Register a new tab spawner with the tab manager.  The spawner will be called when anyone calls
 		 * InvokeTab().
@@ -1141,12 +1145,6 @@ class SLATE_API FTabManager : public TSharedFromThis<FTabManager>
 		/** Returns the default window size for the TabId, or the fallback window size if it wasn't registered */
 		static FVector2D GetDefaultTabWindowSize(const FTabId& TabId);
 
-		/**
-		 * Defensive: True when we are saving the visual state.
-		 * Inevitably someone will ask us to save layout while saving layout.
-		 * We will ignore that request them.
-		 */
-		bool bIsSavingVisualState;
 
 		/** The main tab, this tab cannot be closed. */
 		TWeakPtr<const SDockTab> MainNonCloseableTab;
@@ -1155,11 +1153,14 @@ class SLATE_API FTabManager : public TSharedFromThis<FTabManager>
 		TWeakPtr<SWindow> LastMajorDockWindow;
 
 		/* Prevent or allow Drag operation. */
-		bool bCanDoDragOperation;
+		bool bCanDoDragOperation = true;
 
 		/** Whether or not this tab manager puts any registered menus in the windows menu bar area */
 		bool bAllowPerWindowMenu = false;
 
+		/** Handle to a pending layout save. */
+		FTSTicker::FDelegateHandle PendingLayoutSaveHandle;
+
 		/** Allow systems to dynamically hide tabs */
 		TSharedRef<FNamePermissionList> TabPermissionList;
 
@@ -1335,8 +1336,8 @@ private:
 
 	FGlobalTabmanager()
 	: FTabManager( TSharedPtr<SDockTab>(), MakeShareable( new FTabSpawner() ) )
-	, AllTabsMaxCount(0)
-	, AllAreasWindowMaxCount(0)
+		, AllTabsMaxCount(0)
+		, AllAreasWindowMaxCount(0)
 	{
 	}
 
