UE TextLayout:  create an option for text flow to be set by culture
- this will allow us to predictably handle text that might include names from other cultures but should still be read the same as the set culture

diff --git a/Engine/Source/Runtime/Slate/Private/Framework/Text/TextLayout.cpp b/Engine/Source/Runtime/Slate/Private/Framework/Text/TextLayout.cpp
--- a/Engine/Source/Runtime/Slate/Private/Framework/Text/TextLayout.cpp
+++ b/Engine/Source/Runtime/Slate/Private/Framework/Text/TextLayout.cpp
@@ -7,6 +7,8 @@
 #include "Framework/Text/TextHitPoint.h"
 #include "Framework/Text/ILayoutBlock.h"
 #include "Internationalization/BreakIterator.h"
+#include "Internationalization/Culture.h"
+#include "Internationalization/Internationalization.h"
 #include "Internationalization/TextTransformer.h"
 #include "Framework/Text/ShapedTextCache.h"
 
@@ -1039,6 +1041,9 @@ void FTextLayout::CalculateLineTextDirection(FLineModel& LineModel) const
 	case ETextFlowDirection::RightToLeft:
 		LineModel.TextBaseDirection = TextBiDi::ETextDirection::RightToLeft;
 		break;
+	case ETextFlowDirection::Culture:
+		LineModel.TextBaseDirection = FInternationalization::Get().GetCurrentLanguage()->IsRightToLeft() ? TextBiDi::ETextDirection::RightToLeft : TextBiDi::ETextDirection::LeftToRight;
+		break;
 	default:
 		break;
 	}
diff --git a/Engine/Source/Runtime/Slate/Public/Framework/Text/TextLayout.h b/Engine/Source/Runtime/Slate/Public/Framework/Text/TextLayout.h
--- a/Engine/Source/Runtime/Slate/Public/Framework/Text/TextLayout.h
+++ b/Engine/Source/Runtime/Slate/Public/Framework/Text/TextLayout.h
@@ -77,6 +77,9 @@ enum class ETextFlowDirection : uint8
 
 	/** Force text to be flowed right-to-left */
 	RightToLeft,
+
+	/** Uses the set culture to determine if text should flow left-to-right or right-to-left.  By comparison, Auto will use the text itself to determine it. */
+	Culture,
 };
 
 /** Get the default text flow direction (from the "Slate.DefaultTextFlowDirection" CVar) */
