Slate: Deprecated protected access to SListView::ItemsSource. This will allow us to use other container type.

diff --git a/Engine/Plugins/Animation/ControlRig/Source/ControlRigEditor/Private/EditMode/SControlRigOutliner.cpp b/Engine/Plugins/Animation/ControlRig/Source/ControlRigEditor/Private/EditMode/SControlRigOutliner.cpp
--- a/Engine/Plugins/Animation/ControlRig/Source/ControlRigEditor/Private/EditMode/SControlRigOutliner.cpp
+++ b/Engine/Plugins/Animation/ControlRig/Source/ControlRigEditor/Private/EditMode/SControlRigOutliner.cpp
@@ -856,16 +856,16 @@ TArray<FMultiRigData> SMultiRigHierarchyTreeView::GetSelectedData() const
 
 TSharedPtr<FMultiRigTreeElement> SMultiRigHierarchyTreeView::FindItemAtPosition(FVector2D InScreenSpacePosition) const
 {
-	if (ItemsPanel.IsValid() && ItemsSource != nullptr)
+	if (ItemsPanel.IsValid() && HasValidItemsSource())
 	{
 		const FGeometry MyGeometry = ItemsPanel->GetCachedGeometry();
 		FArrangedChildren ArrangedChildren(EVisibility::Visible);
 		ItemsPanel->ArrangeChildren(MyGeometry, ArrangedChildren, true);
 
 		const int32 Index = ItemsPanel->FindChildUnderPosition(ArrangedChildren, InScreenSpacePosition);
-		if (ItemsSource->IsValidIndex(Index))
+		if (GetItems().IsValidIndex(Index))
 		{
-			return ItemsSource->operator[](Index);
+			return GetItems()[Index];
 		}
 	}
 	return TSharedPtr<FMultiRigTreeElement>();
diff --git a/Engine/Plugins/Animation/ControlRig/Source/ControlRigEditor/Private/Editor/SRigHierarchyTreeView.cpp b/Engine/Plugins/Animation/ControlRig/Source/ControlRigEditor/Private/Editor/SRigHierarchyTreeView.cpp
--- a/Engine/Plugins/Animation/ControlRig/Source/ControlRigEditor/Private/Editor/SRigHierarchyTreeView.cpp
+++ b/Engine/Plugins/Animation/ControlRig/Source/ControlRigEditor/Private/Editor/SRigHierarchyTreeView.cpp
@@ -750,7 +750,7 @@ TArray<FRigElementKey> SRigHierarchyTreeView::GetSelectedKeys() const
 
 TSharedPtr<FRigTreeElement> SRigHierarchyTreeView::FindItemAtPosition(FVector2D InScreenSpacePosition) const
 {
-	if(ItemsPanel.IsValid() && ItemsSource != nullptr)
+	if(ItemsPanel.IsValid() && HasValidItemsSource())
 	{
 		const FGeometry MyGeometry = ItemsPanel->GetCachedGeometry();
 		FArrangedChildren ArrangedChildren(EVisibility::Visible);
@@ -763,7 +763,7 @@ TSharedPtr<FRigTreeElement> SRigHierarchyTreeView::FindItemAtPosition(FVector2D
 			if(ItemWidget->WeakRigTreeElement.IsValid())
 			{
 				const FRigElementKey Key = ItemWidget->WeakRigTreeElement.Pin()->Key;
-				const TSharedPtr<FRigTreeElement>* ResultPtr = ItemsSource->FindByPredicate([Key](const TSharedPtr<FRigTreeElement>& Item) -> bool
+				const TSharedPtr<FRigTreeElement>* ResultPtr = GetItems().FindByPredicate([Key](const TSharedPtr<FRigTreeElement>& Item) -> bool
 				{
 					return Item->Key == Key;
 				});
diff --git a/Engine/Plugins/Animation/ControlRig/Source/ControlRigEditor/Private/Graph/SControlRigGraphPinNameListValueWidget.cpp b/Engine/Plugins/Animation/ControlRig/Source/ControlRigEditor/Private/Graph/SControlRigGraphPinNameListValueWidget.cpp
--- a/Engine/Plugins/Animation/ControlRig/Source/ControlRigEditor/Private/Graph/SControlRigGraphPinNameListValueWidget.cpp
+++ b/Engine/Plugins/Animation/ControlRig/Source/ControlRigEditor/Private/Graph/SControlRigGraphPinNameListValueWidget.cpp
@@ -104,7 +104,7 @@ void SControlRigGraphPinNameListValueWidget::SetOptionsSource(const TArray<TShar
 {
 	check(InOptionsSource);
 	OptionsSource = InOptionsSource;
-	ComboListView->SetListItemsSource(*OptionsSource);
+	ComboListView->SetItemsSource(OptionsSource);
 }
 
 TSharedPtr<FString> SControlRigGraphPinNameListValueWidget::GetSelectedItem()
diff --git a/Engine/Plugins/Experimental/CommonUI/Source/CommonUI/Public/CommonListView.h b/Engine/Plugins/Experimental/CommonUI/Source/CommonUI/Public/CommonListView.h
--- a/Engine/Plugins/Experimental/CommonUI/Source/CommonUI/Public/CommonListView.h
+++ b/Engine/Plugins/Experimental/CommonUI/Source/CommonUI/Public/CommonListView.h
@@ -19,12 +19,12 @@ public:
 		{
 			// Set selection to the first item in a list if no items are selected.
 			// If bReturnFocusToSelection is true find the last selected object and focus on that.
-			if (this->ItemsSource && this->ItemsSource->Num() > 0)
+			if (this->GetItems().Num() > 0)
 			{
 				typename TListTypeTraits<ItemType>::NullableType ItemNavigatedTo(nullptr);
 				if (this->GetNumItemsSelected() == 0)
 				{
-					ItemNavigatedTo = (*this->ItemsSource)[0];
+					ItemNavigatedTo = this->GetItems()[0];
 				}
 				else if (this->bReturnFocusToSelection && TListTypeTraits<ItemType>::IsPtrValid(this->SelectorItem))
 				{
diff --git a/Engine/Plugins/Experimental/CommonUI/Source/CommonUI/Public/CommonTreeView.h b/Engine/Plugins/Experimental/CommonUI/Source/CommonUI/Public/CommonTreeView.h
--- a/Engine/Plugins/Experimental/CommonUI/Source/CommonUI/Public/CommonTreeView.h
+++ b/Engine/Plugins/Experimental/CommonUI/Source/CommonUI/Public/CommonTreeView.h
@@ -19,7 +19,7 @@ public:
 		{
 			// Set selection to the first item in a list if no items are selected.
 			// If bReturnFocusToSelection is true find the last selected object and focus on that.
-			if (this->ItemsSource && this->TreeItemsSource->Num() > 0)
+			if (this->HasValidItemsSource() && this->TreeItemsSource->Num() > 0)
 			{
 				typename TListTypeTraits<ItemType>::NullableType ItemNavigatedTo(nullptr);
 				if (this->GetNumItemsSelected() == 0)
diff --git a/Engine/Plugins/Experimental/CommonUI/Source/CommonUI/Public/SCommonTileView.h b/Engine/Plugins/Experimental/CommonUI/Source/CommonUI/Public/SCommonTileView.h
--- a/Engine/Plugins/Experimental/CommonUI/Source/CommonUI/Public/SCommonTileView.h
+++ b/Engine/Plugins/Experimental/CommonUI/Source/CommonUI/Public/SCommonTileView.h
@@ -12,11 +12,11 @@ public:
 	{
 		if (bScrollToSelectedOnFocus && (InFocusEvent.GetCause() == EFocusCause::Navigation || InFocusEvent.GetCause() == EFocusCause::SetDirectly))
 		{
-			if (this->ItemsSource && this->ItemsSource->Num() > 0)
+			if (this->GetItems().Num() > 0)
 			{
 				if (this->GetNumItemsSelected() == 0)
 				{
-					ItemType FirstItem = (*this->ItemsSource)[0];
+					ItemType FirstItem = this->GetItems()[0];
 					this->SetSelection(FirstItem, ESelectInfo::OnNavigation);
 					this->RequestNavigateToItem(FirstItem, InFocusEvent.GetUser());
 				}
diff --git a/Engine/Plugins/Experimental/PCG/Source/PCGEditor/Private/SPCGEditorGraphAttributeListView.cpp b/Engine/Plugins/Experimental/PCG/Source/PCGEditor/Private/SPCGEditorGraphAttributeListView.cpp
--- a/Engine/Plugins/Experimental/PCG/Source/PCGEditor/Private/SPCGEditorGraphAttributeListView.cpp
+++ b/Engine/Plugins/Experimental/PCG/Source/PCGEditor/Private/SPCGEditorGraphAttributeListView.cpp
@@ -569,7 +569,7 @@ void SPCGEditorGraphAttributeListView::RefreshAttributeList()
 {
 	// Swapping to an empty item list to force a widget clear, otherwise the widgets will try to update during add column and access invalid data
 	static const TArray<PCGListviewItemPtr> EmptyList;
-	ListView->SetListItemsSource(EmptyList);
+	ListView->SetItemsSource(&EmptyList);
 	
 	ListViewItems.Empty();
 	ListViewHeader->ClearColumns();
@@ -659,7 +659,7 @@ void SPCGEditorGraphAttributeListView::RefreshAttributeList()
 		}
 	}
 
-	ListView->SetListItemsSource(ListViewItems);
+	ListView->SetItemsSource(&ListViewItems);
 }
 
 void SPCGEditorGraphAttributeListView::RefreshDataComboBox()
diff --git a/Engine/Plugins/Experimental/PCG/Source/PCGEditor/Private/SPCGEditorGraphLogView.cpp b/Engine/Plugins/Experimental/PCG/Source/PCGEditor/Private/SPCGEditorGraphLogView.cpp
--- a/Engine/Plugins/Experimental/PCG/Source/PCGEditor/Private/SPCGEditorGraphLogView.cpp
+++ b/Engine/Plugins/Experimental/PCG/Source/PCGEditor/Private/SPCGEditorGraphLogView.cpp
@@ -360,7 +360,7 @@ FReply SPCGEditorGraphLogView::Refresh()
 			});
 	}
 
-	ListView->SetListItemsSource(ListViewItems);
+	ListView->SetItemsSource(&ListViewItems);
 
 	return FReply::Handled();
 }
diff --git a/Engine/Plugins/Experimental/PCG/Source/PCGEditor/Private/SPCGEditorGraphProfilingView.cpp b/Engine/Plugins/Experimental/PCG/Source/PCGEditor/Private/SPCGEditorGraphProfilingView.cpp
--- a/Engine/Plugins/Experimental/PCG/Source/PCGEditor/Private/SPCGEditorGraphProfilingView.cpp
+++ b/Engine/Plugins/Experimental/PCG/Source/PCGEditor/Private/SPCGEditorGraphProfilingView.cpp
@@ -532,7 +532,7 @@ FReply SPCGEditorGraphProfilingView::Refresh()
 			});
 	}
 
-	ListView->SetListItemsSource(ListViewItems);
+	ListView->SetItemsSource(&ListViewItems);
 
 	return FReply::Handled();
 }
diff --git a/Engine/Source/Editor/CurveTableEditor/Private/CurveTableEditor.cpp b/Engine/Source/Editor/CurveTableEditor/Private/CurveTableEditor.cpp
--- a/Engine/Source/Editor/CurveTableEditor/Private/CurveTableEditor.cpp
+++ b/Engine/Source/Editor/CurveTableEditor/Private/CurveTableEditor.cpp
@@ -823,7 +823,7 @@ void FCurveTableEditor::RefreshCachedCurveTable()
 	// New Selection 
 	TArray<FCurveEditorTreeItemID> NewSelectedItems;
 
-	TableView->SetListItemsSource(EmptyItems);
+	TableView->SetItemsSource(&EmptyItems);
 	
 	CurveEditor->RemoveAllTreeItems();
 
@@ -927,7 +927,7 @@ void FCurveTableEditor::RefreshCachedCurveTable()
 		}
 	}
 
-	TableView->SetListItemsSource(CurveEditorTree->GetSourceItems());
+	TableView->SetItemsSource(&CurveEditorTree->GetSourceItems());
 
 	TGuardValue<bool> SelectionGuard(bUpdatingTableViewSelection, true);
 	CurveEditor->SetTreeSelection(MoveTemp(NewSelectedItems));
diff --git a/Engine/Source/Editor/SequencerCore/Private/MVVM/Views/SOutlinerView.cpp b/Engine/Source/Editor/SequencerCore/Private/MVVM/Views/SOutlinerView.cpp
--- a/Engine/Source/Editor/SequencerCore/Private/MVVM/Views/SOutlinerView.cpp
+++ b/Engine/Source/Editor/SequencerCore/Private/MVVM/Views/SOutlinerView.cpp
@@ -561,7 +561,7 @@ FReply SOutlinerView::OnDrop(const FGeometry& MyGeometry, const FDragDropEvent&
 
 FReply SOutlinerView::OnKeyDown(const FGeometry& MyGeometry, const FKeyEvent& InKeyEvent)
 {
-	const TArray<TWeakViewModelPtr<IOutlinerExtension>>& ItemsSourceRef = (*this->ItemsSource);
+	const TArrayView<const TWeakViewModelPtr<IOutlinerExtension>> ItemsSourceRef = GetItems();
 
 	auto IsSelectable = [](const TWeakViewModelPtr<IOutlinerExtension>& WeakModel)
 	{
@@ -984,7 +984,7 @@ void SOutlinerView::UpdatePhysicalGeometry(bool bIsRefresh)
 	}
 }
 
-bool ShouldExpand(const TArray<TWeakViewModelPtr<IOutlinerExtension>>& Selection, ETreeRecursion Recursion, TSharedPtr<FOutlinerViewModel> Outliner)
+bool ShouldExpand(const TArrayView<const TWeakViewModelPtr<IOutlinerExtension>>& Selection, ETreeRecursion Recursion, TSharedPtr<FOutlinerViewModel> Outliner)
 {
 	bool bAllExpanded = true;
 	for (TWeakViewModelPtr<IOutlinerExtension> WeakItem : Selection)
@@ -1033,9 +1033,9 @@ void SOutlinerView::ToggleExpandCollapseNodes(ETreeRecursion Recursion, bool bEx
 			}
 		}
 	}
-	else if (ItemsSource->Num() > 0)
+	else if (GetItems().Num() > 0)
 	{
-		bool bExpand = ShouldExpand(*ItemsSource, Recursion, Outliner);
+		bool bExpand = ShouldExpand(GetItems(), Recursion, Outliner);
 
 		if (bExpandAll)
 		{
@@ -1046,7 +1046,7 @@ void SOutlinerView::ToggleExpandCollapseNodes(ETreeRecursion Recursion, bool bEx
 			bExpand = false;
 		}
 
-		for (TWeakViewModelPtr<IOutlinerExtension> WeakItem : *ItemsSource)
+		for (TWeakViewModelPtr<IOutlinerExtension> WeakItem : GetItems())
 		{
 			if (TViewModelPtr<IOutlinerExtension> Item = WeakItem.Pin())
 			{
diff --git a/Engine/Source/Editor/SkeletonEditor/Private/SSkeletonTree.cpp b/Engine/Source/Editor/SkeletonEditor/Private/SSkeletonTree.cpp
--- a/Engine/Source/Editor/SkeletonEditor/Private/SSkeletonTree.cpp
+++ b/Engine/Source/Editor/SkeletonEditor/Private/SSkeletonTree.cpp
@@ -94,7 +94,7 @@ public:
 			return;
 		}
 
-		const TArray<ItemType>& ItemsSourceRef = (*this->ItemsSource);
+		const TArrayView<const ItemType> ItemsSourceRef = this->GetItems();
 
 		int32 RangeStartIndex = 0;
 		if( TListTypeTraits<ItemType>::IsPtrValid(this->RangeSelectionStart) )
diff --git a/Engine/Source/Runtime/Slate/Public/Widgets/Views/SListView.h b/Engine/Source/Runtime/Slate/Public/Widgets/Views/SListView.h
--- a/Engine/Source/Runtime/Slate/Public/Widgets/Views/SListView.h
+++ b/Engine/Source/Runtime/Slate/Public/Widgets/Views/SListView.h
@@ -3,6 +3,7 @@
 #pragma once
 
 #include "CoreMinimal.h"
+#include "Containers/ArrayView.h"
 #include "InputCoreTypes.h"
 #include "Input/Reply.h"
 #include "Layout/Visibility.h"
@@ -218,7 +219,9 @@ public:
 		this->OnItemScrolledIntoView = InArgs._OnItemScrolledIntoView;
 		this->OnFinishedScrolling = InArgs._OnFinishedScrolling;
 
-		this->ItemsSource = InArgs._ListItemsSource;
+PRAGMA_DISABLE_DEPRECATION_WARNINGS
+		this->SetItemsSource(InArgs._ListItemsSource);
+PRAGMA_ENABLE_DEPRECATION_WARNINGS
 		this->OnContextMenuOpening = InArgs._OnContextMenuOpening;
 		this->OnClick = InArgs._OnMouseButtonClick;
 		this->OnDoubleClick = InArgs._OnMouseButtonDoubleClick;
@@ -243,11 +246,10 @@ public:
 		this->bEnableAnimatedScrolling = InArgs._EnableAnimatedScrolling;
 		this->FixedLineScrollOffset = InArgs._FixedLineScrollOffset;
 
-		this->OnItemToString_Debug =
-			InArgs._OnItemToString_Debug.IsBound()
+		this->OnItemToString_Debug = InArgs._OnItemToString_Debug.IsBound()
 			? InArgs._OnItemToString_Debug
-			: GetDefaultDebugDelegate();
-		OnEnteredBadState = InArgs._OnEnteredBadState;
+			: SListView< ItemType >::GetDefaultDebugDelegate();
+		this->OnEnteredBadState = InArgs._OnEnteredBadState;
 
 		this->OnKeyDownHandler = InArgs._OnKeyDownHandler;
 
@@ -264,7 +266,7 @@ public:
 				ErrorString += TEXT("Please specify an OnGenerateRow. \n");
 			}
 
-			if ( this->ItemsSource == nullptr )
+			if ( !this->HasValidItemsSource() )
 			{
 				ErrorString += TEXT("Please specify a ListItemsSource. \n");
 			}
@@ -285,15 +287,16 @@ public:
 		{
 			// Make the TableView
 			ConstructChildren( 0, InArgs._ItemHeight, EListItemAlignment::LeftAligned, InArgs._HeaderRow, InArgs._ExternalScrollbar, InArgs._Orientation, InArgs._OnListViewScrolled, InArgs._ScrollBarStyle, InArgs._PreventThrottling );
-			if(ScrollBar.IsValid())
+			if(this->ScrollBar.IsValid())
 			{
-				ScrollBar->SetDragFocusCause(InArgs._ScrollbarDragFocusCause);
-				ScrollBar->SetUserVisibility(InArgs._ScrollbarVisibility);
+				this->ScrollBar->SetDragFocusCause(InArgs._ScrollbarDragFocusCause);
+				this->ScrollBar->SetUserVisibility(InArgs._ScrollbarVisibility);
 			}
 			this->AddMetadata(MakeShared<TTableViewMetadata<ItemType>>(this->SharedThis(this)));
 		}
 	}
 
+PRAGMA_DISABLE_DEPRECATION_WARNINGS
 	SListView(ETableViewMode::Type InListMode = ETableViewMode::List)
 		: STableViewBase(InListMode)
 		, WidgetGenerator(this)
@@ -311,6 +314,7 @@ public:
 		bCanChildrenBeAccessible = true;
 #endif
 	}
+PRAGMA_ENABLE_DEPRECATION_WARNINGS
 
 public:
 
@@ -332,7 +336,7 @@ public:
 			}
 		}
 
-		const TArray<ItemType>& ItemsSourceRef = (*this->ItemsSource);
+		const TArrayView<const ItemType> ItemsSourceRef = GetItems();
 
 		// Don't respond to key-presses containing "Alt" as a modifier
 		if ( ItemsSourceRef.Num() > 0 && !InKeyEvent.IsAltDown() )
@@ -479,9 +483,9 @@ public:
 
 	virtual FNavigationReply OnNavigation(const FGeometry& MyGeometry, const FNavigationEvent& InNavigationEvent) override
 	{
-		if (this->ItemsSource && this->bHandleDirectionalNavigation && (this->bHandleGamepadEvents || InNavigationEvent.GetNavigationGenesis() != ENavigationGenesis::Controller))
+		if (this->HasValidItemsSource() && this->bHandleDirectionalNavigation && (this->bHandleGamepadEvents || InNavigationEvent.GetNavigationGenesis() != ENavigationGenesis::Controller))
 		{
-			const TArray<ItemType>& ItemsSourceRef = (*this->ItemsSource);
+			const TArrayView<const ItemType> ItemsSourceRef = this->GetItems();
 
 			const int32 NumItemsPerLine = GetNumItemsPerLine();
 			const int32 CurSelectionIndex = (!TListTypeTraits<ItemType>::IsPtrValid(SelectorItem)) ? -1 : ItemsSourceRef.Find(TListTypeTraits<ItemType>::NullableItemTypeConvertToItemType(SelectorItem));
@@ -946,7 +950,7 @@ public:
 			return;
 		}
 
-		const TArray<ItemType>& ItemsSourceRef = (*ItemsSource);
+		const TArrayView<const ItemType> ItemsSourceRef = GetItems();
 
 		int32 RangeStartIndex = 0;
 		if( TListTypeTraits<ItemType>::IsPtrValid(RangeSelectionStart) )
@@ -1296,7 +1300,7 @@ public:
 		if ( TableViewMode != ETableViewMode::Tree )
 		{
 			bool bSelectionChanged = false;
-			if ( ItemsSource == nullptr )
+			if ( !HasValidItemsSource() )
 			{
 				// We are no longer observing items so there is no more selection.
 				this->Private_ClearSelection();
@@ -1307,9 +1311,10 @@ public:
 				// We are observing some items; they are potentially different.
 				// Unselect any that are no longer being observed.
 				TItemSet NewSelectedItems;
-				for ( int32 ItemIndex = 0; ItemIndex < ItemsSource->Num(); ++ItemIndex )
+				const TArrayView<const ItemType> Items = GetItems();
+				for ( int32 ItemIndex = 0; ItemIndex < Items.Num(); ++ItemIndex )
 				{
-					ItemType CurItem = (*ItemsSource)[ItemIndex];
+					ItemType CurItem = Items[ItemIndex];
 					const bool bItemIsSelected = ( nullptr != SelectedItems.Find( CurItem ) );
 					if ( bItemIsSelected )
 					{
@@ -1349,7 +1354,8 @@ public:
 		// Ensure that we always begin and clean up a generation pass.
 		FGenerationPassGuard GenerationPassGuard(WidgetGenerator);
 
-		if (ItemsSource && ItemsSource->Num() > 0)
+		const TArrayView<const ItemType> Items = GetItems();
+		if (Items.Num() > 0)
 		{
 			// Items in view, including fractional items
 			float ItemsInView = 0.0f;
@@ -1362,7 +1368,7 @@ public:
 
 			// Index of the item at which we start generating based on how far scrolled down we are
 			// Note that we must generate at LEAST one item.
-			int32 StartIndex = FMath::Clamp( (int32)(FMath::FloorToDouble(CurrentScrollOffset)), 0, ItemsSource->Num() - 1 );
+			int32 StartIndex = FMath::Clamp( (int32)(FMath::FloorToDouble(CurrentScrollOffset)), 0, Items.Num() - 1 );
 
 			// Length of the first item that is generated. This item is at the location where the user requested we scroll
 			float FirstItemLength = 0.0f;
@@ -1374,9 +1380,9 @@ public:
 			const float LayoutScaleMultiplier = MyGeometry.GetAccumulatedLayoutTransform().GetScale();
 			FTableViewDimensions MyDimensions(this->Orientation, MyGeometry.GetLocalSize());
 			
-			for( int32 ItemIndex = StartIndex; !bHasFilledAvailableArea && ItemIndex < ItemsSource->Num(); ++ItemIndex )
+			for( int32 ItemIndex = StartIndex; !bHasFilledAvailableArea && ItemIndex < Items.Num(); ++ItemIndex )
 			{
-				const ItemType& CurItem = (*ItemsSource)[ItemIndex];
+				const ItemType& CurItem = Items[ItemIndex];
 
 				if (!TListTypeTraits<ItemType>::IsPtrValid(CurItem))
 				{
@@ -1424,7 +1430,7 @@ public:
 					? ItemLength * ItemsInView	// For the first item, ItemsInView <= 1.0f
 					: ItemLength;
 
-				bAtEndOfList = ItemIndex >= ItemsSource->Num() - 1;
+				bAtEndOfList = ItemIndex >= Items.Num() - 1;
 
 				if (bIsFirstItem && ViewLengthUsedSoFar >= MyDimensions.ScrollAxis)
 				{
@@ -1449,7 +1455,7 @@ public:
 
 				for (int32 ItemIndex = StartIndex - 1; LengthGeneratedSoFar < MyDimensions.ScrollAxis && ItemIndex >= 0; --ItemIndex)
 				{
-					const ItemType& CurItem = (*ItemsSource)[ItemIndex];
+					const ItemType& CurItem = Items[ItemIndex];
 					if (TListTypeTraits<ItemType>::IsPtrValid(CurItem))
 					{
 						const float ItemLength = GenerateWidgetForItem(CurItem, ItemIndex, StartIndex, LayoutScaleMultiplier);
@@ -1466,7 +1472,7 @@ public:
 					}
 				}
 
-				return FReGenerateResults(NewScrollOffsetForBackfill, LengthGeneratedSoFar, ItemsSource->Num() - NewScrollOffsetForBackfill, true);
+				return FReGenerateResults(NewScrollOffsetForBackfill, LengthGeneratedSoFar, Items.Num() - NewScrollOffsetForBackfill, true);
 			}
 
 			return FReGenerateResults(CurrentScrollOffset, LengthGeneratedSoFar, ItemsInView, false);
@@ -1548,7 +1554,7 @@ public:
 		}
 		
 
-		const TArray<ItemType>& ItemsSourceRef = (*this->ItemsSource);
+		const TArrayView<const ItemType> ItemsSourceRef = this->GetItems();
 
 		for (int32 ItemIndex = 0; ItemIndex < InItems.Num(); ++ItemIndex)
 		{
@@ -1598,7 +1604,7 @@ public:
 	/** @return how many items there are in the TArray being observed */
 	virtual int32 GetNumItemsBeingObserved() const override
 	{
-		return ItemsSource == nullptr ? 0 : ItemsSource->Num();
+		return GetItems().Num();
 	}
 
 	virtual TSharedRef<ITableRow> GenerateNewPinnedWidget(ItemType InItem, const int32 ItemIndex, const int32 NumPinnedItems)
@@ -1654,16 +1660,45 @@ public:
 	 * Establishes a wholly new list of items being observed by the list.
 	 * Wipes all existing state and requests and will fully rebuild on the next tick.
 	 */
+	void SetItemsSource(const TArray<ItemType>* InListItemsSource)
+	{
+		PRAGMA_DISABLE_DEPRECATION_WARNINGS
+		if (ItemsSource != InListItemsSource)
+		{
+			if (IsConstructed())
+			{
+				Private_ClearSelection();
+				CancelScrollIntoView();
+				ClearWidgets();
+				RebuildList();
+			}
+			ItemsSource = InListItemsSource;
+		}
+		PRAGMA_ENABLE_DEPRECATION_WARNINGS
+	}
+
+	UE_DEPRECATED(5.2, "SetListItemsSource is deprecated. Please use the correct SetItemsSource implementation.")
 	void SetListItemsSource(const TArray<ItemType>& InListItemsSource)
 	{
-		if (ItemsSource != &InListItemsSource)
+		SetItemsSource(&InListItemsSource);
+	}
+
+	bool HasValidItemsSource() const
+	{
+		PRAGMA_DISABLE_DEPRECATION_WARNINGS
+		return ItemsSource != nullptr;
+		PRAGMA_ENABLE_DEPRECATION_WARNINGS
+	}
+
+	TArrayView<const ItemType> GetItems() const
+	{
+		PRAGMA_DISABLE_DEPRECATION_WARNINGS
+		if (ItemsSource)
 		{
-			Private_ClearSelection();
-			CancelScrollIntoView();
-			ClearWidgets();
-			RebuildList();
-			ItemsSource = &InListItemsSource;
+			return *ItemsSource;
 		}
+		return TArrayView<const ItemType>();
+		PRAGMA_ENABLE_DEPRECATION_WARNINGS
 	}
 
 	/**
@@ -2017,9 +2052,10 @@ protected:
 	 */
 	virtual EScrollIntoViewResult ScrollIntoView( const FGeometry& ListViewGeometry ) override
 	{
-		if (ItemsSource && TListTypeTraits<ItemType>::IsPtrValid(ItemToScrollIntoView))
+		if (HasValidItemsSource() && TListTypeTraits<ItemType>::IsPtrValid(ItemToScrollIntoView))
 		{
-			const int32 IndexOfItem = ItemsSource->Find( TListTypeTraits<ItemType>::NullableItemTypeConvertToItemType( ItemToScrollIntoView ) );
+			const TArrayView<const ItemType> Items = GetItems();
+			const int32 IndexOfItem = Items.Find( TListTypeTraits<ItemType>::NullableItemTypeConvertToItemType( ItemToScrollIntoView ) );
 			if (IndexOfItem != INDEX_NONE)
 			{
 				double NumLiveWidgets = GetNumLiveWidgets();
@@ -2053,14 +2089,14 @@ protected:
 					NewScrollOffset -= (NumLiveWidgets / 2.0);
 
 					// Limit offset to top and bottom of the list.
-					const double MaxScrollOffset = FMath::Max(0.0, static_cast<double>(ItemsSource->Num()) - NumLiveWidgets);
+					const double MaxScrollOffset = FMath::Max(0.0, static_cast<double>(Items.Num()) - NumLiveWidgets);
 					NewScrollOffset = FMath::Clamp<double>(NewScrollOffset, 0.0, MaxScrollOffset);
 
 					SetScrollOffset((float)NewScrollOffset);
 				}
 				else if (bNavigateOnScrollIntoView)
 				{
-					if (TSharedPtr<ITableRow> TableRow = WidgetFromItem((*ItemsSource)[IndexOfItem]))
+					if (TSharedPtr<ITableRow> TableRow = WidgetFromItem(Items[IndexOfItem]))
 					{
 						const FGeometry& WidgetGeometry = TableRow->AsWidget()->GetCachedGeometry();
 						const FTableViewDimensions WidgetTopLeft(this->Orientation, WidgetGeometry.GetAbsolutePositionAtCoordinates(FVector2D::ZeroVector));
@@ -2208,13 +2244,14 @@ protected:
 			//           Scroll "one widget's length" at a time until we've scrolled as far as the user asked us to.
 			//           Generate widgets on demand so we can figure out how big they are.
 
-			if (ItemsSource && ItemsSource->Num() > 0)
+			const TArrayView<const ItemType> Items = GetItems();
+			if (Items.Num() > 0)
 			{
 				int32 ItemIndex = StartingItemIndex;
 				const float LayoutScaleMultiplier = MyGeometry.GetAccumulatedLayoutTransform().GetScale();
-				while( AbsScrollByAmount != 0 && ItemIndex < ItemsSource->Num() && ItemIndex >= 0 )
+				while( AbsScrollByAmount != 0 && ItemIndex < Items.Num() && ItemIndex >= 0 )
 				{
-					const ItemType& CurItem = (*ItemsSource)[ItemIndex];
+					const ItemType& CurItem = Items[ItemIndex];
 					if (!TListTypeTraits<ItemType>::IsPtrValid(CurItem))
 					{
 						// If the CurItem is not valid, we do not generate a new widget for it, we skip it.
@@ -2246,7 +2283,7 @@ protected:
 
 						if (AbsScrollByAmount > RemainingDistance)
 						{
-							if (ItemIndex != ItemsSource->Num())
+							if (ItemIndex != Items.Num())
 							{
 								AbsScrollByAmount -= RemainingDistance;
 								NewScrollOffset = 1.0 + (int32)NewScrollOffset;
@@ -2254,7 +2291,7 @@ protected:
 							}
 							else
 							{
-								NewScrollOffset = ItemsSource->Num();
+								NewScrollOffset = Items.Num();
 								break;
 							}
 						} 
@@ -2324,12 +2361,13 @@ protected:
 		if (OnIsSelectableOrNavigable.IsBound())
 		{
 			// Walk through the list until we either find a navigable item or run out of entries.
+			const TArrayView<const ItemType> Items = GetItems();
 			while (!OnIsSelectableOrNavigable.Execute(ItemToSelect))
 			{
 				SelectionIdx += (bSelectForward ? 1 : -1);
-				if (ItemsSource->IsValidIndex(SelectionIdx))
+				if (Items.IsValidIndex(SelectionIdx))
 				{
-					ItemToSelect = (*ItemsSource)[SelectionIdx];
+					ItemToSelect = Items[SelectionIdx];
 				}
 				else
 				{
@@ -2350,7 +2388,8 @@ protected:
 		{
 			if (!OnIsSelectableOrNavigable.Execute(ItemToSelect))
 			{
-				int32 NewSelectionIdx = ItemsSource->Find(ItemToSelect);
+				const TArrayView<const ItemType> Items = GetItems();
+				int32 NewSelectionIdx = Items.Find(ItemToSelect);
 
 				// By default, we walk forward
 				bool bSelectNextItem = true;
@@ -2361,7 +2400,7 @@ protected:
 					if (TListTypeTraits<ItemType>::IsPtrValid(LastSelectedItem))
 					{
 						ItemType NonNullLastSelectedItem = TListTypeTraits<ItemType>::NullableItemTypeConvertToItemType(LastSelectedItem);
-						const int32 LastSelectedItemIdx = ItemsSource->Find(NonNullLastSelectedItem);
+						const int32 LastSelectedItemIdx = Items.Find(NonNullLastSelectedItem);
 
 						bSelectNextItem = LastSelectedItemIdx < NewSelectionIdx;
 					}
@@ -2371,9 +2410,9 @@ protected:
 				do
 				{
 					NewSelectionIdx += (bSelectNextItem ? 1 : -1);
-					if (ItemsSource->IsValidIndex(NewSelectionIdx))
+					if (Items.IsValidIndex(NewSelectionIdx))
 					{
-						ItemToSelect = (*ItemsSource)[NewSelectionIdx];
+						ItemToSelect = Items[NewSelectionIdx];
 					}
 					else
 					{
@@ -2478,6 +2517,7 @@ protected:
 	/** A set of which items should be highlighted */
 	TItemSet HighlightedItems;
 
+	UE_DEPRECATED(5.2, "Protected access to ItemsSource is deprecated. Please use GetItems, SetItemsSoure or HasValidItemsSource.")
 	/** Pointer to the array of data items that we are observing */
 	const TArray<ItemType>* ItemsSource;
 
diff --git a/Engine/Source/Runtime/Slate/Public/Widgets/Views/STileView.h b/Engine/Source/Runtime/Slate/Public/Widgets/Views/STileView.h
--- a/Engine/Source/Runtime/Slate/Public/Widgets/Views/STileView.h
+++ b/Engine/Source/Runtime/Slate/Public/Widgets/Views/STileView.h
@@ -149,7 +149,7 @@ public:
 		this->OnRowReleased = InArgs._OnTileReleased;
 		this->OnItemScrolledIntoView = InArgs._OnItemScrolledIntoView;
 		
-		this->ItemsSource = InArgs._ListItemsSource;
+		this->SetItemsSource(InArgs._ListItemsSource);
 		this->OnContextMenuOpening = InArgs._OnContextMenuOpening;
 		this->OnClick = InArgs._OnMouseButtonClick;
 		this->OnDoubleClick = InArgs._OnMouseButtonDoubleClick;
@@ -161,7 +161,6 @@ public:
 
 		this->AllowOverscroll = InArgs._AllowOverscroll;
 		this->ConsumeMouseWheel = InArgs._ConsumeMouseWheel;
-
 		this->WheelScrollMultiplier = InArgs._WheelScrollMultiplier;
 
 		this->bHandleGamepadEvents = InArgs._HandleGamepadEvents;
@@ -186,7 +185,7 @@ public:
 				ErrorString += TEXT("Please specify an OnGenerateTile. \n");
 			}
 
-			if ( this->ItemsSource == nullptr )
+			if ( !this->HasValidItemsSource() )
 			{
 				ErrorString += TEXT("Please specify a ListItemsSource. \n");
 			}
@@ -196,12 +195,12 @@ public:
 		{
 			// Let the coder know what they forgot
 			this->ChildSlot
-			.HAlign(HAlign_Center)
-			.VAlign(VAlign_Center)
-			[
-				SNew(STextBlock)
-				.Text(FText::FromString(ErrorString))
-			];
+				.HAlign(HAlign_Center)
+				.VAlign(VAlign_Center)
+				[
+					SNew(STextBlock)
+					.Text(FText::FromString(ErrorString))
+				];
 		}
 		else
 		{
@@ -227,9 +226,9 @@ public:
 
 	virtual FNavigationReply OnNavigation(const FGeometry& MyGeometry, const FNavigationEvent& InNavigationEvent) override
 	{
-		if (this->ItemsSource && this->bHandleDirectionalNavigation && (this->bHandleGamepadEvents || InNavigationEvent.GetNavigationGenesis() != ENavigationGenesis::Controller))
+		if (this->HasValidItemsSource() && this->bHandleDirectionalNavigation && (this->bHandleGamepadEvents || InNavigationEvent.GetNavigationGenesis() != ENavigationGenesis::Controller))
 		{
-			const TArray<ItemType>& ItemsSourceRef = (*this->ItemsSource);
+			const TArrayView<const ItemType>& ItemsSourceRef = this->GetItems();
 
 			const int32 NumItemsPerLine = GetNumItemsPerLine();
 			const int32 CurSelectionIndex = (!TListTypeTraits<ItemType>::IsPtrValid(this->SelectorItem)) ? 0 : ItemsSourceRef.Find(TListTypeTraits<ItemType>::NullableItemTypeConvertToItemType(this->SelectorItem));
@@ -271,14 +270,14 @@ public:
 		// Clear all the items from our panel. We will re-add them in the correct order momentarily.
 		this->ClearWidgets();
 		
-		const TArray<ItemType>* SourceItems = this->ItemsSource;
-		if (SourceItems && SourceItems->Num() > 0)
+		const TArrayView<const ItemType> tems = this->GetItems();
+		if (tems.Num() > 0)
 		{
 			// Item width and height is constant by design.
 			FTableViewDimensions TileDimensions = GetTileDimensions();
 			FTableViewDimensions AllottedDimensions(this->Orientation, MyGeometry.GetLocalSize());
 
-			const int32 NumItems = SourceItems->Num();
+			const int32 NumItems = tems.Num();
 			const int32 NumItemsPerLine = GetNumItemsPerLine();
 			const int32 NumItemsPaddedToFillLastLine = (NumItems % NumItemsPerLine != 0)
 				? NumItems + NumItemsPerLine - NumItems % NumItemsPerLine
@@ -306,7 +305,7 @@ public:
 			double NumLinesShownOnScreen = 0;
 			for( int32 ItemIndex = StartIndex; !bHasFilledAvailableArea && ItemIndex < NumItems; ++ItemIndex )
 			{
-				const ItemType& CurItem = (*SourceItems)[ItemIndex];
+				const ItemType& CurItem = tems[ItemIndex];
 
 				if (bNewLine)
 				{
@@ -366,7 +365,7 @@ public:
 
 	virtual int32 GetNumItemsBeingObserved() const override
 	{
-		const int32 NumItemsBeingObserved = this->ItemsSource == nullptr ? 0 : this->ItemsSource->Num();
+		const int32 NumItemsBeingObserved = this->GetItems().Num();
 		const int32 NumItemsPerLine = GetNumItemsPerLine();
 		
 		int32 NumEmptySpacesAtEnd = 0;
@@ -428,9 +427,9 @@ protected:
 	 */
 	virtual typename SListView<ItemType>::EScrollIntoViewResult ScrollIntoView(const FGeometry& ListViewGeometry) override
 	{
-		if (TListTypeTraits<ItemType>::IsPtrValid(this->ItemToScrollIntoView) && this->ItemsSource != nullptr)
+		if (TListTypeTraits<ItemType>::IsPtrValid(this->ItemToScrollIntoView) && this->HasValidItemsSource())
 		{
-			const int32 IndexOfItem = this->ItemsSource->Find(TListTypeTraits<ItemType>::NullableItemTypeConvertToItemType(this->ItemToScrollIntoView));
+			const int32 IndexOfItem = this->GetItems().Find(TListTypeTraits<ItemType>::NullableItemTypeConvertToItemType(this->ItemToScrollIntoView));
 			if (IndexOfItem != INDEX_NONE)
 			{
 				const float NumLinesInView = FTableViewDimensions(this->Orientation, ListViewGeometry.GetLocalSize()).ScrollAxis / GetTileDimensions().ScrollAxis;
diff --git a/Engine/Source/Runtime/Slate/Public/Widgets/Views/STreeView.h b/Engine/Source/Runtime/Slate/Public/Widgets/Views/STreeView.h
--- a/Engine/Source/Runtime/Slate/Public/Widgets/Views/STreeView.h
+++ b/Engine/Source/Runtime/Slate/Public/Widgets/Views/STreeView.h
@@ -330,7 +330,7 @@ public:
 		, TreeItemsSource( nullptr )
 		, bTreeItemsAreDirty( true )
 	{
-		this->ItemsSource = &LinearizedItems;
+		this->SetItemsSource(&LinearizedItems);
 	}
 
 public:
