Fix deprecation spelling error
#rnx

diff --git a/Engine/Source/Runtime/Slate/Public/Widgets/Views/SListView.h b/Engine/Source/Runtime/Slate/Public/Widgets/Views/SListView.h
--- a/Engine/Source/Runtime/Slate/Public/Widgets/Views/SListView.h
+++ b/Engine/Source/Runtime/Slate/Public/Widgets/Views/SListView.h
@@ -2517,7 +2517,7 @@ protected:
 	/** A set of which items should be highlighted */
 	TItemSet HighlightedItems;
 
-	UE_DEPRECATED(5.2, "Protected access to ItemsSource is deprecated. Please use GetItems, SetItemsSoure or HasValidItemsSource.")
+	UE_DEPRECATED(5.2, "Protected access to ItemsSource is deprecated. Please use GetItems, SetItemsSource or HasValidItemsSource.")
 	/** Pointer to the array of data items that we are observing */
 	const TArray<ItemType>* ItemsSource;
 
