Fix to Software Cursor to render using DPI settings
Make the hand cursor bigger.

diff --git a/Engine/Plugins/2D/Paper2D/Source/Paper2DEditor/Private/SPaperEditorViewport.cpp b/Engine/Plugins/2D/Paper2D/Source/Paper2DEditor/Private/SPaperEditorViewport.cpp
--- a/Engine/Plugins/2D/Paper2D/Source/Paper2DEditor/Private/SPaperEditorViewport.cpp
+++ b/Engine/Plugins/2D/Paper2D/Source/Paper2DEditor/Private/SPaperEditorViewport.cpp
@@ -515,11 +515,12 @@ void SPaperEditorViewport::PaintSoftwareCursor(const FGeometry& AllottedGeometry
 	if (bShowSoftwareCursor)
 	{
 		const FSlateBrush* Brush = FAppStyle::GetBrush(TEXT("SoftwareCursor_Grab"));
+		const FVector2D CursorSize = Brush->ImageSize / AllottedGeometry.Scale;
 
 		FSlateDrawElement::MakeBox(
 			OutDrawElements,
 			DrawLayerId,
-			AllottedGeometry.ToPaintGeometry( GraphCoordToPanelCoord(SoftwareCursorPosition) / AllottedGeometry.Scale - ( Brush->ImageSize / 2 ), Brush->ImageSize ),
+			AllottedGeometry.ToPaintGeometry( GraphCoordToPanelCoord(SoftwareCursorPosition) / AllottedGeometry.Scale - ( CursorSize / 2 ), CursorSize ),
 			Brush);
 	}
 }
diff --git a/Engine/Source/Editor/EditorStyle/Private/SlateEditorStyle.cpp b/Engine/Source/Editor/EditorStyle/Private/SlateEditorStyle.cpp
--- a/Engine/Source/Editor/EditorStyle/Private/SlateEditorStyle.cpp
+++ b/Engine/Source/Editor/EditorStyle/Private/SlateEditorStyle.cpp
@@ -783,7 +783,7 @@ void FSlateEditorStyle::FStyle::SetupGeneralStyles()
 	Set( "GenericPause", new IMAGE_BRUSH( "Icons/generic_pause_16x", Icon16x16 ) );
 	Set( "GenericStop", new IMAGE_BRUSH( "Icons/generic_stop_16x", Icon16x16 ) );
 
-	Set( "SoftwareCursor_Grab", new IMAGE_BRUSH( "Icons/cursor_grab", Icon16x16 ) );
+	Set( "SoftwareCursor_Grab", new IMAGE_BRUSH( "Icons/cursor_grab", Icon24x24 ) );
 	Set( "SoftwareCursor_CardinalCross", new IMAGE_BRUSH( "Icons/cursor_cardinal_cross", Icon24x24 ) );
 	Set( "SoftwareCursor_UpDown", new IMAGE_BRUSH( "Icons/cursor_updown", Icon16x20 ) );
 	
diff --git a/Engine/Source/Editor/EditorStyle/Private/StarshipStyle.cpp b/Engine/Source/Editor/EditorStyle/Private/StarshipStyle.cpp
--- a/Engine/Source/Editor/EditorStyle/Private/StarshipStyle.cpp
+++ b/Engine/Source/Editor/EditorStyle/Private/StarshipStyle.cpp
@@ -675,7 +675,7 @@ void FStarshipEditorStyle::FStyle::SetupGeneralStyles()
 	Set( "GenericPause", new IMAGE_BRUSH( "Icons/generic_pause_16x", Icon16x16 ) );
 	Set( "GenericStop", new IMAGE_BRUSH( "Icons/generic_stop_16x", Icon16x16 ) );
 
-	Set( "SoftwareCursor_Grab", new IMAGE_BRUSH( "Icons/cursor_grab", Icon16x16 ) );
+	Set( "SoftwareCursor_Grab", new IMAGE_BRUSH( "Icons/cursor_grab", Icon24x24 ) );
 	Set( "SoftwareCursor_CardinalCross", new IMAGE_BRUSH( "Icons/cursor_cardinal_cross", Icon24x24 ) );
 	Set( "SoftwareCursor_UpDown", new IMAGE_BRUSH( "Icons/cursor_updown", Icon16x20 ) );
 
diff --git a/Engine/Source/Editor/WorldPartitionEditor/Private/WorldPartition/SWorldPartitionEditorGrid2D.cpp b/Engine/Source/Editor/WorldPartitionEditor/Private/WorldPartition/SWorldPartitionEditorGrid2D.cpp
--- a/Engine/Source/Editor/WorldPartitionEditor/Private/WorldPartition/SWorldPartitionEditorGrid2D.cpp
+++ b/Engine/Source/Editor/WorldPartitionEditor/Private/WorldPartition/SWorldPartitionEditorGrid2D.cpp
@@ -1740,11 +1740,12 @@ int32 SWorldPartitionEditorGrid2D::PaintSoftwareCursor(const FGeometry& Allotted
 	if (bIsPanning)
 	{
 		const FSlateBrush* Brush = FAppStyle::GetBrush(TEXT("SoftwareCursor_Grab"));
+		const FVector2D CursorSize = Brush->ImageSize / AllottedGeometry.Scale;
 
 		FSlateDrawElement::MakeBox(
 			OutDrawElements,
 			LayerId,
-			AllottedGeometry.ToPaintGeometry(MouseCursorPos - (Brush->ImageSize * 0.5f), Brush->ImageSize),
+			AllottedGeometry.ToPaintGeometry(MouseCursorPos - (CursorSize * 0.5f), CursorSize),
 			Brush
 		);
 	}
diff --git a/Engine/Source/Runtime/Slate/Private/Widgets/Layout/SScrollBox.cpp b/Engine/Source/Runtime/Slate/Private/Widgets/Layout/SScrollBox.cpp
--- a/Engine/Source/Runtime/Slate/Private/Widgets/Layout/SScrollBox.cpp
+++ b/Engine/Source/Runtime/Slate/Private/Widgets/Layout/SScrollBox.cpp
@@ -1078,11 +1078,12 @@ int32 SScrollBox::OnPaint( const FPaintArgs& Args, const FGeometry& AllottedGeom
 	}
 
 	const FSlateBrush* Brush = FCoreStyle::Get().GetBrush(TEXT("SoftwareCursor_Grab"));
+	const FVector2D CursorSize = Brush->ImageSize / AllottedGeometry.Scale;
 
 	FSlateDrawElement::MakeBox(
 		OutDrawElements,
 		++NewLayerId,
-		AllottedGeometry.ToPaintGeometry( SoftwareCursorPosition - ( Brush->ImageSize / 2 ), Brush->ImageSize ),
+		AllottedGeometry.ToPaintGeometry( SoftwareCursorPosition - (CursorSize / 2 ), CursorSize),
 		Brush
 	);
 
diff --git a/Engine/Source/Runtime/Slate/Private/Widgets/SViewport.cpp b/Engine/Source/Runtime/Slate/Private/Widgets/SViewport.cpp
--- a/Engine/Source/Runtime/Slate/Private/Widgets/SViewport.cpp
+++ b/Engine/Source/Runtime/Slate/Private/Widgets/SViewport.cpp
@@ -191,6 +191,7 @@ int32 SViewport::OnPaint( const FPaintArgs& Args, const FGeometry& AllottedGeome
 		{
 			Brush = FCoreStyle::Get().GetBrush(TEXT("SoftwareCursor_CardinalCross"));
 		}
+		const FVector2D CursorSize = Brush->ImageSize / AllottedGeometry.Scale;
 
 		FVector2D CursorPositionLocalSpace = ViewportInterfacePin->GetSoftwareCursorPosition() / AllottedGeometry.Scale;
 
@@ -198,7 +199,7 @@ int32 SViewport::OnPaint( const FPaintArgs& Args, const FGeometry& AllottedGeome
 		FSlateDrawElement::MakeBox(
 			OutDrawElements,
 			LayerId,
-			AllottedGeometry.ToPaintGeometry( CursorPositionLocalSpace - ( Brush->ImageSize / 2 ), Brush->ImageSize ),
+			AllottedGeometry.ToPaintGeometry( CursorPositionLocalSpace - (CursorSize / 2 ), CursorSize),
 			Brush
 		);
 	}
diff --git a/Engine/Source/Runtime/Slate/Private/Widgets/Text/SMultiLineEditableText.cpp b/Engine/Source/Runtime/Slate/Private/Widgets/Text/SMultiLineEditableText.cpp
--- a/Engine/Source/Runtime/Slate/Private/Widgets/Text/SMultiLineEditableText.cpp
+++ b/Engine/Source/Runtime/Slate/Private/Widgets/Text/SMultiLineEditableText.cpp
@@ -635,11 +635,12 @@ int32 SMultiLineEditableText::OnPaint( const FPaintArgs& Args, const FGeometry&
 	if (bIsSoftwareCursor)
 	{
 		const FSlateBrush* Brush = FCoreStyle::Get().GetBrush(TEXT("SoftwareCursor_Grab"));
+		const FVector2D CursorSize = Brush->ImageSize / AllottedGeometry.Scale;
 
 		FSlateDrawElement::MakeBox(
 			OutDrawElements,
 			++LayerId,
-			AllottedGeometry.ToPaintGeometry(SoftwareCursorPosition - (Brush->ImageSize / 2), Brush->ImageSize),
+			AllottedGeometry.ToPaintGeometry(SoftwareCursorPosition - (CursorSize / 2), CursorSize),
 			Brush
 			);
 	}
diff --git a/Engine/Source/Runtime/Slate/Private/Widgets/Views/STableViewBase.cpp b/Engine/Source/Runtime/Slate/Private/Widgets/Views/STableViewBase.cpp
--- a/Engine/Source/Runtime/Slate/Private/Widgets/Views/STableViewBase.cpp
+++ b/Engine/Source/Runtime/Slate/Private/Widgets/Views/STableViewBase.cpp
@@ -763,11 +763,12 @@ int32 STableViewBase::OnPaint( const FPaintArgs& Args, const FGeometry& Allotted
 	}
 
 	const FSlateBrush* Brush = FCoreStyle::Get().GetBrush(TEXT("SoftwareCursor_Grab"));
+	const FVector2D CursorSize = Brush->ImageSize / AllottedGeometry.Scale;
 
 	FSlateDrawElement::MakeBox(
 		OutDrawElements,
 		++NewLayerId,
-		AllottedGeometry.ToPaintGeometry( SoftwareCursorPosition - ( Brush->ImageSize / 2 ), Brush->ImageSize ),
+		AllottedGeometry.ToPaintGeometry( SoftwareCursorPosition - (CursorSize / 2 ), CursorSize),
 		Brush
 		);
 
diff --git a/Engine/Source/Runtime/SlateCore/Private/Styling/CoreStyle.cpp b/Engine/Source/Runtime/SlateCore/Private/Styling/CoreStyle.cpp
--- a/Engine/Source/Runtime/SlateCore/Private/Styling/CoreStyle.cpp
+++ b/Engine/Source/Runtime/SlateCore/Private/Styling/CoreStyle.cpp
@@ -301,7 +301,7 @@ TSharedRef<ISlateStyle> FCoreStyle::Create( const FName& InStyleSetName )
 
 	// Cursor icons
 	{
-		Style->Set("SoftwareCursor_Grab", new IMAGE_BRUSH( "Icons/cursor_grab", Icon16x16));
+		Style->Set("SoftwareCursor_Grab", new IMAGE_BRUSH( "Icons/cursor_grab", Icon24x24));
 		Style->Set("SoftwareCursor_CardinalCross", new IMAGE_BRUSH( "Icons/cursor_cardinal_cross", Icon24x24));
 	}
 
diff --git a/Engine/Source/Runtime/SlateCore/Private/Styling/StarshipCoreStyle.cpp b/Engine/Source/Runtime/SlateCore/Private/Styling/StarshipCoreStyle.cpp
--- a/Engine/Source/Runtime/SlateCore/Private/Styling/StarshipCoreStyle.cpp
+++ b/Engine/Source/Runtime/SlateCore/Private/Styling/StarshipCoreStyle.cpp
@@ -271,7 +271,7 @@ TSharedRef<ISlateStyle> FStarshipCoreStyle::Create()
 
 	// Cursor icons
 	{
-		Style->Set("SoftwareCursor_Grab", new IMAGE_BRUSH("Icons/cursor_grab", Icon16x16));
+		Style->Set("SoftwareCursor_Grab", new IMAGE_BRUSH("Icons/cursor_grab", Icon24x24));
 		Style->Set("SoftwareCursor_CardinalCross", new IMAGE_BRUSH("Icons/cursor_cardinal_cross", Icon24x24));
 	}
 
diff --git a/Engine/Source/Runtime/SlateCore/Private/Styling/UMGCoreStyle.cpp b/Engine/Source/Runtime/SlateCore/Private/Styling/UMGCoreStyle.cpp
--- a/Engine/Source/Runtime/SlateCore/Private/Styling/UMGCoreStyle.cpp
+++ b/Engine/Source/Runtime/SlateCore/Private/Styling/UMGCoreStyle.cpp
@@ -324,7 +324,7 @@ TSharedRef<ISlateStyle> FUMGCoreStyle::Create()
 
 	// Cursor Icons
 	{
-		Style->Set("SoftwareCursor_Grab", new IMAGE_BRUSH( "Icons/cursor_grab", Icon16x16));
+		Style->Set("SoftwareCursor_Grab", new IMAGE_BRUSH( "Icons/cursor_grab", Icon24x24));
 		Style->Set("SoftwareCursor_CardinalCross", new IMAGE_BRUSH( "Icons/cursor_cardinal_cross", Icon24x24));
 	}
 
