Slate: Initial phase of FVector2D deprecation

diff --git a/Engine/Source/Editor/SequencerCore/Private/MVVM/Views/KeyRenderer.cpp b/Engine/Source/Editor/SequencerCore/Private/MVVM/Views/KeyRenderer.cpp
--- a/Engine/Source/Editor/SequencerCore/Private/MVVM/Views/KeyRenderer.cpp
+++ b/Engine/Source/Editor/SequencerCore/Private/MVVM/Views/KeyRenderer.cpp
@@ -790,11 +790,11 @@ int32 FKeyRenderer::Draw(const FKeyBatchParameters& Params, const FGeometry& All
 			bSelected ? LayerId + 1 : LayerId,
 			// Center the key along Y.  Ensure the middle of the key is at the actual key time
 			AllottedGeometry.ToPaintGeometry(
-				FVector2D(
+				KeySize,
+				FSlateLayoutTransform(FVector2f(
 					KeyPositionPx - FMath::CeilToFloat(KeySize.X / 2.0f),
 					((AllottedGeometry.GetLocalSize().Y / 2.0f) - (KeySize.Y / 2.0f))
-				),
-				KeySize
+				))
 			),
 			KeyDrawParams.BorderBrush,
 			KeyDrawEffects,
@@ -808,12 +808,14 @@ int32 FKeyRenderer::Draw(const FKeyBatchParameters& Params, const FGeometry& All
 			bSelected ? LayerId + 2 : LayerId + 1,
 			// Center the key along Y.  Ensure the middle of the key is at the actual key time
 			AllottedGeometry.ToPaintGeometry(
-				KeyDrawParams.FillOffset + 
-				FVector2D(
-					(KeyPositionPx - FMath::CeilToFloat((KeySize.X / 2.0f) - BrushBorderWidth)),
-					((AllottedGeometry.GetLocalSize().Y / 2.0f) - ((KeySize.Y / 2.0f) - BrushBorderWidth))
-				),
-				KeySize - 2.0f * BrushBorderWidth
+				KeySize - 2.0f * BrushBorderWidth,
+				FSlateLayoutTransform(
+					KeyDrawParams.FillOffset + 
+					FVector2D(
+						(KeyPositionPx - FMath::CeilToFloat((KeySize.X / 2.0f) - BrushBorderWidth)),
+						((AllottedGeometry.GetLocalSize().Y / 2.0f) - ((KeySize.Y / 2.0f) - BrushBorderWidth))
+					)
+				)
 			),
 			KeyDrawParams.FillBrush,
 			KeyDrawEffects,
diff --git a/Engine/Source/Editor/SequencerCore/Private/MVVM/Views/STrackAreaView.cpp b/Engine/Source/Editor/SequencerCore/Private/MVVM/Views/STrackAreaView.cpp
--- a/Engine/Source/Editor/SequencerCore/Private/MVVM/Views/STrackAreaView.cpp
+++ b/Engine/Source/Editor/SequencerCore/Private/MVVM/Views/STrackAreaView.cpp
@@ -281,7 +281,7 @@ int32 STrackAreaView::OnPaint(const FPaintArgs& Args, const FGeometry& AllottedG
 		FSlateDrawElement::MakeBox(
 			OutDrawElements,
 			DashLayer,
-			AllottedGeometry.ToPaintGeometry(FVector2D(DropMinX, TrackPosition), FVector2D(DropMaxX-DropMinX, HorizontalBrush->ImageSize.Y)),
+			AllottedGeometry.ToPaintGeometry(FVector2f(DropMaxX-DropMinX, HorizontalBrush->ImageSize.Y), FSlateLayoutTransform(FVector2f(DropMinX, TrackPosition))),
 			HorizontalBrush,
 			ESlateDrawEffect::None,
 			DashColor.GetSpecifiedColor());
@@ -290,7 +290,7 @@ int32 STrackAreaView::OnPaint(const FPaintArgs& Args, const FGeometry& AllottedG
 		FSlateDrawElement::MakeBox(
 			OutDrawElements,
 			DashLayer,
-			AllottedGeometry.ToPaintGeometry(FVector2D(DropMinX, TrackPosition + (TrackHeight - HorizontalBrush->ImageSize.Y)), FVector2D(DropMaxX-DropMinX, HorizontalBrush->ImageSize.Y)),
+			AllottedGeometry.ToPaintGeometry(FVector2f(DropMaxX-DropMinX, HorizontalBrush->ImageSize.Y), FSlateLayoutTransform(FVector2f(DropMinX, TrackPosition + (TrackHeight - HorizontalBrush->ImageSize.Y)))),
 			HorizontalBrush,
 			ESlateDrawEffect::None,
 			DashColor.GetSpecifiedColor());
@@ -299,7 +299,7 @@ int32 STrackAreaView::OnPaint(const FPaintArgs& Args, const FGeometry& AllottedG
 		FSlateDrawElement::MakeBox(
 			OutDrawElements,
 			DashLayer,
-			AllottedGeometry.ToPaintGeometry(FVector2D(DropMinX, TrackPosition), FVector2D(VerticalBrush->ImageSize.X, TrackHeight)),
+			AllottedGeometry.ToPaintGeometry(FVector2f(VerticalBrush->ImageSize.X, TrackHeight), FSlateLayoutTransform(FVector2f(DropMinX, TrackPosition))),
 			VerticalBrush,
 			ESlateDrawEffect::None,
 			DashColor.GetSpecifiedColor());
@@ -308,7 +308,7 @@ int32 STrackAreaView::OnPaint(const FPaintArgs& Args, const FGeometry& AllottedG
 		FSlateDrawElement::MakeBox(
 			OutDrawElements,
 			DashLayer,
-			AllottedGeometry.ToPaintGeometry(FVector2D(DropMaxX - VerticalBrush->ImageSize.X, TrackPosition), FVector2D(VerticalBrush->ImageSize.X, TrackHeight)),
+			AllottedGeometry.ToPaintGeometry(FVector2f(VerticalBrush->ImageSize.X, TrackHeight), FSlateLayoutTransform(FVector2f(DropMaxX - VerticalBrush->ImageSize.X, TrackPosition))),
 			VerticalBrush,
 			ESlateDrawEffect::None,
 			DashColor.GetSpecifiedColor());
diff --git a/Engine/Source/Editor/SequencerCore/Private/MVVM/Views/STrackLane.cpp b/Engine/Source/Editor/SequencerCore/Private/MVVM/Views/STrackLane.cpp
--- a/Engine/Source/Editor/SequencerCore/Private/MVVM/Views/STrackLane.cpp
+++ b/Engine/Source/Editor/SequencerCore/Private/MVVM/Views/STrackLane.cpp
@@ -272,8 +272,8 @@ int32 STrackLane::PaintLaneBackground(const FGeometry& AllottedGeometry, const F
 			OutDrawElements,
 			PaintLayerId,
 			AllottedGeometry.ToPaintGeometry(
-				FVector2D(0, 0),
-				FVector2D(AllottedGeometry.GetLocalSize().X, TotalNodeHeight)
+				FVector2f(AllottedGeometry.GetLocalSize().X, TotalNodeHeight),
+				FSlateLayoutTransform()
 			),
 			FAppStyle::GetBrush(BorderName),
 			ESlateDrawEffect::None,
@@ -303,8 +303,8 @@ int32 STrackLane::PaintLaneBackground(const FGeometry& AllottedGeometry, const F
 				OutDrawElements,
 				PaintLayerId,
 				AllottedGeometry.ToPaintGeometry(
-					FVector2D(0, 0),
-					FVector2D(AllottedGeometry.GetLocalSize().X, TotalNodeHeight)
+					FVector2f(AllottedGeometry.GetLocalSize().X, TotalNodeHeight),
+					FSlateLayoutTransform()
 				),
 				FAppStyle::GetBrush(BorderName),
 				ESlateDrawEffect::None,
diff --git a/Engine/Source/Editor/SequencerWidgets/Private/STimeRangeSlider.cpp b/Engine/Source/Editor/SequencerWidgets/Private/STimeRangeSlider.cpp
--- a/Engine/Source/Editor/SequencerWidgets/Private/STimeRangeSlider.cpp
+++ b/Engine/Source/Editor/SequencerWidgets/Private/STimeRangeSlider.cpp
@@ -109,7 +109,7 @@ int32 STimeRangeSlider::OnPaint( const FPaintArgs& Args, const FGeometry& Allott
 	FSlateDrawElement::MakeBox( 
 		OutDrawElements,
 		LayerId, 
-		AllottedGeometry.ToPaintGeometry(FVector2D(HandleOffset, 0.0f), FVector2D(RightHandleOffset-LeftHandleOffset-TimeRangeSliderConstants::HandleSize, TimeRangeSliderConstants::HandleSize)),
+		AllottedGeometry.ToPaintGeometry(FVector2f(RightHandleOffset-LeftHandleOffset-TimeRangeSliderConstants::HandleSize, TimeRangeSliderConstants::HandleSize), FSlateLayoutTransform(FVector2f(HandleOffset, 0.0f))),
 		RangeHandle,
 		ESlateDrawEffect::None,
 		(bHandleDragged || bHandleHovered) ? SelectionColor : FLinearColor::Gray);
@@ -118,7 +118,7 @@ int32 STimeRangeSlider::OnPaint( const FPaintArgs& Args, const FGeometry& Allott
 	FSlateDrawElement::MakeBox( 
 		OutDrawElements,
 		LayerId, 
-		AllottedGeometry.ToPaintGeometry(FVector2D(LeftHandleOffset, 0.0f), FVector2D(TimeRangeSliderConstants::HandleSize, TimeRangeSliderConstants::HandleSize)),
+		AllottedGeometry.ToPaintGeometry(FVector2f(TimeRangeSliderConstants::HandleSize, TimeRangeSliderConstants::HandleSize), FSlateLayoutTransform(FVector2f(LeftHandleOffset, 0.0f))),
 		RangeHandleLeft,
 		ESlateDrawEffect::None,
 		(bLeftHandleDragged || bLeftHandleHovered) ? SelectionColor : FLinearColor::Gray);
@@ -127,7 +127,7 @@ int32 STimeRangeSlider::OnPaint( const FPaintArgs& Args, const FGeometry& Allott
 	FSlateDrawElement::MakeBox( 
 		OutDrawElements,
 		LayerId, 
-		AllottedGeometry.ToPaintGeometry(FVector2D(RightHandleOffset, 0.0f), FVector2D(TimeRangeSliderConstants::HandleSize, TimeRangeSliderConstants::HandleSize)),
+		AllottedGeometry.ToPaintGeometry(FVector2f(TimeRangeSliderConstants::HandleSize, TimeRangeSliderConstants::HandleSize), FSlateLayoutTransform(FVector2f(RightHandleOffset, 0.0f))),
 		RangeHandleRight,
 		ESlateDrawEffect::None,
 		(bRightHandleDragged || bRightHandleHovered) ? SelectionColor : FLinearColor::Gray);
@@ -249,9 +249,9 @@ FReply STimeRangeSlider::OnMouseMove( const FGeometry& MyGeometry, const FPointe
 		double RightHandleOffset = 0;
 		ComputeHandleOffsets(LeftHandleOffset, HandleOffset, RightHandleOffset, MyGeometry.GetLocalSize().X);
 		
-		FGeometry LeftHandleRect  = MyGeometry.MakeChild(FVector2D(LeftHandleOffset, 0), FVector2D(TimeRangeSliderConstants::HandleSize, TimeRangeSliderConstants::HandleSize));
-		FGeometry RightHandleRect = MyGeometry.MakeChild(FVector2D(RightHandleOffset, 0), FVector2D(TimeRangeSliderConstants::HandleSize, TimeRangeSliderConstants::HandleSize));
-		FGeometry HandleRect      = MyGeometry.MakeChild(FVector2D(HandleOffset, 0), FVector2D(RightHandleOffset-LeftHandleOffset-TimeRangeSliderConstants::HandleSize, TimeRangeSliderConstants::HandleSize));
+		FGeometry LeftHandleRect  = MyGeometry.MakeChild(FVector2D(TimeRangeSliderConstants::HandleSize, TimeRangeSliderConstants::HandleSize), FSlateLayoutTransform(FVector2D(LeftHandleOffset, 0)));
+		FGeometry RightHandleRect = MyGeometry.MakeChild(FVector2D(TimeRangeSliderConstants::HandleSize, TimeRangeSliderConstants::HandleSize), FSlateLayoutTransform(FVector2D(RightHandleOffset, 0)));
+		FGeometry HandleRect      = MyGeometry.MakeChild(FVector2D(RightHandleOffset-LeftHandleOffset-TimeRangeSliderConstants::HandleSize, TimeRangeSliderConstants::HandleSize), FSlateLayoutTransform(FVector2D(HandleOffset, 0)));
 
 		FVector2D LocalMousePosition = MouseEvent.GetScreenSpacePosition();
 
diff --git a/Engine/Source/Editor/StatusBar/Private/SStatusBar.cpp b/Engine/Source/Editor/StatusBar/Private/SStatusBar.cpp
--- a/Engine/Source/Editor/StatusBar/Private/SStatusBar.cpp
+++ b/Engine/Source/Editor/StatusBar/Private/SStatusBar.cpp
@@ -173,9 +173,8 @@ class SDrawerOverlay : public SCompoundWidget
 
 		if (bIsResizing && this->HasMouseCapture() && !MouseEvent.GetCursorDelta().IsZero())
 		{
-			const FVector2D LocalMousePos = InitialResizeGeometry.AbsoluteToLocal(MouseEvent.GetScreenSpacePosition());
-			const double NarrowingPrecision = 1.0 / 16.0;
-			const float DeltaHeight = FloatCastChecked<float>((InitialResizeGeometry.GetLocalPositionAtCoordinates(FVector2D::ZeroVector) - LocalMousePos).Y, NarrowingPrecision);
+			const FVector2f LocalMousePos = InitialResizeGeometry.AbsoluteToLocal(MouseEvent.GetScreenSpacePosition());
+			const float DeltaHeight = (InitialResizeGeometry.GetLocalPositionAtCoordinates(FVector2f::ZeroVector) - LocalMousePos).Y;
 
 			TargetHeight = FMath::Clamp(InitialHeightAtResize + DeltaHeight, MinHeight, MaxHeight);
 			SetHeight(InitialHeightAtResize + DeltaHeight);
@@ -234,7 +233,7 @@ class SDrawerOverlay : public SCompoundWidget
 		FSlateDrawElement::MakeBox(
 			OutDrawElements,
 			LayerId,
-			RenderTransformedChildGeometry.ToPaintGeometry(ShadowOffset, FVector2D(AllottedGeometry.GetLocalSize().X - (ShadowOffset.X * 2), TargetHeight)),
+			RenderTransformedChildGeometry.ToPaintGeometry(FVector2f(AllottedGeometry.GetLocalSize().X - (ShadowOffset.X * 2), TargetHeight), FSlateLayoutTransform(ShadowOffset)),
 			BackgroundBrush,
 			ESlateDrawEffect::None,
 			BackgroundBrush->GetTint(InWidgetStyle));
@@ -245,7 +244,7 @@ class SDrawerOverlay : public SCompoundWidget
 		FSlateDrawElement::MakeBox(
 			OutDrawElements,
 			OutLayerId,
-			AllottedGeometry.ToPaintGeometry(FVector2D(0.0f, AllottedGeometry.GetLocalSize().Y - ShadowOffset.Y), FVector2D(AllottedGeometry.GetLocalSize().X, ShadowOffset.Y)),
+			AllottedGeometry.ToPaintGeometry(FVector2f(AllottedGeometry.GetLocalSize().X, ShadowOffset.Y), FSlateLayoutTransform(FVector2f(0.0f, AllottedGeometry.GetLocalSize().Y - ShadowOffset.Y))),
 			ShadowBrush,
 			ESlateDrawEffect::None,
 			ShadowBrush->GetTint(InWidgetStyle));
@@ -254,7 +253,7 @@ class SDrawerOverlay : public SCompoundWidget
 		FSlateDrawElement::MakeBox(
 			OutDrawElements,
 			OutLayerId+1,
-			RenderTransformedChildGeometry.ToPaintGeometry(ShadowOffset, FVector2D(AllottedGeometry.GetLocalSize().X - (ShadowOffset.X * 2), TargetHeight)),
+			RenderTransformedChildGeometry.ToPaintGeometry(FVector2f(AllottedGeometry.GetLocalSize().X - (ShadowOffset.X * 2), TargetHeight), FSlateLayoutTransform(ShadowOffset)),
 			BorderBrush,
 			ESlateDrawEffect::None,
 			BorderBrush->GetTint(InWidgetStyle));
@@ -295,7 +294,7 @@ private:
 
 	FGeometry GetResizeHandleGeometry(const FGeometry& AllottedGeometry) const
 	{
-		return GetRenderTransformedGeometry(AllottedGeometry).MakeChild(ShadowOffset - FVector2D(0.0f, ExpanderSize), FVector2D(AllottedGeometry.GetLocalSize().X-ShadowOffset.X*2, ExpanderSize));
+		return GetRenderTransformedGeometry(AllottedGeometry).MakeChild(FVector2D(AllottedGeometry.GetLocalSize().X-ShadowOffset.X*2, ExpanderSize), FSlateLayoutTransform(ShadowOffset - FVector2D(0.0f, ExpanderSize)));
 	}
 
 	void SetHeight(float NewHeight)
diff --git a/Engine/Source/Editor/UMGEditor/Private/Designer/SDesignerView.cpp b/Engine/Source/Editor/UMGEditor/Private/Designer/SDesignerView.cpp
--- a/Engine/Source/Editor/UMGEditor/Private/Designer/SDesignerView.cpp
+++ b/Engine/Source/Editor/UMGEditor/Private/Designer/SDesignerView.cpp
@@ -1804,7 +1804,7 @@ FReply SDesignerView::OnMouseMove(const FGeometry& MyGeometry, const FPointerEve
 							const FSlateRenderTransform& AbsoluteToLocalTransform = Inverse(ParentGeometry.GetAccumulatedRenderTransform());
 
 							FWidgetTransform WidgetRenderTransform = WidgetPreview->GetRenderTransform();
-							WidgetRenderTransform.Translation += AbsoluteToLocalTransform.TransformVector(MouseEvent.GetCursorDelta());
+							WidgetRenderTransform.Translation += TransformVector(AbsoluteToLocalTransform, MouseEvent.GetCursorDelta());
 
 							static const FName RenderTransformName(TEXT("RenderTransform"));
 
@@ -2161,7 +2161,7 @@ void SDesignerView::DrawSafeZone(const FOnPaintHandlerParams& PaintArgs)
 			FSlateDrawElement::MakeBox(
 				PaintArgs.OutDrawElements,
 				PaintArgs.Layer,
-				PreviewGeometry.ToPaintGeometry(FVector2D::ZeroVector, FVector2D(Width, SafeMargin.Top)),
+				PreviewGeometry.ToPaintGeometry(FVector2D(Width, SafeMargin.Top), FSlateLayoutTransform()),
 				WhiteBrush,
 				ESlateDrawEffect::None,
 				UnsafeZoneColor
@@ -2171,7 +2171,7 @@ void SDesignerView::DrawSafeZone(const FOnPaintHandlerParams& PaintArgs)
 			FSlateDrawElement::MakeBox(
 				PaintArgs.OutDrawElements,
 				PaintArgs.Layer,
-				PreviewGeometry.ToPaintGeometry(FVector2D(0.0f, Height - SafeMargin.Bottom), FVector2D(Width, SafeMargin.Bottom)),
+				PreviewGeometry.ToPaintGeometry(FVector2f(Width, SafeMargin.Bottom), FSlateLayoutTransform(FVector2f(0.0f, Height - SafeMargin.Bottom))),
 				WhiteBrush,
 				ESlateDrawEffect::None,
 				UnsafeZoneColor
@@ -2181,7 +2181,7 @@ void SDesignerView::DrawSafeZone(const FOnPaintHandlerParams& PaintArgs)
 			FSlateDrawElement::MakeBox(
 				PaintArgs.OutDrawElements,
 				PaintArgs.Layer,
-				PreviewGeometry.ToPaintGeometry(FVector2D(0.0f, SafeMargin.Top), FVector2D(SafeMargin.Left, HeightOfSides)),
+				PreviewGeometry.ToPaintGeometry(FVector2f(SafeMargin.Left, HeightOfSides), FSlateLayoutTransform(FVector2f(0.0f, SafeMargin.Top))),
 				WhiteBrush,
 				ESlateDrawEffect::None,
 				UnsafeZoneColor
@@ -2191,7 +2191,7 @@ void SDesignerView::DrawSafeZone(const FOnPaintHandlerParams& PaintArgs)
 			FSlateDrawElement::MakeBox(
 				PaintArgs.OutDrawElements,
 				PaintArgs.Layer,
-				PreviewGeometry.ToPaintGeometry(FVector2D(Width - SafeMargin.Right, SafeMargin.Top), FVector2D(SafeMargin.Right, HeightOfSides)),
+				PreviewGeometry.ToPaintGeometry(FVector2f(SafeMargin.Right, HeightOfSides), FSlateLayoutTransform(FVector2f(Width - SafeMargin.Right, SafeMargin.Top))),
 				WhiteBrush,
 				ESlateDrawEffect::None,
 				UnsafeZoneColor
@@ -2217,7 +2217,7 @@ void SDesignerView::DrawSafeZone(const FOnPaintHandlerParams& PaintArgs)
 				FSlateDrawElement::MakeBox(
 					PaintArgs.OutDrawElements,
 					PaintArgs.Layer,
-					PreviewGeometry.ToPaintGeometry(Start, Dimensions),
+					PreviewGeometry.ToPaintGeometry(Dimensions, FSlateLayoutTransform(Start)),
 					WhiteBrush,
 					ESlateDrawEffect::None,
 					UnsafeZoneColor
diff --git a/Engine/Source/Editor/UMGEditor/Private/Extensions/CanvasSlotExtension.cpp b/Engine/Source/Editor/UMGEditor/Private/Extensions/CanvasSlotExtension.cpp
--- a/Engine/Source/Editor/UMGEditor/Private/Extensions/CanvasSlotExtension.cpp
+++ b/Engine/Source/Editor/UMGEditor/Private/Extensions/CanvasSlotExtension.cpp
@@ -562,9 +562,7 @@ void FCanvasSlotExtension::PaintDragPercentages(const TSet< FWidgetReference >&
 				Designer->GetWidgetGeometry(Canvas, CanvasGeometry);
 				// Ignore all widget scales and only use the designer scale (text doesn't need the designer scale, however the rendered lines do)
 				const FGeometry IgnoreScale = CanvasGeometry.MakeChild(
-					FVector2D::ZeroVector,
-					CanvasGeometry.GetLocalSize(),
-					Inverse(CanvasGeometry.GetAccumulatedLayoutTransform().GetScale()) * Designer->GetPreviewScale() * AllottedGeometry.Scale
+					CanvasGeometry.GetLocalSize(), Inverse(CanvasGeometry.GetAccumulatedLayoutTransform().GetScale()) * Designer->GetPreviewScale() * AllottedGeometry.Scale
 				);
 				CanvasGeometry = Designer->MakeGeometryWindowLocal(IgnoreScale);
 
@@ -670,13 +668,13 @@ void FCanvasSlotExtension::PaintLineWithText(FVector2D Start, FVector2D End, FTe
 		Offset.Y += ((End - Start).Y - TextSize.Y) * TextTransform.Y;
 	}
 
-	const FGeometry ChildGeometry = AllottedGeometry.MakeChild(Start + Offset, AllottedGeometry.GetLocalSize());
+	const FGeometry ChildGeometry = AllottedGeometry.MakeChild(AllottedGeometry.GetLocalSize(), FSlateLayoutTransform(Start + Offset));
 
 	// Draw drop shadow
 	FSlateDrawElement::MakeText(
 		OutDrawElements,
 		LayerId,
-		ChildGeometry.ToPaintGeometry(FVector2D(1, 1), TextSize, InverseDesignerScale),
+		ChildGeometry.ToPaintGeometry(TextSize, FSlateLayoutTransform(InverseDesignerScale, FVector2f(1.f, 1.f))),
 		Text,
 		AnchorFont,
 		ESlateDrawEffect::None,
@@ -687,7 +685,7 @@ void FCanvasSlotExtension::PaintLineWithText(FVector2D Start, FVector2D End, FTe
 	FSlateDrawElement::MakeText(
 		OutDrawElements,
 		++LayerId,
-		ChildGeometry.ToPaintGeometry(FVector2D::ZeroVector, TextSize, InverseDesignerScale),
+		ChildGeometry.ToPaintGeometry(TextSize, FSlateLayoutTransform(InverseDesignerScale)),
 		Text,
 		AnchorFont,
 		ESlateDrawEffect::None,
diff --git a/Engine/Source/Editor/UnrealEd/Private/PropertyInfoViewStyle.cpp b/Engine/Source/Editor/UnrealEd/Private/PropertyInfoViewStyle.cpp
--- a/Engine/Source/Editor/UnrealEd/Private/PropertyInfoViewStyle.cpp
+++ b/Engine/Source/Editor/UnrealEd/Private/PropertyInfoViewStyle.cpp
@@ -102,7 +102,7 @@ namespace PropertyInfoViewStyle
 			FSlateDrawElement::MakeBox(
 				OutDrawElements,
 				LayerId,
-				AllottedGeometry.ToPaintGeometry(FVector2D(TabSize * IndentCount, 0), FVector2D(TabSize, AllottedGeometry.GetLocalSize().Y)),
+				AllottedGeometry.ToPaintGeometry(FVector2f(TabSize, AllottedGeometry.GetLocalSize().Y), FSlateLayoutTransform(FVector2f(TabSize * IndentCount, 0.f))),
 				BackgroundBrush,
 				ESlateDrawEffect::None,
 				BackgroundColor.GetColor(InWidgetStyle)
@@ -111,7 +111,7 @@ namespace PropertyInfoViewStyle
 			FSlateDrawElement::MakeBox(
 				OutDrawElements,
 				LayerId + 1,
-				AllottedGeometry.ToPaintGeometry(FVector2D(TabSize * IndentCount, 0), FVector2D(TabSize, AllottedGeometry.GetLocalSize().Y)),
+				AllottedGeometry.ToPaintGeometry(FVector2f(TabSize, AllottedGeometry.GetLocalSize().Y), FSlateLayoutTransform(FVector2f(TabSize * IndentCount, 0.f))),
 				DropShadowBrush
 			);
 		}
diff --git a/Engine/Source/Editor/UnrealEd/Private/SColorGradientEditor.cpp b/Engine/Source/Editor/UnrealEd/Private/SColorGradientEditor.cpp
--- a/Engine/Source/Editor/UnrealEd/Private/SColorGradientEditor.cpp
+++ b/Engine/Source/Editor/UnrealEd/Private/SColorGradientEditor.cpp
@@ -151,7 +151,7 @@ int32 SColorGradientEditor::OnPaint( const FPaintArgs& Args, const FGeometry& Al
 		FGeometry ColorMarkAreaGeometry = GetColorMarkAreaGeometry( AllottedGeometry );
 		FGeometry AlphaMarkAreaGeometry = GetAlphaMarkAreaGeometry( AllottedGeometry );
 
-		FGeometry GradientAreaGeometry = AllottedGeometry.MakeChild( FVector2D(0.0f, 16.0f), FVector2D( AllottedGeometry.GetLocalSize().X, AllottedGeometry.GetLocalSize().Y - 30.0f ) );
+		FGeometry GradientAreaGeometry = AllottedGeometry.MakeChild( FVector2D( AllottedGeometry.GetLocalSize().X, AllottedGeometry.GetLocalSize().Y - 30.0f ), FSlateLayoutTransform(FVector2D(0.0f, 16.0f)) );
 
 		bool bEnabled = ShouldBeEnabled( bParentEnabled );
 		ESlateDrawEffect DrawEffects = bEnabled ? ESlateDrawEffect::None : ESlateDrawEffect::DisabledEffect;
@@ -775,7 +775,7 @@ void SColorGradientEditor::DrawGradientStopMark( const FGradientStopMark& Mark,
 	( 
 		OutDrawElements,
 		LayerId,
-		Geometry.ToPaintGeometry( FVector2D( XPos-HandleRect.Left, HandleRect.Top ), FVector2D( HandleRect.Right, HandleRect.Bottom ) ),
+		Geometry.ToPaintGeometry( FVector2f( HandleRect.Right, HandleRect.Bottom ), FSlateLayoutTransform(FVector2f( XPos-HandleRect.Left, HandleRect.Top )) ),
 		bColor ? ColorStopBrush : AlphaStopBrush,
 		DrawEffects,
 		bSelected ? SelectionColor : FLinearColor::White
@@ -787,7 +787,7 @@ void SColorGradientEditor::DrawGradientStopMark( const FGradientStopMark& Mark,
 	( 
 		OutDrawElements,
 		LayerId+1,
-		Geometry.ToPaintGeometry( FVector2D( XPos-HandleRect.Left+3, bColor ? HandleRect.Top+3.0f : HandleRect.Top+6), FVector2D( HandleRect.Right-6, HandleRect.Bottom-9 ) ),
+		Geometry.ToPaintGeometry( FVector2f( HandleRect.Right-6.f, HandleRect.Bottom-9.f ), FSlateLayoutTransform(FVector2f( XPos-HandleRect.Left+3.f, bColor ? HandleRect.Top+3.0f : HandleRect.Top+6.f)) ),
 		WhiteBrush,
 		DrawEffects,
 		Color.ToFColor(true)
@@ -795,12 +795,12 @@ void SColorGradientEditor::DrawGradientStopMark( const FGradientStopMark& Mark,
 }
 FGeometry SColorGradientEditor::GetColorMarkAreaGeometry( const FGeometry& InGeometry ) const
 {
-	return InGeometry.MakeChild( FVector2D( 0.0f, 0.0f), FVector2D( InGeometry.GetLocalSize().X, 16.0f ) );
+	return InGeometry.MakeChild( FVector2D( InGeometry.GetLocalSize().X, 16.0f ), FSlateLayoutTransform() );
 }
 
 FGeometry SColorGradientEditor::GetAlphaMarkAreaGeometry( const FGeometry& InGeometry ) const
 {
-	return InGeometry.MakeChild( FVector2D( 0.0f, InGeometry.GetLocalSize().Y-14.0f), FVector2D( InGeometry.GetLocalSize().X, 16.0f ) );
+	return InGeometry.MakeChild( FVector2D( InGeometry.GetLocalSize().X, 16.0f ), FSlateLayoutTransform(FVector2D( 0.0f, InGeometry.GetLocalSize().Y-14.0f)) );
 }
 
 FGradientStopMark SColorGradientEditor::GetGradientStopAtPoint( const FVector2D& MousePos, const FGeometry& MyGeometry )
@@ -826,7 +826,7 @@ FGradientStopMark SColorGradientEditor::GetGradientStopAtPoint( const FVector2D&
 
 			if( XVal >= 0 )
 			{
-				FGeometry MarkGeometry = ColorMarkAreaGeometry.MakeChild( FVector2D( XVal-HandleRect.Left, HandleRect.Top ), FVector2D( HandleRect.Right, HandleRect.Bottom ) );
+				FGeometry MarkGeometry = ColorMarkAreaGeometry.MakeChild( FVector2D( HandleRect.Right, HandleRect.Bottom ), FSlateLayoutTransform(FVector2D( XVal-HandleRect.Left, HandleRect.Top )) );
 				if( MarkGeometry.IsUnderLocation( MousePos ) )
 				{
 					return Mark;
@@ -843,7 +843,7 @@ FGradientStopMark SColorGradientEditor::GetGradientStopAtPoint( const FVector2D&
 
 			if( XVal >= 0 )
 			{
-				FGeometry MarkGeometry = AlphaMarkAreaGeometry.MakeChild( FVector2D( XVal-HandleRect.Left, HandleRect.Top ), FVector2D( HandleRect.Right, HandleRect.Bottom ) );
+				FGeometry MarkGeometry = AlphaMarkAreaGeometry.MakeChild( FVector2D( HandleRect.Right, HandleRect.Bottom ), FSlateLayoutTransform(FVector2D( XVal-HandleRect.Left, HandleRect.Top )) );
 				if( MarkGeometry.IsUnderLocation( MousePos ) )
 				{
 					return Mark;
diff --git a/Engine/Source/Editor/UnrealEd/Private/SCurveEditor.cpp b/Engine/Source/Editor/UnrealEd/Private/SCurveEditor.cpp
--- a/Engine/Source/Editor/UnrealEd/Private/SCurveEditor.cpp
+++ b/Engine/Source/Editor/UnrealEd/Private/SCurveEditor.cpp
@@ -772,7 +772,7 @@ int32 SCurveEditor::OnPaint(const FPaintArgs& Args, const FGeometry& AllottedGeo
 		(
 		OutDrawElements,
 		BackgroundLayerId,
-		CurveAreaGeometry.ToPaintGeometry(FVector2D(ZeroInputX, 0), FVector2D(TimelineMaxX - ZeroInputX, CurveAreaGeometry.GetLocalSize().Y)),
+		CurveAreaGeometry.ToPaintGeometry(FVector2D(TimelineMaxX - ZeroInputX, CurveAreaGeometry.GetLocalSize().Y), FSlateLayoutTransform(FVector2D(ZeroInputX, 0.f))),
 		TimelineAreaBrush,
 		DrawEffects,
 		TimelineAreaBrush->GetTint(InWidgetStyle) * InWidgetStyle.GetColorAndOpacityTint()
@@ -803,7 +803,7 @@ int32 SCurveEditor::OnPaint(const FPaintArgs& Args, const FGeometry& AllottedGeo
 		(
 			OutDrawElements,
 			ZeroLineLayerId,
-			CurveAreaGeometry.ToPaintGeometry( FVector2D(0, ZeroOutputY), FVector2D(CurveAreaGeometry.Size.X, 1) ),
+			CurveAreaGeometry.ToPaintGeometry( FVector2D(CurveAreaGeometry.Size.X, 1), FSlateLayoutTransform(FVector2D(0.f, ZeroOutputY)) ),
 			WhiteBrush,
 			DrawEffects,
 			WhiteBrush->GetTint( InWidgetStyle ) * InWidgetStyle.GetColorAndOpacityTint()
@@ -1013,7 +1013,7 @@ void SCurveEditor::PaintKeys(TSharedPtr<FCurveViewModel> CurveViewModel, FTrackS
 		FSlateDrawElement::MakeBox(
 			OutDrawElements,
 			LayerToUse,
-			AllottedGeometry.ToPaintGeometry( KeyIconLocation, CONST_KeySize ),
+			AllottedGeometry.ToPaintGeometry( CONST_KeySize, FSlateLayoutTransform(KeyIconLocation) ),
 			KeyBrush,
 			DrawEffects,
 			KeyBrush->GetTint( InWidgetStyle ) * InWidgetStyle.GetColorAndOpacityTint() * KeyColor * SelectionTint
@@ -1058,8 +1058,8 @@ void SCurveEditor::PaintTangent( TSharedPtr<FCurveViewModel> CurveViewModel, FTr
 
 	//Add lines from tangent control point to 'key'
 	TArray<FVector2D> LinePoints;
-	LinePoints.Add(KeyLocation);
-	LinePoints.Add(ArriveTangentLocation);
+	LinePoints.Add(FVector2D(KeyLocation));
+	LinePoints.Add(FVector2D(ArriveTangentLocation));
 	FSlateDrawElement::MakeLines(
 		OutDrawElements,
 		LayerId,
@@ -1070,8 +1070,8 @@ void SCurveEditor::PaintTangent( TSharedPtr<FCurveViewModel> CurveViewModel, FTr
 		);
 
 	LinePoints.Empty();
-	LinePoints.Add(KeyLocation);
-	LinePoints.Add(LeaveTangentLocation);
+	LinePoints.Add(FVector2D(KeyLocation));
+	LinePoints.Add(FVector2D(LeaveTangentLocation));
 	FSlateDrawElement::MakeLines(
 		OutDrawElements,
 		LayerId,
@@ -1085,7 +1085,7 @@ void SCurveEditor::PaintTangent( TSharedPtr<FCurveViewModel> CurveViewModel, FTr
 	FSlateDrawElement::MakeBox(
 		OutDrawElements,
 		LayerToUse,
-		AllottedGeometry.ToPaintGeometry(ArriveTangentIconLocation, CONST_TangentSize ),
+		AllottedGeometry.ToPaintGeometry( CONST_TangentSize, FSlateLayoutTransform(ArriveTangentIconLocation) ),
 		ArriveTangentSelected ? TangentBrushSelected : TangentBrush,
 		DrawEffects,
 		ArriveTangentSelected ? TangentBrushSelected->GetTint( InWidgetStyle ) * ArriveSelectionTint : TangentBrush->GetTint( InWidgetStyle ) * InWidgetStyle.GetColorAndOpacityTint() * ArriveSelectionTint
@@ -1094,7 +1094,7 @@ void SCurveEditor::PaintTangent( TSharedPtr<FCurveViewModel> CurveViewModel, FTr
 	FSlateDrawElement::MakeBox(
 		OutDrawElements,
 		LayerToUse,
-		AllottedGeometry.ToPaintGeometry(LeaveTangentIconLocation, CONST_TangentSize ),
+		AllottedGeometry.ToPaintGeometry( CONST_TangentSize, FSlateLayoutTransform(LeaveTangentIconLocation) ),
 		LeaveTangentSelected ? TangentBrushSelected : TangentBrush,
 		DrawEffects,
 		LeaveTangentSelected ? TangentBrushSelected->GetTint( InWidgetStyle ) * LeaveSelectionTint : TangentBrush->GetTint( InWidgetStyle ) * InWidgetStyle.GetColorAndOpacityTint() * LeaveSelectionTint
@@ -1160,7 +1160,7 @@ void SCurveEditor::PaintGridLines(const FGeometry &AllottedGeometry, FTrackScale
 					if (bDrawInputGridNumbers)
 					{
 						FString TimeStr = FString::Printf(TEXT("%.2f"), Time);
-						FSlateDrawElement::MakeText(OutDrawElements,LayerId,AllottedGeometry.MakeChild(FVector2D(X, 0.0), FVector2D(1.0f, ScaleX )).ToPaintGeometry(),TimeStr,
+						FSlateDrawElement::MakeText(OutDrawElements,LayerId,AllottedGeometry.MakeChild(FVector2D(1.0f, ScaleX ), FSlateLayoutTransform(FVector2D(X, 0.0))).ToPaintGeometry(),TimeStr,
 							FAppStyle::GetFontStyle("CurveEd.InfoFont"), DrawEffects, GridTextColor );
 					}
 
@@ -1211,12 +1211,12 @@ void SCurveEditor::PaintGridLines(const FGeometry &AllottedGeometry, FTrackScale
 						FVector2D DrawSize = FontMeasureService->Measure(ValueStr, Font);
 
 						// draw at the start
-						FSlateDrawElement::MakeText(OutDrawElements,LayerId,AllottedGeometry.MakeChild(FVector2D(0.0f, Y), FVector2D(ScaleY, 1.0f )).ToPaintGeometry(),ValueStr,
-												 Font, DrawEffects, GridTextColor );
+						FSlateDrawElement::MakeText(OutDrawElements,LayerId,AllottedGeometry.MakeChild(FVector2D(ScaleY, 1.0f ), FSlateLayoutTransform(FVector2D(0.0f, Y))).ToPaintGeometry(),
+							ValueStr, Font, DrawEffects, GridTextColor );
 
 						// draw at the last since sometimes start can be hidden
-						FSlateDrawElement::MakeText(OutDrawElements,LayerId,AllottedGeometry.MakeChild(FVector2D(AllottedGeometry.GetLocalSize().X-DrawSize.X, Y), FVector2D(ScaleY, 1.0f )).ToPaintGeometry(),ValueStr,
-												 Font, DrawEffects, GridTextColor );
+						FSlateDrawElement::MakeText(OutDrawElements,LayerId,AllottedGeometry.MakeChild(FVector2D(ScaleY, 1.0f ), FSlateLayoutTransform(FVector2D(AllottedGeometry.GetLocalSize().X-DrawSize.X, Y))).ToPaintGeometry(),
+							ValueStr, Font, DrawEffects, GridTextColor );
 					}
 					
 					LinePoints.Empty();
@@ -1241,7 +1241,7 @@ void SCurveEditor::PaintMarquee(const FGeometry& AllottedGeometry, const FSlateR
 	FSlateDrawElement::MakeBox(
 		OutDrawElements,
 		LayerId,
-		AllottedGeometry.ToPaintGeometry(MarqueTopLeft, MarqueBottomRight - MarqueTopLeft),
+		AllottedGeometry.ToPaintGeometry(MarqueBottomRight - MarqueTopLeft, FSlateLayoutTransform(MarqueTopLeft)),
 		FAppStyle::GetBrush(TEXT("MarqueeSelection"))
 		);
 }
diff --git a/Engine/Source/Editor/WorldBrowser/Private/Tiles/SWorldComposition.cpp b/Engine/Source/Editor/WorldBrowser/Private/Tiles/SWorldComposition.cpp
--- a/Engine/Source/Editor/WorldBrowser/Private/Tiles/SWorldComposition.cpp
+++ b/Engine/Source/Editor/WorldBrowser/Private/Tiles/SWorldComposition.cpp
@@ -314,9 +314,9 @@ public:
 				//FSlateLayoutTransform LayoutTransform(Scale, AllottedGeometry.GetAccumulatedLayoutTransform().GetTranslation() - InflateAmount);
 				//FSlateRenderTransform SlateRenderTransform(Scale, AllottedGeometry.GetAccumulatedRenderTransform().GetTranslation() - InflateAmount);
 
-				FPaintGeometry PaintGeometry = AllottedGeometry.ToPaintGeometry(
-					ObserverPositionScreen - CameraImage->ImageSize*0.5f, 
-					CameraImage->ImageSize
+				FPaintGeometry PaintGeometry = AllottedGeometry.ToPaintGeometry( 
+					CameraImage->ImageSize,
+					FSlateLayoutTransform(ObserverPositionScreen - CameraImage->ImageSize*0.5f)
 				);
 
 				FSlateDrawElement::MakeRotatedBox(
@@ -339,8 +339,8 @@ public:
 				const FSlateBrush* CameraImage = FAppStyle::GetBrush(TEXT("WorldBrowser.SimulationViewPosition"));
 	
 				FPaintGeometry PaintGeometry = AllottedGeometry.ToPaintGeometry(
-					PlayerPositionScreen - CameraImage->ImageSize*0.5f, 
-					CameraImage->ImageSize
+					CameraImage->ImageSize,
+					FSlateLayoutTransform(PlayerPositionScreen - CameraImage->ImageSize*0.5f)
 					);
 
 				FSlateDrawElement::MakeRotatedBox(
@@ -392,7 +392,7 @@ public:
 				}
 
 				// Panning and mouse is outside of panel? Pasting should just go to the screen center.
-				PastePosition = PanelCoordToGraphCoord( 0.5 * MyGeometry.GetLocalSize() );
+				PastePosition = PanelCoordToGraphCoord( 0.5f * MyGeometry.GetLocalSize() );
 
 				this->bIsPanning = true;
 				ViewOffset -= CursorDelta / GetZoomAmount();
diff --git a/Engine/Source/Editor/WorldPartitionEditor/Private/WorldPartition/SWorldPartitionEditorGrid2D.cpp b/Engine/Source/Editor/WorldPartitionEditor/Private/WorldPartition/SWorldPartitionEditorGrid2D.cpp
--- a/Engine/Source/Editor/WorldPartitionEditor/Private/WorldPartition/SWorldPartitionEditorGrid2D.cpp
+++ b/Engine/Source/Editor/WorldPartitionEditor/Private/WorldPartition/SWorldPartitionEditorGrid2D.cpp
@@ -1201,7 +1201,7 @@ uint32 SWorldPartitionEditorGrid2D::PaintActors(const FGeometry& AllottedGeometr
 				FSlateDrawElement::MakeBox(
 					OutDrawElements,
 					++LayerId,
-					AllottedGeometry.ToPaintGeometry(LabelTextPos - BackgroundGrowSize, LabelTextSize + BackgroundGrowSize),
+					AllottedGeometry.ToPaintGeometry(LabelTextSize + BackgroundGrowSize, FSlateLayoutTransform(LabelTextPos - BackgroundGrowSize)),
 					&BackgroundBrush,
 					ESlateDrawEffect::None,
 					LabelBackgroundColor
@@ -1264,7 +1264,7 @@ uint32 SWorldPartitionEditorGrid2D::PaintActors(const FGeometry& AllottedGeometr
 					const float AreaFadeDistance = 128.0f;
 					if ((Extent.Size2D() < KINDA_SMALL_NUMBER) || (ActorViewBox.GetArea() > MinimumAreaCull))
 					{
-						const FPaintGeometry ActorGeometry = AllottedGeometry.ToPaintGeometry(TopLeft, BottomRight - TopLeft);
+						const FPaintGeometry ActorGeometry = AllottedGeometry.ToPaintGeometry(BottomRight - TopLeft, FSlateLayoutTransform(TopLeft));
 						const float LoaderColorGradient = FMath::Min((ActorViewBox.GetArea() - MinimumAreaCull) / AreaFadeDistance, 1.0f);
 						const FLinearColor LoaderColor = LoaderAdapter->GetColor().IsSet() ? *LoaderAdapter->GetColor() : FColor::White;
 						const bool bIsLocalHovered = LocalHoveredLoaderAdapter == LoaderAdapter;
@@ -1295,8 +1295,8 @@ uint32 SWorldPartitionEditorGrid2D::PaintActors(const FGeometry& AllottedGeometr
 							for (const FBox2D& Rect : LoaderInterfaceHighlight.GetRects())
 							{
 								const FPaintGeometry RectAreaGeometry = AllottedGeometry.ToPaintGeometry(
-									Rect.Min,
-									Rect.Max - Rect.Min
+									Rect.Max - Rect.Min,
+									FSlateLayoutTransform(Rect.Min)
 								);
 
 								FSlateDrawElement::MakeBox(
@@ -1396,8 +1396,8 @@ uint32 SWorldPartitionEditorGrid2D::PaintActors(const FGeometry& AllottedGeometr
 				const float AreaFadeDistance = 128.0f;
 				if ((Extent.Size2D() < KINDA_SMALL_NUMBER) || (ActorViewBox.GetArea() > MinimumAreaCull))
 				{
-					const FPaintGeometry ActorGeometry = AllottedGeometry.ToPaintGeometry(TopLeft, BottomRight - TopLeft);
-					const FPaintGeometry ActorGeometryShadow = AllottedGeometry.ToPaintGeometry(TopLeft - FVector2D::One(), BottomRight - TopLeft + FVector2D::One() * 2.0);
+					const FPaintGeometry ActorGeometry = AllottedGeometry.ToPaintGeometry(BottomRight - TopLeft, FSlateLayoutTransform(TopLeft));
+					const FPaintGeometry ActorGeometryShadow = AllottedGeometry.ToPaintGeometry(BottomRight - TopLeft + FVector2D::One() * 2.0, FSlateLayoutTransform(TopLeft - FVector2D::One()));
 					const bool bIsBoundsEdgeHovered = !bIsBoundsEdgeHoveredFound && IsBoundsEdgeHovered(MouseCursorPos, ActorViewBox, 10.0f);
 					const float ActorColorGradient = FMath::Cube(FMath::Clamp((ActorViewBox.GetArea() - MinimumAreaCull) / AreaFadeDistance, 0.0f, 1.0f));
 					const float ActorBrightness = bIsSpatiallyLoaded ? 1.0f : 0.3f;
@@ -1551,8 +1551,8 @@ uint32 SWorldPartitionEditorGrid2D::PaintViewer(const FGeometry& AllottedGeometr
 	{
 		const FVector2D LocalLocation = WorldToScreen.TransformPoint(Location);
 		const FPaintGeometry PaintGeometryShadow = AllottedGeometry.ToPaintGeometry(
-			LocalLocation - (Image->ImageSize + ShadowSize) * 0.5f, 
-			Image->ImageSize + ShadowSize
+			Image->ImageSize + ShadowSize,
+			FSlateLayoutTransform(LocalLocation - (Image->ImageSize + ShadowSize) * 0.5f)
 		);
 
 		FSlateDrawElement::MakeRotatedBox(
@@ -1568,8 +1568,8 @@ uint32 SWorldPartitionEditorGrid2D::PaintViewer(const FGeometry& AllottedGeometr
 		);
 
 		const FPaintGeometry PaintGeometry = AllottedGeometry.ToPaintGeometry(
-			LocalLocation - Image->ImageSize * 0.5f, 
-			Image->ImageSize
+			Image->ImageSize,
+			FSlateLayoutTransform(LocalLocation - Image->ImageSize * 0.5f)
 		);
 
 		FSlateDrawElement::MakeRotatedBox(
@@ -1635,8 +1635,8 @@ uint32 SWorldPartitionEditorGrid2D::PaintSelection(const FGeometry& AllottedGeom
 				FLinearColor CellColor(FLinearColor(1, 1, 1, 0.1f));
 
 				FPaintGeometry CellGeometry = AllottedGeometry.ToPaintGeometry(
-					TopLeft,
-					BottomRight - TopLeft
+					BottomRight - TopLeft,
+					FSlateLayoutTransform(TopLeft)
 				);
 
 				FSlateDrawElement::MakeBox(
@@ -1745,7 +1745,7 @@ int32 SWorldPartitionEditorGrid2D::PaintSoftwareCursor(const FGeometry& Allotted
 		FSlateDrawElement::MakeBox(
 			OutDrawElements,
 			LayerId,
-			AllottedGeometry.ToPaintGeometry(MouseCursorPos - (CursorSize * 0.5f), CursorSize),
+			AllottedGeometry.ToPaintGeometry(CursorSize, FSlateLayoutTransform(MouseCursorPos - (CursorSize * 0.5f))),
 			Brush
 		);
 	}
@@ -1764,8 +1764,8 @@ int32 SWorldPartitionEditorGrid2D::PaintMinimap(const FGeometry& AllottedGeometr
 		);
 
 		const FPaintGeometry WorldImageGeometry = AllottedGeometry.ToPaintGeometry(
-			MinimapBounds.Min,
-			MinimapBounds.Max - MinimapBounds.Min
+			MinimapBounds.Max - MinimapBounds.Min,
+			FSlateLayoutTransform(MinimapBounds.Min)
 		);
 
 		FSlateDrawElement::MakeBox(
@@ -1813,8 +1813,8 @@ int32 SWorldPartitionEditorGrid2D::PaintMinimap(const FGeometry& AllottedGeometr
 			for (const FBox2D& ShadowArea : ShadowAreas.GetRects())
 			{
 				const FPaintGeometry ShadowAreaGeometry = AllottedGeometry.ToPaintGeometry(
-					ShadowArea.Min,
-					ShadowArea.Max - ShadowArea.Min
+					ShadowArea.Max - ShadowArea.Min,
+					FSlateLayoutTransform(ShadowArea.Min)
 				);
 
 				const FSlateColorBrush ShadowdBrush(FLinearColor::Gray);
@@ -1934,7 +1934,7 @@ int32 SWorldPartitionEditorGrid2D::DrawTextLabel(FSlateWindowElementList& OutDra
 		FSlateDrawElement::MakeText(
 			OutDrawElements,
 			++LayerId,
-			AllottedGeometry.ToPaintGeometry(LabelTextPos + FVector2D(2, 2), FVector2D(1, 1)),
+			AllottedGeometry.ToPaintGeometry(FVector2D(1, 1), FSlateLayoutTransform(LabelTextPos + FVector2D(2, 2))),
 			Label,
 			Font,
 			ESlateDrawEffect::None,
@@ -1944,7 +1944,7 @@ int32 SWorldPartitionEditorGrid2D::DrawTextLabel(FSlateWindowElementList& OutDra
 		FSlateDrawElement::MakeText(
 			OutDrawElements,
 			++LayerId,
-			AllottedGeometry.ToPaintGeometry(LabelTextPos, FVector2D(1,1)),
+			AllottedGeometry.ToPaintGeometry(FVector2D(1,1), FSlateLayoutTransform(LabelTextPos)),
 			Label,
 			Font,
 			ESlateDrawEffect::None,
diff --git a/Engine/Source/Editor/WorldPartitionEditor/Private/WorldPartition/SWorldPartitionEditorGridSpatialHash.cpp b/Engine/Source/Editor/WorldPartitionEditor/Private/WorldPartition/SWorldPartitionEditorGridSpatialHash.cpp
--- a/Engine/Source/Editor/WorldPartitionEditor/Private/WorldPartition/SWorldPartitionEditorGridSpatialHash.cpp
+++ b/Engine/Source/Editor/WorldPartitionEditor/Private/WorldPartition/SWorldPartitionEditorGridSpatialHash.cpp
@@ -94,8 +94,8 @@ int32 SWorldPartitionEditorGridSpatialHash::PaintGrid(const FGeometry& AllottedG
 		FLinearColor ShadowColor(FLinearColor(0.0f, 0.0f, 0.0f, 0.5f));
 
 		FPaintGeometry GridGeometry = AllottedGeometry.ToPaintGeometry(
-			WorldToScreen.TransformPoint(FVector2D(VisibleMinimapRectWorld.Min)),
-			WorldToScreen.TransformPoint(FVector2D(VisibleMinimapRectWorld.Max)) - WorldToScreen.TransformPoint(FVector2D(VisibleMinimapRectWorld.Min))
+			WorldToScreen.TransformPoint(FVector2D(VisibleMinimapRectWorld.Max)) - WorldToScreen.TransformPoint(FVector2D(VisibleMinimapRectWorld.Min)),
+			FSlateLayoutTransform(WorldToScreen.TransformPoint(FVector2D(VisibleMinimapRectWorld.Min)))
 		);
 
 		FSlateDrawElement::MakeBox(
@@ -234,7 +234,7 @@ int32 SWorldPartitionEditorGridSpatialHash::PaintGrid(const FGeometry& AllottedG
 						FSlateDrawElement::MakeText(
 							OutDrawElements,
 							++LayerId,
-							AllottedGeometry.ToPaintGeometry(WorldToScreen.TransformPoint(FVector2D(x + EffectiveCellSize / 2, y + EffectiveCellSize / 2)) - CoordTextSize / 2, FVector2D(1,1)),
+							AllottedGeometry.ToPaintGeometry(FVector2D(1,1), FSlateLayoutTransform(WorldToScreen.TransformPoint(FVector2D(x + EffectiveCellSize / 2, y + EffectiveCellSize / 2)) - CoordTextSize / 2)),
 							FString::Printf(TEXT("(%lld,%lld)"), x / EffectiveCellSize, y / EffectiveCellSize),
 							CoordsFont,
 							ESlateDrawEffect::None,
diff --git a/Engine/Source/Runtime/AdvancedWidgets/Private/Slate/SRadialSlider.cpp b/Engine/Source/Runtime/AdvancedWidgets/Private/Slate/SRadialSlider.cpp
--- a/Engine/Source/Runtime/AdvancedWidgets/Private/Slate/SRadialSlider.cpp
+++ b/Engine/Source/Runtime/AdvancedWidgets/Private/Slate/SRadialSlider.cpp
@@ -140,7 +140,7 @@ int32 SRadialSlider::OnPaint( const FPaintArgs& Args, const FGeometry& AllottedG
 	(
 		OutDrawElements,
 		LayerId,
-		SliderGeometry.ToPaintGeometry(SliderMidPoint, SliderDiameter),
+		SliderGeometry.ToPaintGeometry(SliderDiameter, FSlateLayoutTransform(SliderMidPoint)),
 		SliderBarPoints,
 		DrawEffects,
 		SliderBarColorWithAlpha,
@@ -153,7 +153,7 @@ int32 SRadialSlider::OnPaint( const FPaintArgs& Args, const FGeometry& AllottedG
 	(
 		OutDrawElements,
 		LayerId,
-		SliderGeometry.ToPaintGeometry(SliderMidPoint, SliderDiameter),
+		SliderGeometry.ToPaintGeometry(SliderDiameter, FSlateLayoutTransform(SliderMidPoint)),
 		ProgressBarPoints,
 		DrawEffects,
 		ProgressBarColorWithAlpha,
@@ -169,7 +169,7 @@ int32 SRadialSlider::OnPaint( const FPaintArgs& Args, const FGeometry& AllottedG
 	(
 		OutDrawElements, 
 		LayerId, 
-		SliderGeometry.ToPaintGeometry(SliderMidPoint - SliderRadius, SliderDiameter),
+		SliderGeometry.ToPaintGeometry(SliderDiameter, FSlateLayoutTransform(SliderMidPoint - SliderRadius)),
 		&CenterBackgroundBrush,
 		DrawEffects, 
 		CenterBackgroundColorWithAlpha
@@ -189,7 +189,7 @@ int32 SRadialSlider::OnPaint( const FPaintArgs& Args, const FGeometry& AllottedG
 		(
 			OutDrawElements,
 			LayerId,
-			SliderGeometry.ToPaintGeometry(SliderMidPoint, SliderDiameter),
+			SliderGeometry.ToPaintGeometry(SliderDiameter, FSlateLayoutTransform(SliderMidPoint)),
 			TArray<FVector2D>{HandStartRotated, HandEndRotated},
 			DrawEffects,
 			SliderBarColorWithAlpha,
@@ -243,7 +243,7 @@ int32 SRadialSlider::OnPaint( const FPaintArgs& Args, const FGeometry& AllottedG
 		FSlateDrawElement::MakeLines(
 			OutDrawElements,
 			LayerId,
-			SliderGeometry.ToPaintGeometry(SliderMidPoint, SliderDiameter),
+			SliderGeometry.ToPaintGeometry(SliderDiameter, FSlateLayoutTransform(SliderMidPoint)),
 			TagLinePoints,
 			DrawEffects,
 			TagLineColor,
@@ -265,7 +265,7 @@ int32 SRadialSlider::OnPaint( const FPaintArgs& Args, const FGeometry& AllottedG
 		FSlateDrawElement::MakeText(
 			OutDrawElements,
 			LayerId,
-			SliderGeometry.ToPaintGeometry(TextTopLeftPoint, GetThumbImage()->ImageSize),
+			SliderGeometry.ToPaintGeometry(GetThumbImage()->ImageSize, FSlateLayoutTransform(TextTopLeftPoint)),
 			TextToWrite,
 			MyFont,
 			DrawEffects,
@@ -282,7 +282,7 @@ int32 SRadialSlider::OnPaint( const FPaintArgs& Args, const FGeometry& AllottedG
 		FSlateDrawElement::MakeRotatedBox(
 			OutDrawElements,
 			LayerId,
-			SliderGeometry.ToPaintGeometry(HandleTopLeftPoint, GetThumbImage()->ImageSize),
+			SliderGeometry.ToPaintGeometry(GetThumbImage()->ImageSize, FSlateLayoutTransform(HandleTopLeftPoint)),
 			ThumbImage,
 			DrawEffects,
 			(180.0f + MidPointAngle + AngularOffset) * (PI / 180.0f),
diff --git a/Engine/Source/Runtime/AppFramework/Private/Widgets/Testing/STestSuite.cpp b/Engine/Source/Runtime/AppFramework/Private/Widgets/Testing/STestSuite.cpp
--- a/Engine/Source/Runtime/AppFramework/Private/Widgets/Testing/STestSuite.cpp
+++ b/Engine/Source/Runtime/AppFramework/Private/Widgets/Testing/STestSuite.cpp
@@ -926,7 +926,7 @@ private:
 		FSlateDrawElement::MakeText(
 			InParams.OutDrawElements,
 			InParams.Layer,
-			InParams.Geometry.ToPaintGeometry(FVector2D(0,0), InParams.Geometry.GetLocalSize(), FontScale),
+			InParams.Geometry.ToPaintGeometry(InParams.Geometry.GetLocalSize(), FSlateLayoutTransform(FontScale)),
 			Text.ToString(),
 			FontInfo,
 			InParams.bEnabled ? ESlateDrawEffect::None : ESlateDrawEffect::DisabledEffect,
diff --git a/Engine/Source/Runtime/CoreUObject/Public/UObject/NoExportTypes.h b/Engine/Source/Runtime/CoreUObject/Public/UObject/NoExportTypes.h
--- a/Engine/Source/Runtime/CoreUObject/Public/UObject/NoExportTypes.h
+++ b/Engine/Source/Runtime/CoreUObject/Public/UObject/NoExportTypes.h
@@ -598,10 +598,10 @@ struct FVector4
 USTRUCT(immutable, noexport, BlueprintType, BlueprintInternalUseOnly, IsAlwaysAccessible, HasDefaults, HasNoOpConstructor, IsCoreType)
 struct FVector2f
 {
-	UPROPERTY(EditAnywhere, Category=Vector2D, SaveGame)
+	UPROPERTY(EditAnywhere, BlueprintReadWrite, Category=Vector2D, SaveGame)
 	float X;
 
-	UPROPERTY(EditAnywhere, Category=Vector2D, SaveGame)
+	UPROPERTY(EditAnywhere, BlueprintReadWrite, Category=Vector2D, SaveGame)
 	float Y;
 };
 
diff --git a/Engine/Source/Runtime/Engine/Classes/Engine/LocalPlayer.h b/Engine/Source/Runtime/Engine/Classes/Engine/LocalPlayer.h
--- a/Engine/Source/Runtime/Engine/Classes/Engine/LocalPlayer.h
+++ b/Engine/Source/Runtime/Engine/Classes/Engine/LocalPlayer.h
@@ -499,6 +499,17 @@ public:
 	UE_DEPRECATED(5.0, "Platform User Id now has priority over ControllerId, these are not expected to be the same")
 	bool IsCachedUniqueNetIdPairedWithControllerId() const;
 
+	struct FOptionalAllottedSize
+	{
+		FVector2f Value;
+
+		ENGINE_API FOptionalAllottedSize(std::nullptr_t Empty);
+		ENGINE_API FOptionalAllottedSize(const FVector2d* InVector2D);
+		ENGINE_API FOptionalAllottedSize(const FVector2f* InVector2D);
+
+		ENGINE_API explicit operator bool() const;
+	};
+
 	/**
 	 * This function will give you two points in Pixel Space that surround the World Space box.
 	 *
@@ -507,8 +518,13 @@ public:
 	 * @param	OutUpperRight	The Upper Right corner of the pixel space box
 	 * @return  False if there is no viewport, or if the box is behind the camera completely
 	 */
-	bool GetPixelBoundingBox(const FBox& ActorBox, FVector2D& OutLowerLeft, FVector2D& OutUpperRight, const FVector2D* OptionalAllotedSize = nullptr);
-	static bool GetPixelBoundingBox(const FSceneViewProjectionData& ProjectionData, const FBox& ActorBox, FVector2D& OutLowerLeft, FVector2D& OutUpperRight, const FVector2D* OptionalAllotedSize = nullptr);
+	bool GetPixelBoundingBox(const FBox& ActorBox, FVector2D& OutLowerLeft, FVector2D& OutUpperRight, const FVector2f* OptionalAllotedSize = nullptr);
+	static bool GetPixelBoundingBox(const FSceneViewProjectionData& ProjectionData, const FBox& ActorBox, FVector2D& OutLowerLeft, FVector2D& OutUpperRight, const FVector2f* OptionalAllotedSize = nullptr);
+
+	UE_DEPRECATED(5.2, "Please use const FVector2f* directly")
+	bool GetPixelBoundingBox(const FBox& ActorBox, FVector2D& OutLowerLeft, FVector2D& OutUpperRight, FOptionalAllottedSize OptionalAllotedSize);
+	UE_DEPRECATED(5.2, "Please use const FVector2f* directly")
+	static bool GetPixelBoundingBox(const FSceneViewProjectionData& ProjectionData, const FBox& ActorBox, FVector2D& OutLowerLeft, FVector2D& OutUpperRight, FOptionalAllottedSize OptionalAllotedSize);
 
 	/**
 	 * This function will give you a point in Pixel Space from a World Space position
@@ -517,8 +533,13 @@ public:
 	 * @param	OutPoint	The point in pixel space
 	 * @return  False if there is no viewport, or if the box is behind the camera completely
 	 */
-	bool GetPixelPoint(const FVector& InPoint, FVector2D& OutPoint, const FVector2D* OptionalAllotedSize = nullptr);
-	static bool GetPixelPoint(const FSceneViewProjectionData& ProjectionData, const FVector& InPoint, FVector2D& OutPoint, const FVector2D* OptionalAllotedSize = nullptr);
+	bool GetPixelPoint(const FVector& InPoint, FVector2D& OutPoint, const FVector2f* OptionalAllotedSize = nullptr);
+	static bool GetPixelPoint(const FSceneViewProjectionData& ProjectionData, const FVector& InPoint, FVector2D& OutPoint, const FVector2f* OptionalAllotedSize = nullptr);
+
+	UE_DEPRECATED(5.2, "Please use const FVector2f* directly")
+	bool GetPixelPoint(const FVector& InPoint, FVector2D& OutPoint, FOptionalAllottedSize OptionalAllotedSize);
+	UE_DEPRECATED(5.2, "Please use const FVector2f* directly")
+	static bool GetPixelPoint(const FSceneViewProjectionData& ProjectionData, const FVector& InPoint, FVector2D& OutPoint, FOptionalAllottedSize OptionalAllotedSize);
 
 	/**
 	 * Helper function for deriving various bits of data needed for projection
diff --git a/Engine/Source/Runtime/Engine/Private/Kismet/BlueprintTypeConversions.cpp b/Engine/Source/Runtime/Engine/Private/Kismet/BlueprintTypeConversions.cpp
--- a/Engine/Source/Runtime/Engine/Private/Kismet/BlueprintTypeConversions.cpp
+++ b/Engine/Source/Runtime/Engine/Private/Kismet/BlueprintTypeConversions.cpp
@@ -3,6 +3,7 @@
 #include "Kismet/BlueprintTypeConversions.h"
 #include "UObject/ScriptMacros.h"
 #include "UObject/Stack.h"
+#include "Types/SlateVector2.h"
 
 #include UE_INLINE_GENERATED_CPP_BY_NAME(BlueprintTypeConversions)
 
@@ -453,3 +454,14 @@ UE::Kismet::BlueprintTypeConversions::Internal::FStructConversionEntry FVector2D
 	&UE::Kismet::BlueprintTypeConversions::Internal::ConvertType<FVector2d, FVector2f>,
 	&UE::Kismet::BlueprintTypeConversions::Internal::ConvertType<FVector2f, FVector2d>
 );
+
+UE::Kismet::BlueprintTypeConversions::Internal::FStructConversionEntry FDeprecateVector2DEntry(
+	&FDeprecateSlateVector2D::StaticStruct,
+	&TVariantStructure<FVector2f>::Get,
+	&TBaseStructure<FVector2D>::Get,
+	&TVariantStructure<FVector2f>::Get,
+	TEXT(PREPROCESSOR_TO_STRING(MAKE_CONVERSION_FUNCTION_NAME(FVector2d, FVector2f))),
+	TEXT(PREPROCESSOR_TO_STRING(MAKE_CONVERSION_FUNCTION_NAME(FVector2f, FVector2d))),
+	&UE::Kismet::BlueprintTypeConversions::Internal::ConvertType<FVector2d, FVector2f>,
+	&UE::Kismet::BlueprintTypeConversions::Internal::ConvertType<FVector2f, FVector2d>
+);
\ No newline at end of file
diff --git a/Engine/Source/Runtime/Engine/Private/LocalPlayer.cpp b/Engine/Source/Runtime/Engine/Private/LocalPlayer.cpp
--- a/Engine/Source/Runtime/Engine/Private/LocalPlayer.cpp
+++ b/Engine/Source/Runtime/Engine/Private/LocalPlayer.cpp
@@ -911,7 +911,43 @@ FSceneView* ULocalPlayer::CalcSceneView( class FSceneViewFamily* ViewFamily,
 	return View;
 }
 
-bool ULocalPlayer::GetPixelBoundingBox(const FBox& ActorBox, FVector2D& OutLowerLeft, FVector2D& OutUpperRight, const FVector2D* OptionalAllotedSize)
+ULocalPlayer::FOptionalAllottedSize::FOptionalAllottedSize(std::nullptr_t Empty)
+	: Value(-std::numeric_limits<float>::infinity(), -std::numeric_limits<float>::infinity())
+{}
+
+ULocalPlayer::FOptionalAllottedSize::FOptionalAllottedSize(const FVector2d* InVector2D)
+{
+	if (InVector2D)
+	{
+		Value.X = static_cast<float>(InVector2D->X);
+		Value.Y = static_cast<float>(InVector2D->Y);
+	}
+	else
+	{
+		Value.X = -std::numeric_limits<float>::infinity();
+		Value.Y = -std::numeric_limits<float>::infinity();
+	}
+}
+
+ULocalPlayer::FOptionalAllottedSize::FOptionalAllottedSize(const FVector2f* InVector)
+{
+	if (InVector)
+	{
+		Value = *InVector;
+	}
+	else
+	{
+		Value.X = -std::numeric_limits<float>::infinity();
+		Value.Y = -std::numeric_limits<float>::infinity();
+	}
+}
+
+ULocalPlayer::FOptionalAllottedSize::operator bool() const
+{
+	return Value.X != -std::numeric_limits<float>::infinity();
+}
+
+bool ULocalPlayer::GetPixelBoundingBox(const FBox& ActorBox, FVector2D& OutLowerLeft, FVector2D& OutUpperRight, const FVector2f* OptionalAllotedSize)
 {
 	//@TODO: CAMERA: This has issues with aspect-ratio constrained cameras
 	if ((ViewportClient != NULL) && (ViewportClient->Viewport != NULL) && (PlayerController != NULL))
@@ -931,7 +967,7 @@ bool ULocalPlayer::GetPixelBoundingBox(const FBox& ActorBox, FVector2D& OutLower
 	}
 }
 
-bool ULocalPlayer::GetPixelBoundingBox(const FSceneViewProjectionData& ProjectionData, const FBox& ActorBox, FVector2D& OutLowerLeft, FVector2D& OutUpperRight, const FVector2D* OptionalAllotedSize)
+bool ULocalPlayer::GetPixelBoundingBox(const FSceneViewProjectionData& ProjectionData, const FBox& ActorBox, FVector2D& OutLowerLeft, FVector2D& OutUpperRight, const FVector2f* OptionalAllotedSize)
 {
 	// if we passed in an optional size, use it for the viewrect
 	FIntRect ViewRect = ProjectionData.GetConstrainedViewRect();
@@ -989,7 +1025,31 @@ bool ULocalPlayer::GetPixelBoundingBox(const FSceneViewProjectionData& Projectio
 	return SuccessCount >= 2;
 }
 
-bool ULocalPlayer::GetPixelPoint(const FVector& InPoint, FVector2D& OutPoint, const FVector2D* OptionalAllotedSize)
+bool ULocalPlayer::GetPixelBoundingBox(const FBox& ActorBox, FVector2D& OutLowerLeft, FVector2D& OutUpperRight, FOptionalAllottedSize OptionalAllotedSize)
+{
+	if (OptionalAllotedSize)
+	{
+		return GetPixelBoundingBox(ActorBox, OutLowerLeft, OutUpperRight, &OptionalAllotedSize.Value);
+	}
+	else
+	{
+		return GetPixelBoundingBox(ActorBox, OutLowerLeft, OutUpperRight);
+	}
+}
+
+bool ULocalPlayer::GetPixelBoundingBox(const FSceneViewProjectionData& ProjectionData, const FBox& ActorBox, FVector2D& OutLowerLeft, FVector2D& OutUpperRight, FOptionalAllottedSize OptionalAllotedSize)
+{
+	if (OptionalAllotedSize)
+	{
+		return GetPixelBoundingBox(ProjectionData, ActorBox, OutLowerLeft, OutUpperRight, &OptionalAllotedSize.Value);
+	}
+	else
+	{
+		return GetPixelBoundingBox(ProjectionData, ActorBox, OutLowerLeft, OutUpperRight);
+	}
+}
+
+bool ULocalPlayer::GetPixelPoint(const FVector& InPoint, FVector2D& OutPoint, const FVector2f* OptionalAllotedSize)
 {
 	//@TODO: CAMERA: This has issues with aspect-ratio constrained cameras
 	if ((ViewportClient != NULL) && (ViewportClient->Viewport != NULL) && (PlayerController != NULL))
@@ -1007,7 +1067,7 @@ bool ULocalPlayer::GetPixelPoint(const FVector& InPoint, FVector2D& OutPoint, co
 	return false;
 }
 
-bool ULocalPlayer::GetPixelPoint(const FSceneViewProjectionData& ProjectionData, const FVector& InPoint, FVector2D& OutPoint, const FVector2D* OptionalAllotedSize)
+bool ULocalPlayer::GetPixelPoint(const FSceneViewProjectionData& ProjectionData, const FVector& InPoint, FVector2D& OutPoint, const FVector2f* OptionalAllotedSize)
 {
 	bool bInFrontOfCamera = true;
 
@@ -1041,6 +1101,30 @@ bool ULocalPlayer::GetPixelPoint(const FSceneViewProjectionData& ProjectionData,
 	return bInFrontOfCamera;
 }
 
+bool ULocalPlayer::GetPixelPoint(const FVector& InPoint, FVector2D& OutPoint, FOptionalAllottedSize OptionalAllotedSize)
+{
+	if (OptionalAllotedSize)
+	{
+		return GetPixelPoint(InPoint, OutPoint, &OptionalAllotedSize.Value);
+	}
+	else
+	{
+		return GetPixelPoint(InPoint, OutPoint);
+	}
+}
+
+bool ULocalPlayer::GetPixelPoint(const FSceneViewProjectionData& ProjectionData, const FVector& InPoint, FVector2D& OutPoint, FOptionalAllottedSize OptionalAllotedSize)
+{
+	if (OptionalAllotedSize)
+	{
+		return GetPixelPoint(ProjectionData, InPoint, OutPoint, &OptionalAllotedSize.Value);
+	}
+	else
+	{
+		return GetPixelPoint(ProjectionData, InPoint, OutPoint);
+	}
+}
+
 bool ULocalPlayer::GetProjectionData(FViewport* Viewport, FSceneViewProjectionData& ProjectionData, int32 StereoViewIndex) const
 {
 	// If the actor
diff --git a/Engine/Source/Runtime/Engine/Private/Slate/DeferredCleanupSlateBrush.cpp b/Engine/Source/Runtime/Engine/Private/Slate/DeferredCleanupSlateBrush.cpp
--- a/Engine/Source/Runtime/Engine/Private/Slate/DeferredCleanupSlateBrush.cpp
+++ b/Engine/Source/Runtime/Engine/Private/Slate/DeferredCleanupSlateBrush.cpp
@@ -18,7 +18,7 @@ TSharedRef<FDeferredCleanupSlateBrush> FDeferredCleanupSlateBrush::CreateBrush(
 {
 	FSlateBrush Brush;
 	Brush.SetResourceObject(InTexture);
-	Brush.ImageSize = FVector2D(InTexture->GetSurfaceWidth(), InTexture->GetSurfaceHeight());
+	Brush.ImageSize = FVector2f(InTexture->GetSurfaceWidth(), InTexture->GetSurfaceHeight());
 	Brush.TintColor = InTint;
 	Brush.Margin = InMargin;
 	Brush.Tiling = InTiling;
@@ -30,7 +30,7 @@ TSharedRef<FDeferredCleanupSlateBrush> FDeferredCleanupSlateBrush::CreateBrush(
 
 TSharedRef<FDeferredCleanupSlateBrush> FDeferredCleanupSlateBrush::CreateBrush(
 	UObject* InResource,
-	const FVector2D& InImageSize,
+	const UE::Slate::FDeprecateVector2DParameter& InImageSize,
 	const FLinearColor& InTint,
 	ESlateBrushTileType::Type InTiling,
 	ESlateBrushImageType::Type InImageType,
diff --git a/Engine/Source/Runtime/Engine/Public/Slate/DeferredCleanupSlateBrush.h b/Engine/Source/Runtime/Engine/Public/Slate/DeferredCleanupSlateBrush.h
--- a/Engine/Source/Runtime/Engine/Public/Slate/DeferredCleanupSlateBrush.h
+++ b/Engine/Source/Runtime/Engine/Public/Slate/DeferredCleanupSlateBrush.h
@@ -24,7 +24,7 @@ public:
 
 	static TSharedRef<FDeferredCleanupSlateBrush> CreateBrush(
 		class UObject* InResource,
-		const FVector2D& InImageSize,
+		const UE::Slate::FDeprecateVector2DParameter& InImageSize,
 		const FLinearColor& InTint = FLinearColor(1.0f, 1.0f, 1.0f, 1.0f),
 		ESlateBrushTileType::Type InTiling = ESlateBrushTileType::NoTile,
 		ESlateBrushImageType::Type InImageType = ESlateBrushImageType::FullColor,
diff --git a/Engine/Source/Runtime/Slate/Private/Framework/Application/MenuStack.cpp b/Engine/Source/Runtime/Slate/Private/Framework/Application/MenuStack.cpp
--- a/Engine/Source/Runtime/Slate/Private/Framework/Application/MenuStack.cpp
+++ b/Engine/Source/Runtime/Slate/Private/Framework/Application/MenuStack.cpp
@@ -36,7 +36,7 @@ public:
 		SOverlay::Construct(SOverlay::FArguments());
 	}
 
-	void PushMenu(TSharedRef<FMenuBase> InMenu, const FVector2D& InLocation)
+	void PushMenu(TSharedRef<FMenuBase> InMenu, const FVector2f& InLocation)
 	{
 		check(InMenu->GetContent().IsValid());
 
@@ -44,10 +44,10 @@ public:
 		check(ParentWindow.IsValid());
 
 		// Transform InLocation into a position local to this panel (assumes the panel is in an overlay that covers the whole of the panel window) 
-		FVector2D PanelInScreen = ParentWindow->GetRectInScreen().GetTopLeft();
-		FVector2D PanelInWindow = ParentWindow->GetLocalToScreenTransform().Inverse().TransformPoint(PanelInScreen);
-		FVector2D LocationInWindow = ParentWindow->GetLocalToScreenTransform().Inverse().TransformPoint(InLocation);
-		FVector2D LocationInPanel = LocationInWindow - PanelInWindow;
+		FVector2f PanelInScreen = ParentWindow->GetRectInScreen().GetTopLeft();
+		FVector2f PanelInWindow = ParentWindow->GetLocalToScreenTransform().Inverse().TransformPoint(PanelInScreen);
+		FVector2f LocationInWindow = ParentWindow->GetLocalToScreenTransform().Inverse().TransformPoint(InLocation);
+		FVector2f LocationInPanel = LocationInWindow - PanelInWindow;
 
 		// Add the new menu into a slot on this panel and set the padding so that its position is correct
 		AddSlot()
@@ -181,7 +181,7 @@ namespace MenuStackInternal
 	FSimpleDelegate MenuStackPushDebuggingInfo;
 }	// anon namespace
 
-TSharedRef<IMenu> FMenuStack::Push(const FWidgetPath& InOwnerPath, const TSharedRef<SWidget>& InContent, const FVector2D& SummonLocation, const FPopupTransitionEffect& TransitionEffect, const bool bFocusImmediately, const FVector2D& SummonLocationSize, TOptional<EPopupMethod> InMethod, const bool bIsCollapsedByParent, const bool bEnablePerPixelTransparency)
+TSharedRef<IMenu> FMenuStack::Push(const FWidgetPath& InOwnerPath, const TSharedRef<SWidget>& InContent, const UE::Slate::FDeprecateVector2DParameter& SummonLocation, const FPopupTransitionEffect& TransitionEffect, const bool bFocusImmediately, const UE::Slate::FDeprecateVector2DParameter& SummonLocationSize, TOptional<EPopupMethod> InMethod, const bool bIsCollapsedByParent, const bool bEnablePerPixelTransparency)
 {
 	// We want to ensure that when the window is restored, we restore the current keyboard focus, 
 	// but only if it is valid, otherwise we could end up clearing a previously valid path.
@@ -241,7 +241,7 @@ TSharedRef<IMenu> FMenuStack::Push(const FWidgetPath& InOwnerPath, const TShared
 	return PushInternal(ParentMenu, InContent, Anchor, TransitionEffect, bFocusImmediately, ActiveMethod.GetShouldThrottle(), bIsCollapsedByParent, bEnablePerPixelTransparency);
 }
 
-TSharedRef<IMenu> FMenuStack::Push(const TSharedPtr<IMenu>& InParentMenu, const TSharedRef<SWidget>& InContent, const FVector2D& SummonLocation, const FPopupTransitionEffect& TransitionEffect, const bool bFocusImmediately, const FVector2D& SummonLocationSize, const bool bIsCollapsedByParent, const bool bEnablePerPixelTransparency)
+TSharedRef<IMenu> FMenuStack::Push(const TSharedPtr<IMenu>& InParentMenu, const TSharedRef<SWidget>& InContent, const UE::Slate::FDeprecateVector2DParameter& SummonLocation, const FPopupTransitionEffect& TransitionEffect, const bool bFocusImmediately, const UE::Slate::FDeprecateVector2DParameter& SummonLocationSize, const bool bIsCollapsedByParent, const bool bEnablePerPixelTransparency)
 {
 	check(Stack.Contains(InParentMenu));
 	check(HostWindow.IsValid());
@@ -377,12 +377,12 @@ FMenuStack::FPrePushResults FMenuStack::PrePush(const FPrePushArgs& InArgs)
 		// already handled
 		const bool bAutoAdjustForDPIScale = false;
 		// Places the menu's window in the work area
-		OutResults.AnimStartLocation = OutResults.AnimFinalLocation = FSlateApplication::Get().CalculatePopupWindowPosition(InArgs.Anchor, OutResults.ExpectedSize, bAutoAdjustForDPIScale, FVector2D::ZeroVector, Orientation);
+		OutResults.AnimStartLocation = OutResults.AnimFinalLocation = FSlateApplication::Get().CalculatePopupWindowPosition(InArgs.Anchor, OutResults.ExpectedSize, bAutoAdjustForDPIScale, FVector2f::ZeroVector, Orientation);
 	}
 	else
 	{
 		// Places the menu's content in the host window
-		const FVector2D ProposedPlacement(
+		const FVector2f ProposedPlacement(
 			Orientation == Orient_Horizontal ? InArgs.Anchor.Right : InArgs.Anchor.Left,
 			Orientation == Orient_Horizontal ? InArgs.Anchor.Top : InArgs.Anchor.Bottom);
 
diff --git a/Engine/Source/Runtime/Slate/Private/Framework/Application/SlateApplication.cpp b/Engine/Source/Runtime/Slate/Private/Framework/Application/SlateApplication.cpp
--- a/Engine/Source/Runtime/Slate/Private/Framework/Application/SlateApplication.cpp
+++ b/Engine/Source/Runtime/Slate/Private/Framework/Application/SlateApplication.cpp
@@ -502,10 +502,10 @@ DECLARE_CYCLE_STAT(TEXT("ProcessMouseMove"), STAT_ProcessMouseMove, STATGROUP_Sl
 namespace SlateDefs
 {
 	// How far tool tips should be offset from the mouse cursor position, in pixels
-	static const FVector2D ToolTipOffsetFromMouse( 12.0f, 8.0f );
+	static const FVector2f ToolTipOffsetFromMouse( 12.0f, 8.0f );
 
 	// How far tool tips should be pushed out from a force field border, in pixels
-	static const FVector2D ToolTipOffsetFromForceField( 4.0f, 3.0f );
+	static const FVector2f ToolTipOffsetFromForceField( 4.0f, 3.0f );
 
 	// Empty set of Touch Key
 	static TSet<FKey> EmptyTouchKeySet;
@@ -962,19 +962,19 @@ float FSlateApplication::GetSoundDuration(const FSlateSound& Sound) const
 	return SlateSoundDevice->GetSoundDuration(Sound);
 }
 
-FVector2D FSlateApplication::GetCursorPos() const
+UE::Slate::FDeprecateVector2DResult FSlateApplication::GetCursorPos() const
 {
 	return GetCursorUser()->GetCursorPosition();
 }
 
-FVector2D FSlateApplication::GetLastCursorPos() const
+UE::Slate::FDeprecateVector2DResult FSlateApplication::GetLastCursorPos() const
 {
 	return GetCursorUser()->GetPreviousCursorPosition();
 }
 
 void FSlateApplication::SetCursorPos(const FVector2D& MouseCoordinate)
 {
-	GetCursorUser()->SetCursorPosition(MouseCoordinate);
+	GetCursorUser()->SetCursorPosition(UE::Slate::CastToVector2f(MouseCoordinate));
 }
 
 void FSlateApplication::UsePlatformCursorForCursorUser(bool bUsePlatformCursor)
@@ -1003,7 +1003,7 @@ void FSlateApplication::SetPlatformCursorVisibility(bool bNewVisibility)
 	}
 }
 
-FWidgetPath FSlateApplication::LocateWindowUnderMouse( FVector2D ScreenspaceMouseCoordinate, const TArray< TSharedRef< SWindow > >& Windows, bool bIgnoreEnabledStatus, int32 UserIndex)
+FWidgetPath FSlateApplication::LocateWindowUnderMouse( UE::Slate::FDeprecateVector2DParameter ScreenspaceMouseCoordinate, const TArray< TSharedRef< SWindow > >& Windows, bool bIgnoreEnabledStatus, int32 UserIndex)
 {
 	// First, give the OS a chance to tell us which window to use, in case a child window is not guaranteed to stay on top of its parent window
 	TSharedPtr<FGenericWindow> NativeWindowUnderMouse = PlatformApplication->GetWindowUnderCursor();
@@ -1685,8 +1685,8 @@ void FSlateApplication::ThrottleApplicationBasedOnMouseMovement()
 		// Use a small movement threshold to avoid engaging the throttle when the user bumps the mouse
 		const float MinMouseMovePixelsBeforeThrottle = 2.0f;
 
-		const FVector2D& CursorPos = GetCursorPos();
-		static FVector2D LastCursorPos = GetCursorPos();
+		const FVector2f& CursorPos = GetCursorPos();
+		static FVector2f LastCursorPos = GetCursorPos();
 		//static double LastMouseMoveTime = FPlatformTime::Seconds();
 		static bool bIsMouseMoving = false;
 		if( CursorPos != LastCursorPos )
@@ -1748,19 +1748,19 @@ void FSlateApplication::ThrottleApplicationBasedOnMouseMovement()
 	}
 }
 
-FWidgetPath FSlateApplication::LocateWidgetInWindow(FVector2D ScreenspaceMouseCoordinate, const TSharedRef<SWindow>& Window, bool bIgnoreEnabledStatus, int32 UserIndex) const
+FWidgetPath FSlateApplication::LocateWidgetInWindow(UE::Slate::FDeprecateVector2DParameter ScreenspaceMouseCoordinate, const TSharedRef<SWindow>& Window, bool bIgnoreEnabledStatus, int32 UserIndex) const
 {
 	const bool bAcceptsInput = Window->IsVisible() && (Window->AcceptsInput() || IsWindowHousingInteractiveTooltip(Window));
 	if (bAcceptsInput && Window->IsScreenspaceMouseWithin(ScreenspaceMouseCoordinate))
 	{
-		FVector2D CursorPosition = ScreenspaceMouseCoordinate;
+		FVector2f CursorPosition = ScreenspaceMouseCoordinate;
 
 		if (TransformFullscreenMouseInput && !GIsEditor && Window->GetWindowMode() == EWindowMode::Fullscreen)
 		{
 			// Screen space mapping scales everything. When window resolution doesn't match platform resolution, 
 			// this causes offset cursor hit-tests in fullscreen. Correct in slate since we are first window-aware slate processor.
-			FVector2D WindowSize = Window->GetSizeInScreen();
-			FVector2D DisplaySize = { (float)CachedDisplayMetrics.PrimaryDisplayWidth, (float)CachedDisplayMetrics.PrimaryDisplayHeight };
+			FVector2f WindowSize = Window->GetSizeInScreen();
+			FVector2f DisplaySize = { (float)CachedDisplayMetrics.PrimaryDisplayWidth, (float)CachedDisplayMetrics.PrimaryDisplayHeight };
 
 			CursorPosition *= WindowSize / DisplaySize;
 		}
@@ -1823,11 +1823,11 @@ TSharedRef< FGenericWindow > FSlateApplication::MakeWindow( TSharedRef<SWindow>
 
 	Definition->Type = InSlateWindow->GetType();
 
-	const FVector2D Size = InSlateWindow->GetInitialDesiredSizeInScreen();
+	const FVector2f Size = InSlateWindow->GetInitialDesiredSizeInScreen();
 	Definition->WidthDesiredOnScreen = Size.X;
 	Definition->HeightDesiredOnScreen = Size.Y;
 
-	const FVector2D Position = InSlateWindow->GetInitialDesiredPositionInScreen();
+	const FVector2f Position = InSlateWindow->GetInitialDesiredPositionInScreen();
 	Definition->XDesiredPositionOnScreen = Position.X;
 	Definition->YDesiredPositionOnScreen = Position.Y;
 
@@ -2138,7 +2138,7 @@ TSharedRef<SWindow> FSlateApplication::AddWindowAsNativeChild( TSharedRef<SWindo
 	return InSlateWindow;
 }
 
-TSharedPtr<IMenu> FSlateApplication::PushMenu(const TSharedRef<SWidget>& InParentWidget, const FWidgetPath& InOwnerPath, const TSharedRef<SWidget>& InContent, const FVector2D& SummonLocation, const FPopupTransitionEffect& TransitionEffect, const bool bFocusImmediately, const FVector2D& SummonLocationSize, TOptional<EPopupMethod> Method, const bool bIsCollapsedByParent)
+TSharedPtr<IMenu> FSlateApplication::PushMenu(const TSharedRef<SWidget>& InParentWidget, const FWidgetPath& InOwnerPath, const TSharedRef<SWidget>& InContent, const UE::Slate::FDeprecateVector2DParameter& SummonLocation, const FPopupTransitionEffect& TransitionEffect, const bool bFocusImmediately, const UE::Slate::FDeprecateVector2DParameter& SummonLocationSize, TOptional<EPopupMethod> Method, const bool bIsCollapsedByParent)
 {
 	// Caller supplied a valid path? Pass it to the menu stack.
 	if (InOwnerPath.IsValid())
@@ -2157,7 +2157,7 @@ TSharedPtr<IMenu> FSlateApplication::PushMenu(const TSharedRef<SWidget>& InParen
 	return TSharedPtr<IMenu>();
 }
 
-TSharedPtr<IMenu> FSlateApplication::PushMenu(const TSharedPtr<IMenu>& InParentMenu, const TSharedRef<SWidget>& InContent, const FVector2D& SummonLocation, const FPopupTransitionEffect& TransitionEffect, const bool bFocusImmediately, const FVector2D& SummonLocationSize, const bool bIsCollapsedByParent)
+TSharedPtr<IMenu> FSlateApplication::PushMenu(const TSharedPtr<IMenu>& InParentMenu, const TSharedRef<SWidget>& InContent, const UE::Slate::FDeprecateVector2DParameter& SummonLocation, const FPopupTransitionEffect& TransitionEffect, const bool bFocusImmediately, const UE::Slate::FDeprecateVector2DParameter& SummonLocationSize, const bool bIsCollapsedByParent)
 {
 	return MenuStack.Push(InParentMenu, InContent, SummonLocation, TransitionEffect, bFocusImmediately, SummonLocationSize, bIsCollapsedByParent);
 }
@@ -3457,7 +3457,7 @@ void FSlateApplication::ProcessCursorReply(const FCursorReply& CursorReply)
 	GetCursorUser()->ProcessCursorReply(CursorReply);
 }
 
-void FSlateApplication::SpawnToolTip(const TSharedRef<IToolTip>& InToolTip, const FVector2D& InSpawnLocation)
+void FSlateApplication::SpawnToolTip(const TSharedRef<IToolTip>& InToolTip, const UE::Slate::FDeprecateVector2DParameter& InSpawnLocation)
 {
 	GetCursorUser()->ShowTooltip(InToolTip, InSpawnLocation);
 }
@@ -3657,12 +3657,12 @@ float FSlateApplication::GetDragTriggerDistanceSquared() const
 	return DragTriggerDistance * DragTriggerDistance;
 }
 
-bool FSlateApplication::HasTraveledFarEnoughToTriggerDrag(const FPointerEvent& PointerEvent, const FVector2D ScreenSpaceOrigin) const
+bool FSlateApplication::HasTraveledFarEnoughToTriggerDrag(const FPointerEvent& PointerEvent, const UE::Slate::FDeprecateVector2DParameter ScreenSpaceOrigin) const
 {
 	return ( PointerEvent.GetScreenSpacePosition() - ScreenSpaceOrigin ).SizeSquared() >= ( DragTriggerDistance * DragTriggerDistance );
 }
 
-bool FSlateApplication::HasTraveledFarEnoughToTriggerDrag(const FPointerEvent& PointerEvent, const FVector2D ScreenSpaceOrigin, EOrientation Orientation) const
+bool FSlateApplication::HasTraveledFarEnoughToTriggerDrag(const FPointerEvent& PointerEvent, const UE::Slate::FDeprecateVector2DParameter ScreenSpaceOrigin, EOrientation Orientation) const
 {
 	if (Orientation == Orient_Horizontal)
 	{
@@ -3720,11 +3720,11 @@ bool FSlateApplication::GetAllowTooltips() const
 	return bEnableTooltips;
 }
 
-FVector2D FSlateApplication::CalculateTooltipWindowPosition( const FSlateRect& InAnchorRect, const FVector2D& InSize, bool bAutoAdjustForDPIScale) const
+UE::Slate::FDeprecateVector2DResult FSlateApplication::CalculateTooltipWindowPosition( const FSlateRect& InAnchorRect, const UE::Slate::FDeprecateVector2DParameter& InSize, bool bAutoAdjustForDPIScale) const
 {
 	// first use the CalculatePopupWindowPosition and if cursor is not inside it, proceed with it to avoid behavior change.
-	FVector2D PopupPosition = CalculatePopupWindowPosition(InAnchorRect, InSize, bAutoAdjustForDPIScale);
-	FVector2D Cursor = GetCursorPos();
+	FVector2f PopupPosition = CalculatePopupWindowPosition(InAnchorRect, InSize, bAutoAdjustForDPIScale);
+	FVector2f Cursor = GetCursorPos();
 	if (PopupPosition.X > Cursor.X || PopupPosition.X + InSize.X < Cursor.X ||
 		PopupPosition.Y > Cursor.Y || PopupPosition.Y + InSize.Y < Cursor.Y)
 	{
@@ -3750,10 +3750,10 @@ FVector2D FSlateApplication::CalculateTooltipWindowPosition( const FSlateRect& I
 	// We want the Tooltip to appear in a 'comfortable' distance. The following vector: 'TooltipCursorOffset' 
 	// is used to move away from the cursor tip position. If we wouldn't do this the Tooltip would directly
 	// appear at the tip of the cursor. The coefficients 16 and 12 are estimated empirical.
-	const FVector2D TooltipCursorOffset(16 * DPIScale, 12 * DPIScale);
+	const FVector2f TooltipCursorOffset(16 * DPIScale, 12 * DPIScale);
 
 	// Calculate the new position of the Tooltip by starting at the Top/Left corner.
-	FVector2D ToolTipLocation = Cursor - TooltipCursorOffset - InSize;
+	FVector2f ToolTipLocation = Cursor - TooltipCursorOffset - InSize;
 
 	// Adjust the horizontal position so that it will be inside the work area.
 	if ( ToolTipLocation.X < WorkAreaRect.Left )
@@ -3770,9 +3770,9 @@ FVector2D FSlateApplication::CalculateTooltipWindowPosition( const FSlateRect& I
 	return ToolTipLocation;
 }
 
-FVector2D FSlateApplication::CalculatePopupWindowPosition(const FSlateRect& InAnchor, const FVector2D& InSize, bool bAutoAdjustForDPIScale, const FVector2D& InProposedPlacement, const EOrientation Orientation) const
+UE::Slate::FDeprecateVector2DResult FSlateApplication::CalculatePopupWindowPosition(const FSlateRect& InAnchor, const UE::Slate::FDeprecateVector2DParameter& InSize, bool bAutoAdjustForDPIScale, const UE::Slate::FDeprecateVector2DParameter& InProposedPlacement, const EOrientation Orientation) const
 {
-	FVector2D CalculatedPopUpWindowPosition(0, 0);
+	FVector2D CalculatedPopUpWindowPosition(0.f, 0.f);
 
 	float DPIScale = 1.0f;
 
@@ -3781,7 +3781,7 @@ FVector2D FSlateApplication::CalculatePopupWindowPosition(const FSlateRect& InAn
 		DPIScale = FPlatformApplicationMisc::GetDPIScaleFactorAtPoint(InAnchor.Left, InAnchor.Top);
 	}
 
-	FVector2D AdjustedSize = InSize * DPIScale;
+	FVector2f AdjustedSize = InSize * DPIScale;
 
 	FPlatformRect AnchorRect;
 	AnchorRect.Left = InAnchor.Left;
@@ -3796,9 +3796,9 @@ FVector2D FSlateApplication::CalculatePopupWindowPosition(const FSlateRect& InAn
 		PopUpOrientation = EPopUpOrientation::Vertical;
 	}
 
-	if (PlatformApplication->TryCalculatePopupWindowPosition(AnchorRect, AdjustedSize, InProposedPlacement, PopUpOrientation, /*OUT*/&CalculatedPopUpWindowPosition))
+	if (PlatformApplication->TryCalculatePopupWindowPosition(AnchorRect, FVector2D(AdjustedSize), FVector2D(InProposedPlacement), PopUpOrientation, /*OUT*/&CalculatedPopUpWindowPosition))
 	{
-		return CalculatedPopUpWindowPosition / DPIScale;
+		return UE::Slate::CastToVector2f(CalculatedPopUpWindowPosition / DPIScale);
 	}
 	else
 	{
@@ -3816,12 +3816,12 @@ FVector2D FSlateApplication::CalculatePopupWindowPosition(const FSlateRect& InAn
 			PlatformWorkArea.Right,
 			PlatformWorkArea.Bottom);
 
-		FVector2D ProposedPlacement = InProposedPlacement;
+		FVector2f ProposedPlacement = InProposedPlacement;
 
 		if (ProposedPlacement.IsZero())
 		{
 			// Assume natural left-to-right, top-to-bottom flow; position popup below and to the right.
-			ProposedPlacement = FVector2D(
+			ProposedPlacement = FVector2f(
 				Orientation == Orient_Horizontal ? AnchorRect.Right : AnchorRect.Left,
 				Orientation == Orient_Horizontal ? AnchorRect.Top : AnchorRect.Bottom);
 		}
@@ -3899,13 +3899,13 @@ FSlateRect FSlateApplication::GetPreferredWorkArea() const
 		// First see if we have a focused widget
 		if (FocusedWidgetPath.IsValid() && FocusedWidgetPath.Window.IsValid())
 		{
-			const FVector2D WindowPos = FocusedWidgetPath.Window.Pin()->GetPositionInScreen();
-			const FVector2D WindowSize = FocusedWidgetPath.Window.Pin()->GetSizeInScreen();
+			const FVector2f WindowPos = FocusedWidgetPath.Window.Pin()->GetPositionInScreen();
+			const FVector2f WindowSize = FocusedWidgetPath.Window.Pin()->GetSizeInScreen();
 			return GetWorkArea(FSlateRect(WindowPos.X, WindowPos.Y, WindowPos.X + WindowSize.X, WindowPos.Y + WindowSize.Y));
 		}
 
 		// no focus widget, so use cursor position if there are windows present in the work area
-		const FVector2D CursorPos = KeyboardUser->GetCursorPosition();
+		const FVector2f CursorPos = KeyboardUser->GetCursorPosition();
 		const FSlateRect WorkArea = GetWorkArea(FSlateRect(CursorPos.X, CursorPos.Y, CursorPos.X + 1.0f, CursorPos.Y + 1.0f));
 
 		if (FSlateWindowHelper::CheckWorkAreaForWindows(SlateWindows, WorkArea))
@@ -3989,9 +3989,9 @@ void TakeScreenshotCommon(const TSharedRef<SWidget>& Widget, const FIntRect& Inn
 	FSlateApplication::Get().GeneratePathToWidgetChecked(Widget, WidgetPath);
 
 	FArrangedWidget ArrangedWidget = WidgetPath.FindArrangedWidget(Widget).Get(FArrangedWidget::GetNullWidget());
-	FVector2D Position = FVector2D(ArrangedWidget.Geometry.AbsolutePosition);
-	FVector2D Size = ArrangedWidget.Geometry.GetDrawSize();
-	FVector2D WindowPosition = WidgetWindow->GetPositionInScreen();
+	FVector2f Position = FVector2f(ArrangedWidget.Geometry.AbsolutePosition);
+	FVector2f Size = ArrangedWidget.Geometry.GetDrawSize();
+	FVector2f WindowPosition = WidgetWindow->GetPositionInScreen();
 
 	ScreenshotRect = InnerWidgetArea.IsEmpty() ? FIntRect(0, 0, (int32)Size.X, (int32)Size.Y) : InnerWidgetArea;
 
@@ -4223,17 +4223,17 @@ void FSlateApplication::ForEachUser(TFunctionRef<void(FSlateUser*)> InPredicate,
 /* FSlateApplicationBase interface
  *****************************************************************************/
 
-FVector2D FSlateApplication::GetCursorSize( ) const
+UE::Slate::FDeprecateVector2DResult FSlateApplication::GetCursorSize( ) const
 {
 	if ( PlatformApplication->Cursor.IsValid() )
 	{
 		int32 X;
 		int32 Y;
 		PlatformApplication->Cursor->GetSize( X, Y );
-		return FVector2D( X, Y );
+		return FVector2f( static_cast<float>(X), static_cast<float>(Y) );
 	}
 
-	return FVector2D( 1.0f, 1.0f );
+	return FVector2f( 1.0f, 1.0f );
 }
 
 EVisibility FSlateApplication::GetSoftwareCursorVis( ) const
@@ -4284,8 +4284,8 @@ FPointerEvent FSlateApplication::TransformPointerEvent(const FPointerEvent& Poin
 		{
 			// Screen space mapping scales everything. When window resolution doesn't match platform resolution, 
 			// this causes offset cursor hit-tests in fullscreen. Correct in slate since we are first window-aware slate processor.
-			FVector2D WindowSize = Window->GetSizeInScreen();
-			FVector2D DisplaySize = { (float)CachedDisplayMetrics.PrimaryDisplayWidth, (float)CachedDisplayMetrics.PrimaryDisplayHeight };
+			FVector2f WindowSize = Window->GetSizeInScreen();
+			FVector2f DisplaySize = { (float)CachedDisplayMetrics.PrimaryDisplayWidth, (float)CachedDisplayMetrics.PrimaryDisplayHeight };
 
 			TransformedPointerEvent = FPointerEvent(PointerEvent, PointerEvent.GetScreenSpacePosition() * WindowSize / DisplaySize, PointerEvent.GetLastScreenSpacePosition() * WindowSize / DisplaySize);
 		}
@@ -5717,7 +5717,7 @@ bool FSlateApplication::ProcessMouseWheelOrGestureEvent( const FPointerEvent& In
 			bShouldProcessEvent = true;
 			break;
 		default:
-			bShouldProcessEvent = InGestureEvent->GetGestureDelta() != FVector2D::ZeroVector;
+			bShouldProcessEvent = InGestureEvent->GetGestureDelta() != FVector2f::ZeroVector;
 			break;
 		}
 	}
@@ -5806,8 +5806,8 @@ bool FSlateApplication::OnMouseMove()
 	}
 
 	bool Result = true;
-	const FVector2D CurrentCursorPosition = GetCursorPos();
-	const FVector2D LastCursorPosition = GetLastCursorPos();
+	const FVector2f CurrentCursorPosition = GetCursorPos();
+	const FVector2f LastCursorPosition = GetLastCursorPos();
 	
 	// Force the cursor user index to use the platform cursor since we've been notified that the platform 
 	// cursor position has changed. This is done intentionally after getting the positions in order to avoid
@@ -5872,7 +5872,7 @@ bool FSlateApplication::OnRawMouseMove( const int32 X, const int32 Y )
 			CursorPointerIndex,
 			GetCursorPos(),
 			GetLastCursorPos(),
-			FVector2D( X, Y ), 
+			FVector2f( X, Y ), 
 			PressedMouseButtons,
 			PlatformApplication->GetModifierKeys()
 		);
@@ -6124,7 +6124,7 @@ bool FSlateApplication::OnControllerButtonReleased(FGamepadKeyNames::Type KeyNam
 
 bool FSlateApplication::OnTouchGesture( EGestureEvent GestureType, const FVector2D &Delta, const float MouseWheelDelta, bool bIsDirectionInvertedFromDevice )
 {
-	const FVector2D CurrentCursorPosition = GetCursorPos();
+	const FVector2f CurrentCursorPosition = GetCursorPos();
 	
 	FPointerEvent GestureEvent(
 		CurrentCursorPosition,
@@ -6350,7 +6350,7 @@ bool FSlateApplication::OnSizeChanged( const TSharedRef< FGenericWindow >& Platf
 
 	if ( Window.IsValid() )
 	{
-		Window->SetCachedSize( FVector2D( Width, Height ) );
+		Window->SetCachedSize( FVector2f( Width, Height ) );
 
 		Renderer->RequestResize( Window, Width, Height );
 
@@ -6472,7 +6472,7 @@ void FSlateApplication::OnMovedWindow( const TSharedRef< FGenericWindow >& Platf
 
 	if ( Window.IsValid() )
 	{
-		Window->SetCachedScreenPosition( FVector2D( X, Y ) );
+		Window->SetCachedScreenPosition( FVector2f( X, Y ) );
 	}
 }
 
@@ -6723,7 +6723,7 @@ EWindowZone::Type FSlateApplication::GetWindowZoneForPoint( const TSharedRef< FG
 
 	if ( Window.IsValid() )
 	{
-		return Window->GetCurrentWindowZone( FVector2D( X, Y ) );
+		return Window->GetCurrentWindowZone( FVector2f( X, Y ) );
 	}
 
 	return EWindowZone::NotInWindow;
@@ -6840,8 +6840,8 @@ EDropEffect::Type FSlateApplication::OnDragEnter( const TSharedRef< SWindow >& W
 	// Assume we cannot handle it.
 	DragIsHandled = false;
 
-	const FVector2D CurrentCursorPosition = GetCursorPos();
-	const FVector2D LastCursorPosition = GetLastCursorPos();
+	const FVector2f CurrentCursorPosition = GetCursorPos();
+	const FVector2f LastCursorPosition = GetLastCursorPos();
 
 	// Tell slate to enter drag and drop mode.
 	// Make a faux mouse event for slate, so we can initiate a drag and drop.
@@ -6895,9 +6895,9 @@ EDropEffect::Type FSlateApplication::OnDragOver( const TSharedPtr< FGenericWindo
 	if (GetCursorUser()->IsDragDropping())
 	{
 		bool MouseMoveHandled = true;
-		FVector2D CursorMovementDelta( 0, 0 );
-		const FVector2D CurrentCursorPosition = GetCursorPos();
-		const FVector2D LastCursorPosition = GetLastCursorPos();
+		FVector2f CursorMovementDelta( 0, 0 );
+		const FVector2f CurrentCursorPosition = GetCursorPos();
+		const FVector2f LastCursorPosition = GetLastCursorPos();
 
 		if ( LastCursorPosition != CurrentCursorPosition )
 		{
diff --git a/Engine/Source/Runtime/Slate/Private/Framework/Application/SlateUser.cpp b/Engine/Source/Runtime/Slate/Private/Framework/Application/SlateUser.cpp
--- a/Engine/Source/Runtime/Slate/Private/Framework/Application/SlateUser.cpp
+++ b/Engine/Source/Runtime/Slate/Private/Framework/Application/SlateUser.cpp
@@ -19,10 +19,10 @@ DECLARE_CYCLE_STAT(TEXT("Update Tooltip Time"), STAT_SlateUpdateTooltip, STATGRO
 namespace SlateDefs
 {
 	// How far tooltips should be offset from the mouse cursor position, in pixels
-	static const FVector2D TooltipOffsetFromMouse(12.0f, 8.0f);
+	static const FVector2f TooltipOffsetFromMouse(12.0f, 8.0f);
 
 	// How far tooltips should be pushed out from a force field border, in pixels
-	static const FVector2D TooltipOffsetFromForceField(4.0f, 3.0f);
+	static const FVector2f TooltipOffsetFromForceField(4.0f, 3.0f);
 }
 
 static bool bEnableSyntheticCursorMoves = true;
@@ -117,8 +117,8 @@ FSlateUser::FSlateUser(int32 InUserIndex, TSharedPtr<ICursor> InCursor)
 	PlatformUser = FPlatformMisc::GetPlatformUserForUserIndex(InUserIndex);
 	UE_LOG(LogSlate, Log, TEXT("New Slate User Created. Platform User Id %d, User Index %d, Is Virtual User: %d"), PlatformUser.GetInternalId(), UserIndex, IsVirtualUser());
 
-	PointerPositionsByIndex.Add(FSlateApplication::CursorPointerIndex, FVector2D::ZeroVector);
-	PreviousPointerPositionsByIndex.Add(FSlateApplication::CursorPointerIndex, FVector2D::ZeroVector);
+	PointerPositionsByIndex.Add(FSlateApplication::CursorPointerIndex, FVector2f::ZeroVector);
+	PreviousPointerPositionsByIndex.Add(FSlateApplication::CursorPointerIndex, FVector2f::ZeroVector);
 }
 
 FSlateUser::FSlateUser(FPlatformUserId InPlatformUser, TSharedPtr<ICursor> InCursor)
@@ -130,8 +130,8 @@ FSlateUser::FSlateUser(FPlatformUserId InPlatformUser, TSharedPtr<ICursor> InCur
 	
 	UE_LOG(LogSlate, Log, TEXT("New Slate User Created.  Platform User Id %d,  Old User Index: %d  , Is Virtual User: %d"), PlatformUser.GetInternalId(), UserIndex, IsVirtualUser());
 	
-	PointerPositionsByIndex.Add(FSlateApplication::CursorPointerIndex, FVector2D::ZeroVector);
-	PreviousPointerPositionsByIndex.Add(FSlateApplication::CursorPointerIndex, FVector2D::ZeroVector);
+	PointerPositionsByIndex.Add(FSlateApplication::CursorPointerIndex, FVector2f::ZeroVector);
+	PreviousPointerPositionsByIndex.Add(FSlateApplication::CursorPointerIndex, FVector2f::ZeroVector);
 }
 
 FSlateUser::~FSlateUser()
@@ -379,7 +379,7 @@ void FSlateUser::SetCursorPosition(int32 PosX, int32 PosY)
 	SetPointerPosition(FSlateApplication::CursorPointerIndex, PosX, PosY);
 }
 
-void FSlateUser::SetCursorPosition(const FVector2D& NewCursorPos)
+void FSlateUser::SetCursorPosition(const UE::Slate::FDeprecateVector2DParameter& NewCursorPos)
 {
 	SetCursorPosition((int32)NewCursorPos.X, (int32)NewCursorPos.Y);
 }
@@ -393,37 +393,37 @@ void FSlateUser::SetPointerPosition(uint32 PointerIndex, int32 PosX, int32 PosY)
 		Cursor->SetPosition(PosX, PosY);
 	}
 
-	UpdatePointerPosition(PointerIndex, FVector2D(PosX, PosY));
+	UpdatePointerPosition(PointerIndex, FVector2f(PosX, PosY));
 }
 
-void FSlateUser::SetPointerPosition(uint32 PointerIndex, const FVector2D& NewPointerPos)
+void FSlateUser::SetPointerPosition(uint32 PointerIndex, const UE::Slate::FDeprecateVector2DParameter& NewPointerPos)
 {
 	SetPointerPosition(PointerIndex, (int32)NewPointerPos.X, (int32)NewPointerPos.Y);
 }
 
-FVector2D FSlateUser::GetCursorPosition() const
+UE::Slate::FDeprecateVector2DResult FSlateUser::GetCursorPosition() const
 {
 	return GetPointerPosition(FSlateApplication::CursorPointerIndex);
 }
 
-FVector2D FSlateUser::GetPreviousCursorPosition() const
+UE::Slate::FDeprecateVector2DResult FSlateUser::GetPreviousCursorPosition() const
 {
 	return GetPreviousPointerPosition(FSlateApplication::CursorPointerIndex);
 }
 
-FVector2D FSlateUser::GetPointerPosition(uint32 PointerIndex) const
+UE::Slate::FDeprecateVector2DResult FSlateUser::GetPointerPosition(uint32 PointerIndex) const
 {
 	if (Cursor && PointerIndex == FSlateApplication::CursorPointerIndex)
 	{
-		return Cursor->GetPosition();
+		return UE::Slate::CastToVector2f(Cursor->GetPosition());
 	}
-	const FVector2D* FoundPosition = PointerPositionsByIndex.Find(PointerIndex);
-	return FoundPosition ? *FoundPosition : FVector2D::ZeroVector;
+	const FVector2f* FoundPosition = PointerPositionsByIndex.Find(PointerIndex);
+	return FoundPosition ? *FoundPosition : FVector2f::ZeroVector;
 }
 
-FVector2D FSlateUser::GetPreviousPointerPosition(uint32 PointerIndex) const
+UE::Slate::FDeprecateVector2DResult FSlateUser::GetPreviousPointerPosition(uint32 PointerIndex) const
 {
-	const FVector2D* FoundPosition = PreviousPointerPositionsByIndex.Find(PointerIndex);
+	const FVector2f* FoundPosition = PreviousPointerPositionsByIndex.Find(PointerIndex);
 	return FoundPosition ? *FoundPosition : GetPointerPosition(PointerIndex);
 }
 
@@ -525,7 +525,7 @@ void FSlateUser::CancelDragDrop()
 	}
 }
 
-void FSlateUser::ShowTooltip(const TSharedRef<IToolTip>& InTooltip, const FVector2D& InLocation)
+void FSlateUser::ShowTooltip(const TSharedRef<IToolTip>& InTooltip, const UE::Slate::FDeprecateVector2DParameter& InLocation)
 {
 	CloseTooltip();
 
@@ -599,7 +599,7 @@ void FSlateUser::DrawWindowlessDragDropContent(const TSharedRef<SWindow>& Window
 				DecoratorWidget->SetVisibility(EVisibility::HitTestInvisible);
 				DecoratorWidget->SlatePrepass(WindowRootScale);
 
-				FVector2D DragDropContentInWindowSpace = WindowToDraw->GetWindowGeometryInScreen().AbsoluteToLocal(DragDropContent->GetDecoratorPosition()) * WindowRootScale;
+				FVector2f DragDropContentInWindowSpace = WindowToDraw->GetWindowGeometryInScreen().AbsoluteToLocal(DragDropContent->GetDecoratorPosition()) * WindowRootScale;
 				const FGeometry DragDropContentGeometry = FGeometry::MakeRoot(DecoratorWidget->GetDesiredSize(), FSlateLayoutTransform(DragDropContentInWindowSpace));
 
 				DecoratorWidget->Paint(
@@ -627,8 +627,8 @@ void FSlateUser::DrawCursor(const TSharedRef<SWindow>& WindowToDraw, FSlateWindo
 			CursorWidget->SetVisibility(EVisibility::HitTestInvisible);
 			CursorWidget->SlatePrepass(WindowRootScale);
 
-			FVector2D CursorInScreen = GetCursorPosition();
-			FVector2D CursorPosInWindowSpace = WindowToDraw->GetWindowGeometryInScreen().AbsoluteToLocal(CursorInScreen) * WindowRootScale;
+			FVector2f CursorInScreen = GetCursorPosition();
+			FVector2f CursorPosInWindowSpace = WindowToDraw->GetWindowGeometryInScreen().AbsoluteToLocal(CursorInScreen) * WindowRootScale;
 			CursorPosInWindowSpace += (CursorWidget->GetDesiredSize() * SoftwareCursorScale * -0.5);
 			const FGeometry CursorGeometry = FGeometry::MakeRoot(CursorWidget->GetDesiredSize() * SoftwareCursorScale, FSlateLayoutTransform(CursorPosInWindowSpace));
 
@@ -743,8 +743,8 @@ void FSlateUser::QueryCursor()
 		if (!CursorReply.IsEventHandled())
 		{
 			const bool bHasHardwareCursor = SlateApp.GetPlatformCursor() == Cursor;
-			const FVector2D CurrentCursorPosition = GetCursorPosition();
-			const FVector2D LastCursorPosition = GetPreviousCursorPosition();			
+			const FVector2f CurrentCursorPosition = GetCursorPosition();
+			const FVector2f LastCursorPosition = GetPreviousCursorPosition();			
 			
 			const TSet<FKey> EmptySet;
 			const FPointerEvent CursorEvent(
@@ -951,8 +951,8 @@ void FSlateUser::LockCursorInternal(const FWidgetPath& WidgetPath)
 		{
 			FDisplayMetrics CachedDisplayMetrics;
 			FSlateApplication::Get().GetCachedDisplayMetrics(CachedDisplayMetrics);
-			FVector2D DisplaySize = { (float)CachedDisplayMetrics.PrimaryDisplayWidth, (float)CachedDisplayMetrics.PrimaryDisplayHeight };
-			FVector2D DisplayDistortion = SlateClipRect.GetSize() / DisplaySize;
+			FVector2f DisplaySize = { (float)CachedDisplayMetrics.PrimaryDisplayWidth, (float)CachedDisplayMetrics.PrimaryDisplayHeight };
+			FVector2f DisplayDistortion = SlateClipRect.GetSize() / DisplaySize;
 
 			SlateClipRect.Left /= DisplayDistortion.X;
 			SlateClipRect.Top /= DisplayDistortion.Y;
@@ -1061,13 +1061,13 @@ void FSlateUser::UpdatePointerPosition(const FPointerEvent& PointerEvent)
 	UpdatePointerPosition(PointerEvent.GetPointerIndex(), PointerEvent.GetScreenSpacePosition());
 }
 
-void FSlateUser::UpdatePointerPosition(uint32 PointerIndex, const FVector2D& Position)
+void FSlateUser::UpdatePointerPosition(uint32 PointerIndex, const FVector2f& Position)
 {
 	PointerPositionsByIndex.FindOrAdd(PointerIndex) = Position;
 	PreviousPointerPositionsByIndex.FindOrAdd(PointerIndex) = Position;
 }
 
-void FSlateUser::StartDragDetection(const FWidgetPath& PathToWidget, int32 PointerIndex, FKey DragButton, FVector2D StartLocation)
+void FSlateUser::StartDragDetection(const FWidgetPath& PathToWidget, int32 PointerIndex, FKey DragButton, UE::Slate::FDeprecateVector2DParameter StartLocation)
 {
 	DragStatesByPointerIndex.Add(PointerIndex, FDragDetectionState(PathToWidget, PointerIndex, DragButton, StartLocation));
 }
@@ -1076,7 +1076,7 @@ FWidgetPath FSlateUser::DetectDrag(const FPointerEvent& PointerEvent, float Drag
 {
 	if (FDragDetectionState* DragState = DragStatesByPointerIndex.Find(PointerEvent.GetPointerIndex()))
 	{
-		const FVector2D DragDelta = DragState->DragStartLocation - PointerEvent.GetScreenSpacePosition();
+		const FVector2f DragDelta = DragState->DragStartLocation - PointerEvent.GetScreenSpacePosition();
 		if (DragDelta.SizeSquared() > FMath::Square(DragTriggerDistance))
 		{
 			FWidgetPath DragDetectionPath = DragState->DetectDragForWidget.ToWidgetPath(FWeakWidgetPath::EInterruptedPathHandling::ReturnInvalid);
@@ -1289,7 +1289,7 @@ void FSlateUser::UpdateTooltip(const FMenuStack& MenuStack, bool bCanSpawnNewToo
 		}
 	}
 
-	FVector2D DesiredLocation = ActiveTooltipInfo.DesiredLocation;
+	FVector2f DesiredLocation = ActiveTooltipInfo.DesiredLocation;
 	if ((ActiveTooltip && !ActiveTooltip->IsInteractive()) || (NewTooltip && NewTooltip != ActiveTooltip))
 	{
 		// New tooltips and non-interactive tooltips appear offset from the cursor position, and they follow the cursor as it moves.	
@@ -1298,7 +1298,9 @@ void FSlateUser::UpdateTooltip(const FMenuStack& MenuStack, bool bCanSpawnNewToo
 		// Allow interactive tooltips to adjust the window location
 		if (NewTooltip && NewTooltip->IsInteractive() && !NewTooltipVisualizer.IsValid())
 		{
-			NewTooltip->OnSetInteractiveWindowLocation(DesiredLocation);
+			FVector2D DesiredLocation2d(DesiredLocation);
+			NewTooltip->OnSetInteractiveWindowLocation(DesiredLocation2d);
+			DesiredLocation = UE::Slate::CastToVector2f(DesiredLocation2d);
 		}
 	}
 
@@ -1311,7 +1313,7 @@ void FSlateUser::UpdateTooltip(const FMenuStack& MenuStack, bool bCanSpawnNewToo
 	// Repel tooltip from a force field, if necessary
 	if (ForceFieldRect.IsSet())
 	{
-		FVector2D TooltipShift;
+		FVector2f TooltipShift;
 		TooltipShift.X = (ForceFieldRect->Right + SlateDefs::TooltipOffsetFromForceField.X) - DesiredLocation.X;
 		TooltipShift.Y = (ForceFieldRect->Bottom + SlateDefs::TooltipOffsetFromForceField.Y) - DesiredLocation.Y;
 
@@ -1377,12 +1379,12 @@ void FSlateUser::UpdateTooltip(const FMenuStack& MenuStack, bool bCanSpawnNewToo
 		TooltipWindow->SetOpacity(TooltipOpacity);
 
 		// How far tool tips should slide
-		const FVector2D SlideDistance(30.0f, 5.0f);
+		const FVector2f SlideDistance(30.0f, 5.0f);
 
 		// Apply steep inbound curve to the movement, so it looks like it quickly decelerating
 		const float SlideProgress = bAllowAnimations ? FMath::Pow(1.0f - TooltipOpacity, 3.0f) : 0.0f;
 
-		FVector2D WindowLocation = DesiredLocation + SlideProgress * SlideDistance;
+		FVector2f WindowLocation = DesiredLocation + SlideProgress * SlideDistance;
 		if (WindowLocation != TooltipWindow->GetPositionInScreen())
 		{
 			// already handled
diff --git a/Engine/Source/Runtime/Slate/Private/Framework/Docking/SDockingTabWell.cpp b/Engine/Source/Runtime/Slate/Private/Framework/Docking/SDockingTabWell.cpp
--- a/Engine/Source/Runtime/Slate/Private/Framework/Docking/SDockingTabWell.cpp
+++ b/Engine/Source/Runtime/Slate/Private/Framework/Docking/SDockingTabWell.cpp
@@ -173,7 +173,7 @@ int32 SDockingTabWell::OnPaint( const FPaintArgs& Args, const FGeometry& Allotte
 
 				// Center the separator
 				float Offset = (CurWidget.Geometry.GetLocalSize().Y - SeparatorHeight) / 2.0f;
-				FPaintGeometry Geometry = CurWidget.Geometry.ToPaintGeometry(FVector2D(CurWidget.Geometry.GetLocalSize().X + 1.0f, Offset), FVector2D(1.0f, SeparatorHeight));
+				FPaintGeometry Geometry = CurWidget.Geometry.ToPaintGeometry(FVector2f(1.0f, SeparatorHeight), FSlateLayoutTransform(FVector2f(CurWidget.Geometry.GetLocalSize().X + 1.0f, Offset)));
 		
 				// This code rounds the position of the widget so we don't end up on half a pixel and end up with a larger size separator than we want
 				FSlateRenderTransform NewTransform = Geometry.GetAccumulatedRenderTransform();
diff --git a/Engine/Source/Runtime/Slate/Private/Framework/Docking/STabDrawer.cpp b/Engine/Source/Runtime/Slate/Private/Framework/Docking/STabDrawer.cpp
--- a/Engine/Source/Runtime/Slate/Private/Framework/Docking/STabDrawer.cpp
+++ b/Engine/Source/Runtime/Slate/Private/Framework/Docking/STabDrawer.cpp
@@ -300,7 +300,7 @@ int32 STabDrawer::OnPaint(const FPaintArgs& Args, const FGeometry& AllottedGeome
 		ContentsLocalSize = FVector2D(LocalSize.X - (ShadowOffset.X * 2), TargetDrawerSize);
 	}
 
-	const FPaintGeometry OffsetPaintGeom = RenderTransformedChildGeometry.ToPaintGeometry(ContentsLocalOrigin, ContentsLocalSize);
+	const FPaintGeometry OffsetPaintGeom = RenderTransformedChildGeometry.ToPaintGeometry(ContentsLocalSize, FSlateLayoutTransform(ContentsLocalOrigin));
 
 	// Draw the resize handle
 	if (bIsResizing || bIsResizeHandleHovered)
@@ -386,12 +386,12 @@ int32 STabDrawer::OnPaint(const FPaintArgs& Args, const FGeometry& AllottedGeome
 		// Create the geometry for the notched portion, where one edge extends past the clipping rect.
 		const FVector2D NotchOffsetSize(TabButtonGeometry.GetLocalSize().X, 0.0f);
 		const FVector2D NotchOffsetTranslate = OpenDirection == ETabDrawerOpenDirection::Left ? -NotchOffsetSize : FVector2D::ZeroVector;
-		const FPaintGeometry NotchOffsetPaintGeom = RenderTransformedChildGeometry.ToPaintGeometry(ContentsLocalOrigin + NotchOffsetTranslate, ContentsLocalSize + NotchOffsetSize);
+		const FPaintGeometry NotchOffsetPaintGeom = RenderTransformedChildGeometry.ToPaintGeometry(ContentsLocalSize + NotchOffsetSize, FSlateLayoutTransform(ContentsLocalOrigin + NotchOffsetTranslate));
 
 		// Split the border box into three clipping zones.
-		const FPaintGeometry ClipAboveTabButton = RenderTransformedChildGeometry.ToPaintGeometry(FVector2D(0, 0), FVector2D(LocalSize.X, TabTopY));
-		const FPaintGeometry ClipAtTabButton = RenderTransformedChildGeometry.ToPaintGeometry(FVector2D(0, TabTopY), FVector2D(LocalSize.X, TabBottomY - TabTopY));
-		const FPaintGeometry ClipBelowTabButton = RenderTransformedChildGeometry.ToPaintGeometry(FVector2D(0, TabBottomY), FVector2D(LocalSize.X, LocalSize.Y - TabBottomY));
+		const FPaintGeometry ClipAboveTabButton = RenderTransformedChildGeometry.ToPaintGeometry(FVector2f(LocalSize.X, TabTopY), FSlateLayoutTransform(FVector2f(0, 0)));
+		const FPaintGeometry ClipAtTabButton = RenderTransformedChildGeometry.ToPaintGeometry(FVector2f(LocalSize.X, TabBottomY - TabTopY), FSlateLayoutTransform(FVector2f(0, TabTopY)));
+		const FPaintGeometry ClipBelowTabButton = RenderTransformedChildGeometry.ToPaintGeometry(FVector2f(LocalSize.X, LocalSize.Y - TabBottomY), FSlateLayoutTransform(FVector2f(0, TabBottomY)));
 
 		// If the tab button touches a corner on the edge of the border, switch the brush to
 		// draw that corner squared-off. When a tab is near the very top or bottom of its sidebar,
@@ -463,15 +463,24 @@ FGeometry STabDrawer::GetResizeHandleGeometry(const FGeometry& AllottedGeometry)
 
 	if (OpenDirection == ETabDrawerOpenDirection::Left)
 	{
-		return RenderTransformedGeometry.MakeChild((FVector2D(RenderTransformedGeometry.GetLocalSize().X-ShadowOffset.X, ShadowOffset.Y)), FVector2D(ExpanderSize, AllottedGeometry.GetLocalSize().Y - ShadowOffset.Y * 2));
+		return RenderTransformedGeometry.MakeChild(
+			FVector2D(ExpanderSize, AllottedGeometry.GetLocalSize().Y - ShadowOffset.Y * 2),
+			FSlateLayoutTransform(FVector2D(RenderTransformedGeometry.GetLocalSize().X-ShadowOffset.X, ShadowOffset.Y))
+		);
 	}
 	else if (OpenDirection == ETabDrawerOpenDirection::Right)
 	{
-		return RenderTransformedGeometry.MakeChild(ShadowOffset - FVector2D(ExpanderSize, 0.0f), FVector2D(ExpanderSize, AllottedGeometry.GetLocalSize().Y - ShadowOffset.Y * 2));
+		return RenderTransformedGeometry.MakeChild(
+			FVector2D(ExpanderSize, AllottedGeometry.GetLocalSize().Y - ShadowOffset.Y * 2),
+			FSlateLayoutTransform(ShadowOffset - FVector2D(ExpanderSize, 0.0f))
+		);
 	}
 	else
 	{
-		return RenderTransformedGeometry.MakeChild(ShadowOffset - FVector2D(0.0f, ExpanderSize), FVector2D(AllottedGeometry.GetLocalSize().X - ShadowOffset.X * 2, ExpanderSize));
+		return RenderTransformedGeometry.MakeChild(
+			FVector2D(AllottedGeometry.GetLocalSize().X - ShadowOffset.X * 2, ExpanderSize),
+			FSlateLayoutTransform(ShadowOffset - FVector2D(0.0f, ExpanderSize))
+		);
 	}
 }
 
