Slate: Initial phase of FVector2D deprecation

diff --git a/Engine/Source/Runtime/Slate/Private/Framework/Layout/ScrollyZoomy.cpp b/Engine/Source/Runtime/Slate/Private/Framework/Layout/ScrollyZoomy.cpp
--- a/Engine/Source/Runtime/Slate/Private/Framework/Layout/ScrollyZoomy.cpp
+++ b/Engine/Source/Runtime/Slate/Private/Framework/Layout/ScrollyZoomy.cpp
@@ -13,7 +13,7 @@
 FScrollyZoomy::FScrollyZoomy(const bool InUseInertialScrolling)
 	: AmountScrolledWhileRightMouseDown(0.0f)
 	, bShowSoftwareCursor(false)
-	, SoftwareCursorPosition(FVector2D::ZeroVector)
+	, SoftwareCursorPosition(FVector2f::ZeroVector)
 	, UseInertialScrolling(InUseInertialScrolling)
 { }
 
@@ -25,7 +25,7 @@ void FScrollyZoomy::Tick(const float DeltaTime, IScrollableZoomable& ScrollableZ
 {
 	if (!IsRightClickScrolling())
 	{
-		FVector2D ScrollBy = FVector2D::ZeroVector;
+		FVector2f ScrollBy = FVector2f::ZeroVector;
 
 		this->HorizontalIntertia.UpdateScrollVelocity(DeltaTime);
 		const float HorizontalScrollVelocity = this->HorizontalIntertia.GetScrollVelocity();
@@ -43,9 +43,9 @@ void FScrollyZoomy::Tick(const float DeltaTime, IScrollableZoomable& ScrollableZ
 			ScrollBy.Y = VerticalScrollVelocity * DeltaTime;
 		}
 
-		if (ScrollBy != FVector2D::ZeroVector)
+		if (ScrollBy != FVector2f::ZeroVector)
 		{
-			ScrollableZoomable.ScrollBy(ScrollBy);
+			ScrollableZoomable.ScrollBy(FVector2D(ScrollBy));
 		}
 	}
 }
@@ -87,7 +87,7 @@ FReply FScrollyZoomy::OnMouseButtonUp(const TSharedRef<SWidget> MyWidget, const
 		if (MyWidget->HasMouseCapture())
 		{
 			FSlateRect PanelScreenSpaceRect = MyGeometry.GetLayoutBoundingRect();
-			FVector2D CursorPosition = MyGeometry.LocalToAbsolute(SoftwareCursorPosition);
+			FVector2f CursorPosition = MyGeometry.LocalToAbsolute(SoftwareCursorPosition);
 
 			FIntPoint BestPositionInPanel(
 				FMath::RoundToInt(FMath::Clamp(CursorPosition.X, PanelScreenSpaceRect.Left, PanelScreenSpaceRect.Right)),
@@ -198,7 +198,7 @@ bool FScrollyZoomy::NeedsSoftwareCursor() const
 }
 
 
-const FVector2D& FScrollyZoomy::GetSoftwareCursorPosition() const
+UE::Slate::FDeprecateVector2DResult FScrollyZoomy::GetSoftwareCursorPosition() const
 {
 	return SoftwareCursorPosition;
 }
@@ -213,7 +213,7 @@ int32 FScrollyZoomy::PaintSoftwareCursorIfNeeded(const FGeometry& AllottedGeomet
 		FSlateDrawElement::MakeBox(
 			OutDrawElements,
 			++LayerId,
-			AllottedGeometry.ToPaintGeometry(SoftwareCursorPosition - (Brush->ImageSize / 2), Brush->ImageSize),
+			AllottedGeometry.ToPaintGeometry(Brush->ImageSize, FSlateLayoutTransform(SoftwareCursorPosition - (Brush->ImageSize / 2))),
 			Brush
 		);
 	}
diff --git a/Engine/Source/Runtime/Slate/Private/Framework/MultiBox/SToolBarButtonBlock.cpp b/Engine/Source/Runtime/Slate/Private/Framework/MultiBox/SToolBarButtonBlock.cpp
--- a/Engine/Source/Runtime/Slate/Private/Framework/MultiBox/SToolBarButtonBlock.cpp
+++ b/Engine/Source/Runtime/Slate/Private/Framework/MultiBox/SToolBarButtonBlock.cpp
@@ -207,9 +207,9 @@ void SToolBarButtonBlock::BuildMultiBlockWidget(const ISlateStyle* StyleSet, con
 	TSharedRef<SWidget> ButtonContent = SNullWidget::NullWidget;
 	if (MultiBox->GetType() == EMultiBoxType::SlimHorizontalToolBar)
 	{
-		const FVector2D IconSize = ToolBarStyle.IconSize;
+		const FVector2f IconSize = ToolBarStyle.IconSize;
 
-		IconWidget->SetDesiredSizeOverride(IconSize);
+		IconWidget->SetDesiredSizeOverride(FVector2D(IconSize));
 		ButtonContent =
 			SNew(SHorizontalBox)
 			.AddMetaData<FTagMetaData>(FTagMetaData(TutorialHighlightName))
diff --git a/Engine/Source/Runtime/Slate/Private/Framework/MultiBox/SToolBarComboButtonBlock.cpp b/Engine/Source/Runtime/Slate/Private/Framework/MultiBox/SToolBarComboButtonBlock.cpp
--- a/Engine/Source/Runtime/Slate/Private/Framework/MultiBox/SToolBarComboButtonBlock.cpp
+++ b/Engine/Source/Runtime/Slate/Private/Framework/MultiBox/SToolBarComboButtonBlock.cpp
@@ -95,7 +95,7 @@ void SToolBarComboButtonBlock::BuildMultiBlockWidget(const ISlateStyle* StyleSet
 
 		if (MultiBox->GetType() == EMultiBoxType::SlimHorizontalToolBar)
 		{
-			const FVector2D IconSize = ToolBarStyle.IconSize;
+			const FVector2f IconSize = ToolBarStyle.IconSize;
 
 			IconWidget =
 				SNew(SBox)
diff --git a/Engine/Source/Runtime/Slate/Private/Framework/MultiBox/SToolBarStackButtonBlock.cpp b/Engine/Source/Runtime/Slate/Private/Framework/MultiBox/SToolBarStackButtonBlock.cpp
--- a/Engine/Source/Runtime/Slate/Private/Framework/MultiBox/SToolBarStackButtonBlock.cpp
+++ b/Engine/Source/Runtime/Slate/Private/Framework/MultiBox/SToolBarStackButtonBlock.cpp
@@ -282,9 +282,9 @@ void SToolBarStackButtonBlock::BuildMultiBlockWidget(const ISlateStyle* StyleSet
 	TSharedRef<SWidget> ButtonContent = SNullWidget::NullWidget;
 	if (MultiBox->GetType() == EMultiBoxType::SlimHorizontalToolBar)
 	{
-		const FVector2D IconSize = ToolBarStyle.IconSize;
+		const FVector2f IconSize = ToolBarStyle.IconSize;
 
-		IconWidget->SetDesiredSizeOverride(IconSize);
+		IconWidget->SetDesiredSizeOverride(FVector2D(IconSize));
 		ButtonContent =
 			SNew(SHorizontalBox)
 			.AddMetaData<FTagMetaData>(FTagMetaData(TagName))
diff --git a/Engine/Source/Runtime/Slate/Private/Framework/Text/SlateImageRun.cpp b/Engine/Source/Runtime/Slate/Private/Framework/Text/SlateImageRun.cpp
--- a/Engine/Source/Runtime/Slate/Private/Framework/Text/SlateImageRun.cpp
+++ b/Engine/Source/Runtime/Slate/Private/Framework/Text/SlateImageRun.cpp
@@ -68,7 +68,7 @@ FSlateImageRun::FSlateImageRun( const FRunInfo& InRunInfo, const TSharedRef< con
 	, Baseline( InBaseline )
 {
 	FIntPoint Size = FSlateApplication::Get().GetRenderer()->GenerateDynamicImageResource(InDynamicBrushName);
-	DynamicBrush = MakeShareable(new FSlateDynamicImageBrush( InDynamicBrushName, FVector2D(Size.X, Size.Y) ) );
+	DynamicBrush = MakeShareable(new FSlateDynamicImageBrush( InDynamicBrushName, FVector2f(Size.X, Size.Y) ) );
 	Image = DynamicBrush.Get();
 }
 
@@ -80,7 +80,7 @@ FSlateImageRun::FSlateImageRun( const FRunInfo& InRunInfo, const TSharedRef< con
 	, Baseline( InBaseline )
 {
 	FIntPoint Size = FSlateApplication::Get().GetRenderer()->GenerateDynamicImageResource(InDynamicBrushName);
-	DynamicBrush = MakeShareable(new FSlateDynamicImageBrush( InDynamicBrushName, FVector2D(Size.X, Size.Y) ) );
+	DynamicBrush = MakeShareable(new FSlateDynamicImageBrush( InDynamicBrushName, FVector2f(Size.X, Size.Y) ) );
 	Image = DynamicBrush.Get();
 }
 
@@ -123,7 +123,7 @@ int32 FSlateImageRun::GetTextIndexAt( const TSharedRef< ILayoutBlock >& Block, c
 		return INDEX_NONE;
 	}
 
-	const FVector2D ScaledImageSize = Image->ImageSize * Scale;
+	const FVector2f ScaledImageSize = Image->ImageSize * Scale;
 	const int32 Index = (Location.X <= (Left + (ScaledImageSize.X * 0.5f))) ? Range.BeginIndex : Range.EndIndex;
 	
 	if (OutHitPoint)
@@ -182,7 +182,7 @@ FVector2D FSlateImageRun::Measure( int32 BeginIndex, int32 EndIndex, float Scale
 		return FVector2D( 0, GetMaxHeight( Scale ) );
 	}
 
-	return Image->ImageSize * Scale;
+	return FVector2D(Image->ImageSize * Scale);
 }
 
 int16 FSlateImageRun::GetMaxHeight( float Scale ) const 
diff --git a/Engine/Source/Runtime/Slate/Private/Framework/Text/SlateTextUnderlineLineHighlighter.cpp b/Engine/Source/Runtime/Slate/Private/Framework/Text/SlateTextUnderlineLineHighlighter.cpp
--- a/Engine/Source/Runtime/Slate/Private/Framework/Text/SlateTextUnderlineLineHighlighter.cpp
+++ b/Engine/Source/Runtime/Slate/Private/Framework/Text/SlateTextUnderlineLineHighlighter.cpp
@@ -6,7 +6,7 @@
 #include "Fonts/FontCache.h"
 #include "Framework/Application/SlateApplication.h"
 
-ISlateTextLineHighlighter::ISlateTextLineHighlighter(const FSlateBrush& InLineBrush, const FSlateFontInfo& InFontInfo, const FSlateColor InColorAndOpacity, const FVector2D InShadowOffset, const FLinearColor InShadowColorAndOpacity)
+ISlateTextLineHighlighter::ISlateTextLineHighlighter(const FSlateBrush& InLineBrush, const FSlateFontInfo& InFontInfo, const FSlateColor InColorAndOpacity, const UE::Slate::FDeprecateVector2DParameter InShadowOffset, const FLinearColor InShadowColorAndOpacity)
 	: LineBrush(InLineBrush)
 	, FontInfo(InFontInfo)
 	, ColorAndOpacity(InColorAndOpacity)
@@ -25,8 +25,8 @@ int32 ISlateTextLineHighlighter::OnPaint(const FPaintArgs& Args, const FTextLayo
 	int16 LinePos, LineThickness;
 	GetLineMetrics(AllottedGeometry.Scale, LinePos, LineThickness);
 
-	const FVector2D Location(Line.Offset.X + OffsetX, Line.Offset.Y + MaxHeight + Baseline - (LinePos * 0.5f));
-	const FVector2D Size(Width, FMath::Max<int16>(1, LineThickness));
+	const FVector2f Location(Line.Offset.X + OffsetX, Line.Offset.Y + MaxHeight + Baseline - (LinePos * 0.5f));
+	const FVector2f Size(Width, FMath::Max<int16>(1, LineThickness));
 
 	// The block size and offset values are pre-scaled, so we need to account for that when converting the block offsets into paint geometry
 	const float InverseScale = Inverse(AllottedGeometry.Scale);
@@ -38,11 +38,11 @@ int32 ISlateTextLineHighlighter::OnPaint(const FPaintArgs& Args, const FTextLayo
 		const bool ShouldDropShadow = ShadowColorAndOpacity.A > 0.f && ShadowOffset.SizeSquared() > 0.f;
 
 		// A negative shadow offset should be applied as a positive offset to the underline to avoid clipping issues
-		const FVector2D DrawShadowOffset(
+		const FVector2f DrawShadowOffset(
 			(ShadowOffset.X > 0.0f) ? ShadowOffset.X * AllottedGeometry.Scale : 0.0f,
 			(ShadowOffset.Y > 0.0f) ? ShadowOffset.Y * AllottedGeometry.Scale : 0.0f
 			);
-		const FVector2D DrawUnderlineOffset(
+		const FVector2f DrawUnderlineOffset(
 			(ShadowOffset.X < 0.0f) ? -ShadowOffset.X * AllottedGeometry.Scale : 0.0f,
 			(ShadowOffset.Y < 0.0f) ? -ShadowOffset.Y * AllottedGeometry.Scale : 0.0f
 			);
@@ -74,12 +74,12 @@ int32 ISlateTextLineHighlighter::OnPaint(const FPaintArgs& Args, const FTextLayo
 	return LayerId;
 }
 
-FSlateTextUnderlineLineHighlighter::FSlateTextUnderlineLineHighlighter(const FSlateBrush& InUnderlineBrush, const FSlateFontInfo& InFontInfo, const FSlateColor InColorAndOpacity, const FVector2D InShadowOffset, const FLinearColor InShadowColorAndOpacity)
+FSlateTextUnderlineLineHighlighter::FSlateTextUnderlineLineHighlighter(const FSlateBrush& InUnderlineBrush, const FSlateFontInfo& InFontInfo, const FSlateColor InColorAndOpacity, const UE::Slate::FDeprecateVector2DParameter InShadowOffset, const FLinearColor InShadowColorAndOpacity)
 	: ISlateTextLineHighlighter(InUnderlineBrush, InFontInfo, InColorAndOpacity, InShadowOffset, InShadowColorAndOpacity)
 {
 }
 
-TSharedRef<FSlateTextUnderlineLineHighlighter> FSlateTextUnderlineLineHighlighter::Create(const FSlateBrush& InUnderlineBrush, const FSlateFontInfo& InFontInfo, const FSlateColor InColorAndOpacity, const FVector2D InShadowOffset, const FLinearColor InShadowColorAndOpacity)
+TSharedRef<FSlateTextUnderlineLineHighlighter> FSlateTextUnderlineLineHighlighter::Create(const FSlateBrush& InUnderlineBrush, const FSlateFontInfo& InFontInfo, const FSlateColor InColorAndOpacity, const UE::Slate::FDeprecateVector2DParameter InShadowOffset, const FLinearColor InShadowColorAndOpacity)
 {
 	return MakeShareable(new FSlateTextUnderlineLineHighlighter(InUnderlineBrush, InFontInfo, InColorAndOpacity, InShadowOffset, InShadowColorAndOpacity));
 }
@@ -90,12 +90,12 @@ void FSlateTextUnderlineLineHighlighter::GetLineMetrics(const float InFontScale,
 	FontCache->GetUnderlineMetrics(FontInfo, InFontScale, OutLinePos, OutLineThickness);
 }
 
-FSlateTextStrikeLineHighlighter::FSlateTextStrikeLineHighlighter(const FSlateBrush& InUnderlineBrush, const FSlateFontInfo& InFontInfo, const FSlateColor InColorAndOpacity, const FVector2D InShadowOffset, const FLinearColor InShadowColorAndOpacity)
+FSlateTextStrikeLineHighlighter::FSlateTextStrikeLineHighlighter(const FSlateBrush& InUnderlineBrush, const FSlateFontInfo& InFontInfo, const FSlateColor InColorAndOpacity, const UE::Slate::FDeprecateVector2DParameter InShadowOffset, const FLinearColor InShadowColorAndOpacity)
 	: ISlateTextLineHighlighter(InUnderlineBrush, InFontInfo, InColorAndOpacity, InShadowOffset, InShadowColorAndOpacity)
 {
 }
 
-TSharedRef<FSlateTextStrikeLineHighlighter> FSlateTextStrikeLineHighlighter::Create(const FSlateBrush& InUnderlineBrush, const FSlateFontInfo& InFontInfo, const FSlateColor InColorAndOpacity, const FVector2D InShadowOffset, const FLinearColor InShadowColorAndOpacity)
+TSharedRef<FSlateTextStrikeLineHighlighter> FSlateTextStrikeLineHighlighter::Create(const FSlateBrush& InUnderlineBrush, const FSlateFontInfo& InFontInfo, const FSlateColor InColorAndOpacity, const UE::Slate::FDeprecateVector2DParameter InShadowOffset, const FLinearColor InShadowColorAndOpacity)
 {
 	return MakeShareable(new FSlateTextStrikeLineHighlighter(InUnderlineBrush, InFontInfo, InColorAndOpacity, InShadowOffset, InShadowColorAndOpacity));
 }
diff --git a/Engine/Source/Runtime/Slate/Private/Widgets/Colors/SColorGradingWheel.cpp b/Engine/Source/Runtime/Slate/Private/Widgets/Colors/SColorGradingWheel.cpp
--- a/Engine/Source/Runtime/Slate/Private/Widgets/Colors/SColorGradingWheel.cpp
+++ b/Engine/Source/Runtime/Slate/Private/Widgets/Colors/SColorGradingWheel.cpp
@@ -77,7 +77,7 @@ FVector2D SColorGradingWheel::ComputeDesiredSize(float) const
 		int32 CachedDesiredWheelSize = DesiredWheelSizeAttribute.Get();
 		return FVector2D(CachedDesiredWheelSize, CachedDesiredWheelSize);
 	}
-	return Image->ImageSize;
+	return FVector2D(Image->ImageSize);
 }
 
 
@@ -148,9 +148,9 @@ int32 SColorGradingWheel::OnPaint(const FPaintArgs& Args, const FGeometry& Allot
 {
 	const bool bIsEnabled = ShouldBeEnabled(bParentEnabled);
 	const ESlateDrawEffect DrawEffects = bIsEnabled ? ESlateDrawEffect::None : ESlateDrawEffect::DisabledEffect;
-	const FVector2D& SelectorSize = SelectorImage->ImageSize;
-	FVector2D CircleSize = AllottedGeometry.GetLocalSize() - SelectorSize;
-	FVector2D AllottedGeometrySize = AllottedGeometry.GetLocalSize();
+	const FVector2f& SelectorSize = SelectorImage->ImageSize;
+	FVector2f CircleSize = AllottedGeometry.GetLocalSize() - SelectorSize;
+	FVector2f AllottedGeometrySize = AllottedGeometry.GetLocalSize();
 	if (bIsAttributeDesiredWheelSizeSet)
 	{
 		int32 CachedDesiredWheelSize = DesiredWheelSizeAttribute.Get();
@@ -163,7 +163,7 @@ int32 SColorGradingWheel::OnPaint(const FPaintArgs& Args, const FGeometry& Allot
 	FSlateDrawElement::MakeBox(
 		OutDrawElements,
 		LayerId,
-		AllottedGeometry.ToPaintGeometry(0.5 * SelectorSize, CircleSize),
+		AllottedGeometry.ToPaintGeometry(CircleSize, FSlateLayoutTransform(0.5f * SelectorSize)),
 		Image,
 		DrawEffects,
 		InWidgetStyle.GetColorAndOpacityTint() * Image->GetTint(InWidgetStyle)
@@ -172,7 +172,7 @@ int32 SColorGradingWheel::OnPaint(const FPaintArgs& Args, const FGeometry& Allot
 	FSlateDrawElement::MakeBox(
 		OutDrawElements,
 		LayerId + 1,
-		AllottedGeometry.ToPaintGeometry(0.5f * (AllottedGeometrySize + CalcRelativePositionFromCenter() * CircleSize - SelectorSize), SelectorSize),
+		AllottedGeometry.ToPaintGeometry(SelectorSize, FSlateLayoutTransform(0.5f * (AllottedGeometrySize + CalcRelativePositionFromCenter() * CircleSize - SelectorSize))),
 		SelectorImage,
 		DrawEffects,
 		InWidgetStyle.GetColorAndOpacityTint() * SelectorImage->GetTint(InWidgetStyle)
@@ -185,7 +185,7 @@ int32 SColorGradingWheel::OnPaint(const FPaintArgs& Args, const FGeometry& Allot
 /* SColorGradingWheel implementation
  *****************************************************************************/
 
-FVector2D SColorGradingWheel::CalcRelativePositionFromCenter() const
+UE::Slate::FDeprecateVector2DResult SColorGradingWheel::CalcRelativePositionFromCenter() const
 {
 	float Hue = SelectedColorAttribute.Get().R;
 	float Saturation = SelectedColorAttribute.Get().G;
@@ -197,13 +197,13 @@ FVector2D SColorGradingWheel::CalcRelativePositionFromCenter() const
 	float Angle = Hue / 180.0f * PI;
 	float Radius = Saturation;
 
-	return FVector2D(FMath::Cos(Angle), FMath::Sin(Angle)) * Radius;
+	return FVector2f(FMath::Cos(Angle), FMath::Sin(Angle)) * Radius;
 }
 
 
 bool SColorGradingWheel::ProcessMouseAction(const FGeometry& MyGeometry, const FPointerEvent& MouseEvent, bool bProcessWhenOutsideColorWheel)
 {
-	FVector2D GeometrySize = MyGeometry.GetLocalSize();
+	FVector2f GeometrySize = MyGeometry.GetLocalSize();
 	if (bIsAttributeDesiredWheelSizeSet)
 	{
 		DesiredWheelSizeAttribute.UpdateNow(*this);
@@ -212,8 +212,8 @@ bool SColorGradingWheel::ProcessMouseAction(const FGeometry& MyGeometry, const F
 		GeometrySize.Y = CachedDesiredWheelSize;
 	}
 
-	const FVector2D LocalMouseCoordinate = MyGeometry.AbsoluteToLocal(MouseEvent.GetScreenSpacePosition());
-	const FVector2D RelativePositionFromCenter = (2.0f * LocalMouseCoordinate - GeometrySize) / (GeometrySize - SelectorImage->ImageSize);
+	const FVector2f LocalMouseCoordinate = MyGeometry.AbsoluteToLocal(MouseEvent.GetScreenSpacePosition());
+	const FVector2f RelativePositionFromCenter = (2.0f * LocalMouseCoordinate - GeometrySize) / (GeometrySize - SelectorImage->ImageSize);
 	const float RelativeRadius = RelativePositionFromCenter.Size();
 
 	if (RelativeRadius <= 1.0f || bProcessWhenOutsideColorWheel)
diff --git a/Engine/Source/Runtime/Slate/Private/Widgets/Colors/SColorSpectrum.cpp b/Engine/Source/Runtime/Slate/Private/Widgets/Colors/SColorSpectrum.cpp
--- a/Engine/Source/Runtime/Slate/Private/Widgets/Colors/SColorSpectrum.cpp
+++ b/Engine/Source/Runtime/Slate/Private/Widgets/Colors/SColorSpectrum.cpp
@@ -30,7 +30,7 @@ void SColorSpectrum::Construct( const FArguments& InArgs )
 
 FVector2D SColorSpectrum::ComputeDesiredSize( float ) const
 {
-	return Image->ImageSize;
+	return FVector2D(Image->ImageSize);
 }
 
 
@@ -106,7 +106,7 @@ int32 SColorSpectrum::OnPaint( const FPaintArgs& Args, const FGeometry& Allotted
 	FSlateDrawElement::MakeBox(
 		OutDrawElements,
 		LayerId + 1,
-		AllottedGeometry.ToPaintGeometry(CalcRelativeSelectedPosition() * AllottedGeometry.Size - SelectorImage->ImageSize * 0.5f, SelectorImage->ImageSize),
+		AllottedGeometry.ToPaintGeometry(SelectorImage->ImageSize, FSlateLayoutTransform(CalcRelativeSelectedPosition() * AllottedGeometry.Size - SelectorImage->ImageSize * 0.5f)),
 		SelectorImage,
 		DrawEffects,
 		InWidgetStyle.GetColorAndOpacityTint() * SelectorImage->GetTint(InWidgetStyle));
diff --git a/Engine/Source/Runtime/Slate/Private/Widgets/Colors/SColorWheel.cpp b/Engine/Source/Runtime/Slate/Private/Widgets/Colors/SColorWheel.cpp
--- a/Engine/Source/Runtime/Slate/Private/Widgets/Colors/SColorWheel.cpp
+++ b/Engine/Source/Runtime/Slate/Private/Widgets/Colors/SColorWheel.cpp
@@ -28,7 +28,7 @@ void SColorWheel::Construct(const FArguments& InArgs)
 
 FVector2D SColorWheel::ComputeDesiredSize(float) const
 {
-	return Image->ImageSize + SelectorImage->ImageSize;
+	return FVector2D(Image->ImageSize + SelectorImage->ImageSize);
 }
 
 
@@ -87,13 +87,13 @@ int32 SColorWheel::OnPaint(const FPaintArgs& Args, const FGeometry& AllottedGeom
 {
 	const bool bIsEnabled = ShouldBeEnabled(bParentEnabled);
 	const ESlateDrawEffect DrawEffects = bIsEnabled ? ESlateDrawEffect::None : ESlateDrawEffect::DisabledEffect;
-	const FVector2D& SelectorSize = SelectorImage->ImageSize;
-	const FVector2D CircleSize = AllottedGeometry.GetLocalSize() - SelectorSize;
+	const FVector2f SelectorSize = SelectorImage->ImageSize;
+	const FVector2f CircleSize = AllottedGeometry.GetLocalSize() - SelectorSize;
 	
 	FSlateDrawElement::MakeBox(
 		OutDrawElements,
 		LayerId,
-		AllottedGeometry.ToPaintGeometry(0.5 * SelectorSize, CircleSize),
+		AllottedGeometry.ToPaintGeometry(CircleSize, FSlateLayoutTransform(0.5f * SelectorSize)),
 		Image,
 		DrawEffects,
 		InWidgetStyle.GetColorAndOpacityTint() * Image->GetTint(InWidgetStyle)
@@ -102,7 +102,7 @@ int32 SColorWheel::OnPaint(const FPaintArgs& Args, const FGeometry& AllottedGeom
 	FSlateDrawElement::MakeBox(
 		OutDrawElements,
 		LayerId + 1,
-		AllottedGeometry.ToPaintGeometry(0.5f * (AllottedGeometry.GetLocalSize() + CalcRelativePositionFromCenter() * CircleSize - SelectorSize), SelectorSize),
+		AllottedGeometry.ToPaintGeometry(SelectorSize, FSlateLayoutTransform(0.5f * (AllottedGeometry.GetLocalSize() + CalcRelativePositionFromCenter() * CircleSize - SelectorSize))),
 		SelectorImage,
 		DrawEffects,
 		InWidgetStyle.GetColorAndOpacityTint() * SelectorImage->GetTint(InWidgetStyle)
@@ -115,21 +115,21 @@ int32 SColorWheel::OnPaint(const FPaintArgs& Args, const FGeometry& AllottedGeom
 /* SColorWheel implementation
  *****************************************************************************/
 
-FVector2D SColorWheel::CalcRelativePositionFromCenter() const
+UE::Slate::FDeprecateVector2DResult SColorWheel::CalcRelativePositionFromCenter() const
 {
 	float Hue = SelectedColor.Get().R;
 	float Saturation = SelectedColor.Get().G;
 	float Angle = Hue / 180.0f * PI;
 	float Radius = Saturation;
 
-	return FVector2D(FMath::Cos(Angle), FMath::Sin(Angle)) * Radius;
+	return FVector2f(FMath::Cos(Angle), FMath::Sin(Angle)) * Radius;
 }
 
 
 bool SColorWheel::ProcessMouseAction(const FGeometry& MyGeometry, const FPointerEvent& MouseEvent, bool bProcessWhenOutsideColorWheel)
 {
-	const FVector2D LocalMouseCoordinate = MyGeometry.AbsoluteToLocal(MouseEvent.GetScreenSpacePosition());
-	const FVector2D RelativePositionFromCenter = (2.0f * LocalMouseCoordinate - MyGeometry.GetLocalSize()) / (MyGeometry.GetLocalSize() - SelectorImage->ImageSize);
+	const FVector2f LocalMouseCoordinate = MyGeometry.AbsoluteToLocal(MouseEvent.GetScreenSpacePosition());
+	const FVector2f RelativePositionFromCenter = (2.0f * LocalMouseCoordinate - MyGeometry.GetLocalSize()) / (MyGeometry.GetLocalSize() - SelectorImage->ImageSize);
 	const float RelativeRadius = RelativePositionFromCenter.Size();
 
 	if (RelativeRadius <= 1.0f || bProcessWhenOutsideColorWheel)
diff --git a/Engine/Source/Runtime/Slate/Private/Widgets/Docking/SDockTab.cpp b/Engine/Source/Runtime/Slate/Private/Widgets/Docking/SDockTab.cpp
--- a/Engine/Source/Runtime/Slate/Private/Widgets/Docking/SDockTab.cpp
+++ b/Engine/Source/Runtime/Slate/Private/Widgets/Docking/SDockTab.cpp
@@ -941,7 +941,7 @@ EVisibility SDockTab::HandleIsCloseButtonVisible() const
 
 TOptional<FVector2D> SDockTab::GetTabIconSize() const
 {
-	return GetCurrentStyle().IconSize;
+	return FVector2D(GetCurrentStyle().IconSize);
 }
 
 bool SDockTab::CanCloseTab() const
@@ -981,10 +981,10 @@ void SDockTab::PersistVisualState()
 }
 
 
-FVector2D SDockTab::GetAnimatedScale() const
+UE::Slate::FDeprecateVector2DResult SDockTab::GetAnimatedScale() const
 {
-	static FVector2D FullyOpen = FVector2D::UnitVector;
-	static FVector2D FullyClosed = FVector2D(1.0f, 0.0f);
+	static FVector2f FullyOpen = FVector2f::UnitVector;
+	static FVector2f FullyClosed = FVector2f(1.0f, 0.0f);
 	return FMath::Lerp(FullyClosed, FullyOpen, SpawnAnimCurve.GetLerp());
 }
 
diff --git a/Engine/Source/Runtime/Slate/Private/Widgets/Input/SButton.cpp b/Engine/Source/Runtime/Slate/Private/Widgets/Input/SButton.cpp
--- a/Engine/Source/Runtime/Slate/Private/Widgets/Input/SButton.cpp
+++ b/Engine/Source/Runtime/Slate/Private/Widgets/Input/SButton.cpp
@@ -547,7 +547,7 @@ FVector2D SButton::ComputeDesiredSize(float LayoutScaleMultiplier) const
 	// the border image specified by the style.
 	if (ChildSlot.GetWidget() == SNullWidget::NullWidget)
 	{
-		return GetBorderImage()->ImageSize;
+		return FVector2D(GetBorderImage()->ImageSize);
 	}
 	else
 	{
diff --git a/Engine/Source/Runtime/Slate/Private/Widgets/Input/SSlider.cpp b/Engine/Source/Runtime/Slate/Private/Widgets/Input/SSlider.cpp
--- a/Engine/Source/Runtime/Slate/Private/Widgets/Input/SSlider.cpp
+++ b/Engine/Source/Runtime/Slate/Private/Widgets/Input/SSlider.cpp
@@ -19,7 +19,7 @@ void SSlider::PrivateRegisterAttributes(FSlateAttributeInitializer& AttributeIni
 
 SSlider::SSlider()
 	: Style(nullptr)
-	, PressedScreenSpaceTouchDownPosition(FVector2D(0, 0))
+	, PressedScreenSpaceTouchDownPosition(FVector2f(0, 0))
 	, ValueSlateAttribute(*this, 1.f)
 	, IndentHandleSlateAttribute(*this, true)
 	, LockedSlateAttribute(*this, false)
@@ -67,13 +67,13 @@ int32 SSlider::OnPaint( const FPaintArgs& Args, const FGeometry& AllottedGeometr
 	const float AllottedHeight = Orientation == Orient_Horizontal ? AllottedGeometry.GetLocalSize().Y : AllottedGeometry.GetLocalSize().X;
 
 	float HandleRotation;
-	FVector2D HandleTopLeftPoint;
-	FVector2D SliderStartPoint;
-	FVector2D SliderEndPoint;
+	FVector2f HandleTopLeftPoint;
+	FVector2f SliderStartPoint;
+	FVector2f SliderEndPoint;
 
 	// calculate slider geometry as if it's a horizontal slider (we'll rotate it later if it's vertical)
-	const FVector2D HandleSize = GetThumbImage()->ImageSize;
-	const FVector2D HalfHandleSize = 0.5f * HandleSize;
+	const FVector2f HandleSize = GetThumbImage()->ImageSize;
+	const FVector2f HalfHandleSize = 0.5f * HandleSize;
 	const float Indentation = IndentHandleSlateAttribute.Get() ? HandleSize.X : 0.0f;
 
 	// We clamp to make sure that the slider cannot go out of the slider Length.
@@ -83,10 +83,10 @@ int32 SSlider::OnPaint( const FPaintArgs& Args, const FGeometry& AllottedGeometr
 	const float SliderY = 0.5f * AllottedHeight;
 
 	HandleRotation = 0.0f;
-	HandleTopLeftPoint = FVector2D(SliderHandleOffset + (0.5f * Indentation), SliderY - HalfHandleSize.Y);
+	HandleTopLeftPoint = FVector2f(SliderHandleOffset + (0.5f * Indentation), SliderY - HalfHandleSize.Y);
 
-	SliderStartPoint = FVector2D(HalfHandleSize.X, SliderY);
-	SliderEndPoint = FVector2D(AllottedWidth - HalfHandleSize.X, SliderY);
+	SliderStartPoint = FVector2f(HalfHandleSize.X, SliderY);
+	SliderEndPoint = FVector2f(AllottedWidth - HalfHandleSize.X, SliderY);
 
 	FGeometry SliderGeometry = AllottedGeometry;
 	
@@ -94,26 +94,26 @@ int32 SSlider::OnPaint( const FPaintArgs& Args, const FGeometry& AllottedGeometr
 	if (Orientation == Orient_Vertical)
 	{
 		// Do this by translating along -X by the width of the geometry, then rotating 90 degreess CCW (left-hand coords)
-		FSlateRenderTransform SlateRenderTransform = TransformCast<FSlateRenderTransform>(Concatenate(Inverse(FVector2D(AllottedWidth, 0)), FQuat2D(FMath::DegreesToRadians(-90.0f))));
+		FSlateRenderTransform SlateRenderTransform = TransformCast<FSlateRenderTransform>(Concatenate(Inverse(FVector2f(AllottedWidth, 0)), FQuat2D(FMath::DegreesToRadians(-90.0f))));
 		// create a child geometry matching this one, but with the render transform.
 		SliderGeometry = AllottedGeometry.MakeChild(
-			FVector2D(AllottedWidth, AllottedHeight), 
+			FVector2f(AllottedWidth, AllottedHeight), 
 			FSlateLayoutTransform(), 
-			SlateRenderTransform, FVector2D::ZeroVector);
+			SlateRenderTransform, FVector2f::ZeroVector);
 	}
 
 	const bool bEnabled = ShouldBeEnabled(bParentEnabled);
 	const ESlateDrawEffect DrawEffects = bEnabled ? ESlateDrawEffect::None : ESlateDrawEffect::DisabledEffect;
 
 	// draw slider bar
-	auto BarTopLeft = FVector2D(SliderStartPoint.X, SliderStartPoint.Y - Style->BarThickness * 0.5f);
-	auto BarSize = FVector2D(SliderEndPoint.X - SliderStartPoint.X, Style->BarThickness);
+	auto BarTopLeft = FVector2f(SliderStartPoint.X, SliderStartPoint.Y - Style->BarThickness * 0.5f);
+	auto BarSize = FVector2f(SliderEndPoint.X - SliderStartPoint.X, Style->BarThickness);
 	auto BarImage = GetBarImage();
 	auto ThumbImage = GetThumbImage();
 	FSlateDrawElement::MakeBox(
 		OutDrawElements,
 		LayerId,
-		SliderGeometry.ToPaintGeometry(BarTopLeft, BarSize),
+		SliderGeometry.ToPaintGeometry(BarSize, FSlateLayoutTransform(BarTopLeft)),
 		BarImage,
 		DrawEffects,
 		BarImage->GetTint(InWidgetStyle) * SliderBarColorSlateAttribute.Get().GetColor(InWidgetStyle) * InWidgetStyle.GetColorAndOpacityTint()
@@ -125,7 +125,7 @@ int32 SSlider::OnPaint( const FPaintArgs& Args, const FGeometry& AllottedGeometr
 	FSlateDrawElement::MakeBox( 
 		OutDrawElements,
 		LayerId,
-		SliderGeometry.ToPaintGeometry(HandleTopLeftPoint, GetThumbImage()->ImageSize),
+		SliderGeometry.ToPaintGeometry(GetThumbImage()->ImageSize, FSlateLayoutTransform(HandleTopLeftPoint)),
 		ThumbImage,
 		DrawEffects,
 		ThumbImage->GetTint(InWidgetStyle) * SliderHandleColorSlateAttribute.Get().GetColor(InWidgetStyle) * InWidgetStyle.GetColorAndOpacityTint()
@@ -415,9 +415,9 @@ void SSlider::CommitValue(float NewValue)
 	OnValueChanged.ExecuteIfBound(NewValue);
 }
 
-float SSlider::PositionToValue( const FGeometry& MyGeometry, const FVector2D& AbsolutePosition )
+float SSlider::PositionToValue( const FGeometry& MyGeometry, const UE::Slate::FDeprecateVector2DParameter& AbsolutePosition )
 {
-	const FVector2D LocalPosition = MyGeometry.AbsoluteToLocal(AbsolutePosition);
+	const FVector2f LocalPosition = MyGeometry.AbsoluteToLocal(AbsolutePosition);
 
 	float RelativeValue;
 	float Denominator;
diff --git a/Engine/Source/Runtime/Slate/Private/Widgets/Input/SVirtualJoystick.cpp b/Engine/Source/Runtime/Slate/Private/Widgets/Input/SVirtualJoystick.cpp
--- a/Engine/Source/Runtime/Slate/Private/Widgets/Input/SVirtualJoystick.cpp
+++ b/Engine/Source/Runtime/Slate/Private/Widgets/Input/SVirtualJoystick.cpp
@@ -125,8 +125,9 @@ int32 SVirtualJoystick::OnPaint(const FPaintArgs& Args, const FGeometry& Allotte
 					OutDrawElements,
 					RetLayerId++,
 					AllottedGeometry.ToPaintGeometry(
-					Control.VisualCenter - FVector2D(Control.CorrectedVisualSize.X * 0.5f, Control.CorrectedVisualSize.Y * 0.5f),
-					Control.CorrectedVisualSize),
+						Control.CorrectedVisualSize,
+						FSlateLayoutTransform(Control.VisualCenter - FVector2D(Control.CorrectedVisualSize.X * 0.5f, Control.CorrectedVisualSize.Y * 0.5f))
+					),
 					Control.Info.Image2->GetSlateBrush(),
 					ESlateDrawEffect::None,
 					ColorAndOpacitySRGB
@@ -139,8 +140,9 @@ int32 SVirtualJoystick::OnPaint(const FPaintArgs& Args, const FGeometry& Allotte
 					OutDrawElements,
 					RetLayerId++,
 					AllottedGeometry.ToPaintGeometry(
-					Control.VisualCenter + Control.ThumbPosition - FVector2D(Control.CorrectedThumbSize.X * 0.5f, Control.CorrectedThumbSize.Y * 0.5f),
-					Control.CorrectedThumbSize),
+						Control.CorrectedThumbSize,
+						FSlateLayoutTransform(Control.VisualCenter + Control.ThumbPosition - FVector2D(Control.CorrectedThumbSize.X * 0.5f, Control.CorrectedThumbSize.Y * 0.5f))
+					),
 					Control.Info.Image1->GetSlateBrush(),
 					ESlateDrawEffect::None,
 					ColorAndOpacitySRGB
diff --git a/Engine/Source/Runtime/Slate/Private/Widgets/Input/SVirtualKeyboardEntry.cpp b/Engine/Source/Runtime/Slate/Private/Widgets/Input/SVirtualKeyboardEntry.cpp
--- a/Engine/Source/Runtime/Slate/Private/Widgets/Input/SVirtualKeyboardEntry.cpp
+++ b/Engine/Source/Runtime/Slate/Private/Widgets/Input/SVirtualKeyboardEntry.cpp
@@ -200,7 +200,7 @@ int32 SVirtualKeyboardEntry::OnPaint( const FPaintArgs& Args, const FGeometry& A
 		FSlateDrawElement::MakeText(
 			OutDrawElements,
 			LayerId + TextLayer,
-			AllottedGeometry.ToPaintGeometry( FVector2D( 0, DrawPositionY ), AllottedGeometry.GetLocalSize() ),
+			AllottedGeometry.ToPaintGeometry( AllottedGeometry.GetLocalSize(), FSlateLayoutTransform(FVector2f( 0.0, DrawPositionY )) ),
 			ThisHintText,          // Text
 			FontInfo,              // Font information (font name, size)
 			DrawEffects,           // Effects to use
@@ -241,7 +241,7 @@ int32 SVirtualKeyboardEntry::OnPaint( const FPaintArgs& Args, const FGeometry& A
 		FSlateDrawElement::MakeText(
 			OutDrawElements,
 			LayerId + TextLayer,
-			AllottedGeometry.ToPaintGeometry(FVector2D(0, DrawPositionY), AllottedGeometry.GetLocalSize()),
+			AllottedGeometry.ToPaintGeometry(AllottedGeometry.GetLocalSize(), FSlateLayoutTransform(FVector2f(0, DrawPositionY))),
 			VisibleText,           // Text
 			FontInfo,              // Font information (font name, size)
 			DrawEffects,           // Effects to use
diff --git a/Engine/Source/Runtime/Slate/Private/Widgets/Layout/SGridPanel.cpp b/Engine/Source/Runtime/Slate/Private/Widgets/Layout/SGridPanel.cpp
--- a/Engine/Source/Runtime/Slate/Private/Widgets/Layout/SGridPanel.cpp
+++ b/Engine/Source/Runtime/Slate/Private/Widgets/Layout/SGridPanel.cpp
@@ -465,7 +465,7 @@ int32 SGridPanel::LayoutDebugPaint(const FGeometry& AllottedGeometry, const FSla
 			(
 				OutDrawElements, 
 				LayerId,
-				AllottedGeometry.ToPaintGeometry( FVector2D(XOffset, YOffset), FVector2D( Columns[Column], Rows[Row] ) )
+				AllottedGeometry.ToPaintGeometry( FVector2f( Columns[Column], Rows[Row] ), FSlateLayoutTransform(FVector2f(XOffset, YOffset)) )
 			);
 
 			YOffset += Rows[Row];
diff --git a/Engine/Source/Runtime/Slate/Private/Widgets/Layout/SScrollBox.cpp b/Engine/Source/Runtime/Slate/Private/Widgets/Layout/SScrollBox.cpp
--- a/Engine/Source/Runtime/Slate/Private/Widgets/Layout/SScrollBox.cpp
+++ b/Engine/Source/Runtime/Slate/Private/Widgets/Layout/SScrollBox.cpp
@@ -124,7 +124,7 @@ void SScrollBox::Construct( const FArguments& InArgs )
 	AmountScrolledWhileRightMouseDown = 0;
 	PendingScrollTriggerAmount = 0;
 	bShowSoftwareCursor = false;
-	SoftwareCursorPosition = FVector2D::ZeroVector;
+	SoftwareCursorPosition = FVector2f::ZeroVector;
 	OnUserScrolled = InArgs._OnUserScrolled;
 	Orientation = InArgs._Orientation;
 	bScrollToEnd = false;
@@ -429,7 +429,7 @@ bool SScrollBox::InternalScrollDescendantIntoView(const FGeometry& MyGeometry, c
 			{
 				// Calculate how much we would need to scroll to bring this to the top/left of the scroll box
 				const float WidgetPosition = GetScrollComponentFromVector(MyGeometry.AbsoluteToLocal(WidgetGeometry->Geometry.GetAbsolutePosition()) + (WidgetGeometry->Geometry.GetLocalSize() / 2));
-				const float MyPosition = GetScrollComponentFromVector(MyGeometry.GetLocalSize() * FVector2D(0.5f, 0.5f));
+				const float MyPosition = GetScrollComponentFromVector(MyGeometry.GetLocalSize() * FVector2f(0.5f, 0.5f));
 				ScrollOffset = WidgetPosition - MyPosition;
 			}
 			else
@@ -516,7 +516,7 @@ void SScrollBox::SetScrollBarTrackAlwaysVisible(bool InAlwaysVisible)
 	ScrollBar->SetScrollBarTrackAlwaysVisible(InAlwaysVisible);
 }
 
-void SScrollBox::SetScrollBarThickness(FVector2D InThickness)
+void SScrollBox::SetScrollBarThickness(UE::Slate::FDeprecateVector2DParameter InThickness)
 {
 	ScrollBar->SetThickness(InThickness);
 }
@@ -720,7 +720,7 @@ FReply SScrollBox::OnMouseButtonUp( const FGeometry& MyGeometry, const FPointerE
 		if ( HasMouseCapture() )
 		{
 			FSlateRect PanelScreenSpaceRect = MyGeometry.GetLayoutBoundingRect();
-			FVector2D CursorPosition = MyGeometry.LocalToAbsolute( SoftwareCursorPosition );
+			FVector2f CursorPosition = MyGeometry.LocalToAbsolute( SoftwareCursorPosition );
 
 			FIntPoint BestPositionInPanel(
 				FMath::RoundToInt( FMath::Clamp( CursorPosition.X, PanelScreenSpaceRect.Left, PanelScreenSpaceRect.Right ) ),
@@ -1078,12 +1078,12 @@ int32 SScrollBox::OnPaint( const FPaintArgs& Args, const FGeometry& AllottedGeom
 	}
 
 	const FSlateBrush* Brush = FCoreStyle::Get().GetBrush(TEXT("SoftwareCursor_Grab"));
-	const FVector2D CursorSize = Brush->ImageSize / AllottedGeometry.Scale;
+	const FVector2f CursorSize = Brush->ImageSize / AllottedGeometry.Scale;
 
 	FSlateDrawElement::MakeBox(
 		OutDrawElements,
 		++NewLayerId,
-		AllottedGeometry.ToPaintGeometry( SoftwareCursorPosition - (CursorSize / 2 ), CursorSize),
+		AllottedGeometry.ToPaintGeometry( CursorSize, FSlateLayoutTransform(SoftwareCursorPosition - (CursorSize *.5f )) ),
 		Brush
 	);
 
diff --git a/Engine/Source/Runtime/Slate/Private/Widgets/Layout/SSplitter.cpp b/Engine/Source/Runtime/Slate/Private/Widgets/Layout/SSplitter.cpp
--- a/Engine/Source/Runtime/Slate/Private/Widgets/Layout/SSplitter.cpp
+++ b/Engine/Source/Runtime/Slate/Private/Widgets/Layout/SSplitter.cpp
@@ -949,7 +949,7 @@ int32 SSplitter2x2::OnPaint(const FPaintArgs& Args, const FGeometry& AllottedGeo
 	FSlateDrawElement::MakeBox(
 		OutDrawElements,
 		MaxLayerId,
-		AllottedGeometry.ToPaintGeometry(HorizontalSplitterPos, HorizontalSplitterSize),
+		AllottedGeometry.ToPaintGeometry(HorizontalSplitterSize, FSlateLayoutTransform(HorizontalSplitterPos)),
 		HorizontalBrush,
 		ShouldBeEnabled(bParentEnabled) ? ESlateDrawEffect::None : ESlateDrawEffect::DisabledEffect,
 		InWidgetStyle.GetColorAndOpacityTint() * HorizontalBrush->TintColor.GetSpecifiedColor());
@@ -962,7 +962,7 @@ int32 SSplitter2x2::OnPaint(const FPaintArgs& Args, const FGeometry& AllottedGeo
 	FSlateDrawElement::MakeBox(
 		OutDrawElements,
 		MaxLayerId,
-		AllottedGeometry.ToPaintGeometry(VerticalSplitterPos, VerticalSplitterSize),
+		AllottedGeometry.ToPaintGeometry(VerticalSplitterSize, FSlateLayoutTransform(VerticalSplitterPos)),
 		VerticalBrush,
 		ShouldBeEnabled(bParentEnabled) ? ESlateDrawEffect::None : ESlateDrawEffect::DisabledEffect,
 		InWidgetStyle.GetColorAndOpacityTint() * VerticalBrush->TintColor.GetSpecifiedColor());
diff --git a/Engine/Source/Runtime/Slate/Private/Widgets/Notifications/SNotificationBackground.cpp b/Engine/Source/Runtime/Slate/Private/Widgets/Notifications/SNotificationBackground.cpp
--- a/Engine/Source/Runtime/Slate/Private/Widgets/Notifications/SNotificationBackground.cpp
+++ b/Engine/Source/Runtime/Slate/Private/Widgets/Notifications/SNotificationBackground.cpp
@@ -47,7 +47,7 @@ int32 SNotificationBackground::OnPaint(const FPaintArgs& Args, const FGeometry&
 		FSlateDrawElement::MakeBox(
 			OutDrawElements,
 			LayerId + 1,
-			AllottedGeometry.ToPaintGeometry(FVector2D(4.0f, 4.0f), FVector2D(WatermarkSize, WatermarkSize)),
+			AllottedGeometry.ToPaintGeometry(FVector2f(WatermarkSize, WatermarkSize), FSlateLayoutTransform(FVector2f(4.0f, 4.0f))),
 			WatermarkBrush,
 			ESlateDrawEffect::None,
 			WatermarkTint * InWidgetStyle.GetColorAndOpacityTint() * GetBorderBackgroundColor().GetColor(InWidgetStyle));
diff --git a/Engine/Source/Runtime/Slate/Private/Widgets/Notifications/SProgressBar.cpp b/Engine/Source/Runtime/Slate/Private/Widgets/Notifications/SProgressBar.cpp
--- a/Engine/Source/Runtime/Slate/Private/Widgets/Notifications/SProgressBar.cpp
+++ b/Engine/Source/Runtime/Slate/Private/Widgets/Notifications/SProgressBar.cpp
@@ -228,14 +228,16 @@ int32 SProgressBar::OnPaint( const FPaintArgs& Args, const FGeometry& AllottedGe
 					if (bScaleWithFillPerc)
 					{
 						PaintRect = AllottedGeometry.ToPaintGeometry(
-							FVector2D((AllottedGeometry.GetLocalSize().X * 0.5f) - ((AllottedGeometry.GetLocalSize().X * (ClampedFraction)) * 0.5f), (AllottedGeometry.GetLocalSize().Y * 0.5f) - ((AllottedGeometry.GetLocalSize().Y * (ClampedFraction)) * 0.5f)),
-							FVector2D(AllottedGeometry.GetLocalSize().X * ClampedFraction, AllottedGeometry.GetLocalSize().Y * ClampedFraction));
+							FVector2f(AllottedGeometry.GetLocalSize().X * ClampedFraction, AllottedGeometry.GetLocalSize().Y * ClampedFraction),
+							FSlateLayoutTransform(FVector2f((AllottedGeometry.GetLocalSize().X * 0.5f) - ((AllottedGeometry.GetLocalSize().X * (ClampedFraction)) * 0.5f), (AllottedGeometry.GetLocalSize().Y * 0.5f) - ((AllottedGeometry.GetLocalSize().Y * (ClampedFraction)) * 0.5f)))
+						);
 					}
 					else
 					{
 						PaintRect = AllottedGeometry.ToPaintGeometry(
-							FVector2D(0.0f, 0.0f),
-							FVector2D(AllottedGeometry.GetLocalSize().X, AllottedGeometry.GetLocalSize().Y));
+							FVector2f(AllottedGeometry.GetLocalSize().X, AllottedGeometry.GetLocalSize().Y),
+							FSlateLayoutTransform()
+						);
 					}
 
 					// Draw Fill
@@ -261,14 +263,16 @@ int32 SProgressBar::OnPaint( const FPaintArgs& Args, const FGeometry& AllottedGe
 					if (bScaleWithFillPerc)
 					{
 						PaintRect = AllottedGeometry.ToPaintGeometry(
-							FVector2D((AllottedGeometry.GetLocalSize().X * 0.5f) - ((AllottedGeometry.GetLocalSize().X * (ClampedFraction)) * 0.5f), 0.0f),
-							FVector2D(AllottedGeometry.GetLocalSize().X * ClampedFraction, AllottedGeometry.GetLocalSize().Y));
+							FVector2f(AllottedGeometry.GetLocalSize().X * ClampedFraction, AllottedGeometry.GetLocalSize().Y),
+							FSlateLayoutTransform(FVector2f((AllottedGeometry.GetLocalSize().X * 0.5f) - ((AllottedGeometry.GetLocalSize().X * (ClampedFraction)) * 0.5f), 0.0f))
+						);
 					}
 					else
 					{
 						PaintRect = AllottedGeometry.ToPaintGeometry(
-							FVector2D(0.0f, 0.0f),
-							FVector2D(AllottedGeometry.GetLocalSize().X, AllottedGeometry.GetLocalSize().Y));
+							FVector2f(AllottedGeometry.GetLocalSize().X, AllottedGeometry.GetLocalSize().Y),
+							FSlateLayoutTransform(FVector2f(0.0f, 0.0f))
+						);
 					}
 
 					// Draw Fill
@@ -294,14 +298,16 @@ int32 SProgressBar::OnPaint( const FPaintArgs& Args, const FGeometry& AllottedGe
 					if (bScaleWithFillPerc)
 					{
 						PaintRect = AllottedGeometry.ToPaintGeometry(
-							FVector2D(0.0f, (AllottedGeometry.GetLocalSize().Y * 0.5f) - ((AllottedGeometry.GetLocalSize().Y * (ClampedFraction)) * 0.5f)),
-							FVector2D(AllottedGeometry.GetLocalSize().X, AllottedGeometry.GetLocalSize().Y * ClampedFraction));
+							FVector2f(AllottedGeometry.GetLocalSize().X, AllottedGeometry.GetLocalSize().Y * ClampedFraction),
+							FSlateLayoutTransform(FVector2f(0.0f, (AllottedGeometry.GetLocalSize().Y * 0.5f) - ((AllottedGeometry.GetLocalSize().Y * (ClampedFraction)) * 0.5f)))
+						);
 					}
 					else
 					{
 						PaintRect = AllottedGeometry.ToPaintGeometry(
-							FVector2D(0.0f, 0.0f),
-							FVector2D(AllottedGeometry.GetLocalSize().X, AllottedGeometry.GetLocalSize().Y));
+							FVector2f(AllottedGeometry.GetLocalSize().X, AllottedGeometry.GetLocalSize().Y),
+							FSlateLayoutTransform(FVector2f(0.0f, 0.0f))
+						);
 					}
 
 					// Draw Fill
@@ -326,14 +332,16 @@ int32 SProgressBar::OnPaint( const FPaintArgs& Args, const FGeometry& AllottedGe
 					if (bScaleWithFillPerc)
 					{
 						PaintRect = AllottedGeometry.ToPaintGeometry(
-							FVector2D(0.0f, 0.0f),
-							FVector2D(AllottedGeometry.GetLocalSize().X, AllottedGeometry.GetLocalSize().Y * ClampedFraction));
+							FVector2f(AllottedGeometry.GetLocalSize().X, AllottedGeometry.GetLocalSize().Y * ClampedFraction),
+							FSlateLayoutTransform(FVector2f(0.0f, 0.0f))
+						);
 					}
 					else
 					{
 						PaintRect = AllottedGeometry.ToPaintGeometry(
-							FVector2D(0.0f, MarqueeAnimOffsetY - MarqueeImageSizeY),
-							FVector2D(AllottedGeometry.GetLocalSize().X, AllottedGeometry.GetLocalSize().Y + MarqueeImageSizeY));
+							FVector2f(AllottedGeometry.GetLocalSize().X, AllottedGeometry.GetLocalSize().Y + MarqueeImageSizeY),
+							FSlateLayoutTransform(FVector2f(0.0f, MarqueeAnimOffsetY - MarqueeImageSizeY))
+						);
 					}
 
 					// Draw Fill
@@ -358,14 +366,16 @@ int32 SProgressBar::OnPaint( const FPaintArgs& Args, const FGeometry& AllottedGe
 					if (bScaleWithFillPerc)
 					{
 						PaintRect = AllottedGeometry.ToPaintGeometry(
-							FVector2D(0.0f, AllottedGeometry.GetLocalSize().Y * (1.0f - ClampedFraction)),
-							FVector2D(AllottedGeometry.GetLocalSize().X, AllottedGeometry.GetLocalSize().Y * ClampedFraction));
+							FVector2f(AllottedGeometry.GetLocalSize().X, AllottedGeometry.GetLocalSize().Y * ClampedFraction),
+							FSlateLayoutTransform(FVector2f(0.0f, AllottedGeometry.GetLocalSize().Y * (1.0f - ClampedFraction)))
+						);
 					}
 					else
 					{
 						PaintRect = AllottedGeometry.ToPaintGeometry(
-							FVector2D(0.0f, MarqueeAnimOffsetY - MarqueeImageSizeY),
-							FVector2D(AllottedGeometry.GetLocalSize().X, AllottedGeometry.GetLocalSize().Y + MarqueeImageSizeY));
+							FVector2f(AllottedGeometry.GetLocalSize().X, AllottedGeometry.GetLocalSize().Y + MarqueeImageSizeY),
+							FSlateLayoutTransform(FVector2f(0.0f, MarqueeAnimOffsetY - MarqueeImageSizeY))
+						);
 					}
 
 					// Draw Fill
@@ -390,14 +400,16 @@ int32 SProgressBar::OnPaint( const FPaintArgs& Args, const FGeometry& AllottedGe
 					if (bScaleWithFillPerc)
 					{
 						PaintRect = AllottedGeometry.ToPaintGeometry(
-							FVector2D(AllottedGeometry.GetLocalSize().X * (1.0f - ClampedFraction), 0.0f),
-							FVector2D(AllottedGeometry.GetLocalSize().X * ClampedFraction, AllottedGeometry.GetLocalSize().Y));
+							FVector2f(AllottedGeometry.GetLocalSize().X * ClampedFraction, AllottedGeometry.GetLocalSize().Y),
+							FSlateLayoutTransform(FVector2f(AllottedGeometry.GetLocalSize().X * (1.0f - ClampedFraction), 0.0f))
+						);
 					}
 					else
 					{
 						PaintRect = AllottedGeometry.ToPaintGeometry(
-							FVector2D(MarqueeAnimOffsetX - MarqueeImageSizeX, 0.0f),
-							FVector2D(AllottedGeometry.GetLocalSize().X + MarqueeImageSizeX, AllottedGeometry.GetLocalSize().Y));
+							FVector2f(AllottedGeometry.GetLocalSize().X + MarqueeImageSizeX, AllottedGeometry.GetLocalSize().Y),
+							FSlateLayoutTransform(FVector2f(MarqueeAnimOffsetX - MarqueeImageSizeX, 0.0f))
+						);
 					}
 
 					// Draw Fill
@@ -423,14 +435,16 @@ int32 SProgressBar::OnPaint( const FPaintArgs& Args, const FGeometry& AllottedGe
 					if (bScaleWithFillPerc)
 					{
 						PaintRect = AllottedGeometry.ToPaintGeometry(
-							FVector2D(0.0f, 0.0f),
-							FVector2D(AllottedGeometry.GetLocalSize().X * ClampedFraction, AllottedGeometry.GetLocalSize().Y));
+							FVector2f(AllottedGeometry.GetLocalSize().X * ClampedFraction, AllottedGeometry.GetLocalSize().Y),
+							FSlateLayoutTransform(FVector2f(0.0f, 0.0f))
+						);
 					}
 					else
 					{
 						PaintRect = AllottedGeometry.ToPaintGeometry(
-							FVector2D(MarqueeAnimOffsetX - MarqueeImageSizeX, 0.0f),
-							FVector2D(AllottedGeometry.GetLocalSize().X + MarqueeImageSizeX, AllottedGeometry.GetLocalSize().Y));
+							FVector2f(AllottedGeometry.GetLocalSize().X + MarqueeImageSizeX, AllottedGeometry.GetLocalSize().Y),
+							FSlateLayoutTransform(FVector2f(MarqueeAnimOffsetX - MarqueeImageSizeX, 0.0f))
+						);
 					}
 
 					// Draw fill
@@ -467,8 +481,9 @@ int32 SProgressBar::OnPaint( const FPaintArgs& Args, const FGeometry& AllottedGe
 				OutDrawElements,
 				RetLayerId++,
 				AllottedGeometry.ToPaintGeometry(
-					FVector2D( MarqueeAnimOffset - MarqueeImageSize, 0.0f ),
-					FVector2D( AllottedGeometry.GetLocalSize().X + MarqueeImageSize, AllottedGeometry.GetLocalSize().Y )),
+					FVector2f( AllottedGeometry.GetLocalSize().X + MarqueeImageSize, AllottedGeometry.GetLocalSize().Y ),
+					FSlateLayoutTransform(FVector2f( MarqueeAnimOffset - MarqueeImageSize, 0.0f ))
+				),
 				CurrentMarqueeImage,
 				DrawEffects,
 				FillColorAndOpacitySRGB
@@ -483,7 +498,7 @@ int32 SProgressBar::OnPaint( const FPaintArgs& Args, const FGeometry& AllottedGe
 
 FVector2D SProgressBar::ComputeDesiredSize( float ) const
 {
-	return GetMarqueeImage()->ImageSize;
+	return FVector2D(GetMarqueeImage()->ImageSize);
 }
 
 bool SProgressBar::ComputeVolatility() const
diff --git a/Engine/Source/Runtime/Slate/Private/Widgets/SViewport.cpp b/Engine/Source/Runtime/Slate/Private/Widgets/SViewport.cpp
--- a/Engine/Source/Runtime/Slate/Private/Widgets/SViewport.cpp
+++ b/Engine/Source/Runtime/Slate/Private/Widgets/SViewport.cpp
@@ -15,9 +15,9 @@ DECLARE_CYCLE_STAT(TEXT("Game UI Paint"), STAT_ViewportPaintTime, STATGROUP_Slat
 
 /* SViewport::FArguments interface
 *****************************************************************************/
-FVector2D SViewport::FArguments::GetDefaultViewportSize()
+UE::Slate::FDeprecateVector2DResult SViewport::FArguments::GetDefaultViewportSize()
 {
-	return FVector2D(320.0f, 240.0f);
+	return FVector2f(320.0f, 240.0f);
 }
 
 /* SViewport constructors
@@ -173,14 +173,14 @@ int32 SViewport::OnPaint( const FPaintArgs& Args, const FGeometry& AllottedGeome
 
 	if( ViewportInterfacePin.IsValid() && ViewportInterfacePin->IsSoftwareCursorVisible() )
 	{
-		const FVector2D CursorPosScreenSpace = FSlateApplication::Get().GetCursorPos();		
+		const FVector2f CursorPosScreenSpace = FSlateApplication::Get().GetCursorPos();		
 		// @todo Slate: why are we calling OnCursorQuery in here?
 		FCursorReply Reply = ViewportInterfacePin->OnCursorQuery( AllottedGeometry,
 			FPointerEvent(
 				FSlateApplicationBase::CursorPointerIndex,
 				CursorPosScreenSpace,
 				CursorPosScreenSpace,
-				FVector2D::ZeroVector,
+				FVector2f::ZeroVector,
 				TSet<FKey>(),
 				FModifierKeysState() )
 		 );
@@ -191,15 +191,15 @@ int32 SViewport::OnPaint( const FPaintArgs& Args, const FGeometry& AllottedGeome
 		{
 			Brush = FCoreStyle::Get().GetBrush(TEXT("SoftwareCursor_CardinalCross"));
 		}
-		const FVector2D CursorSize = Brush->ImageSize / AllottedGeometry.Scale;
+		const FVector2f CursorSize = Brush->ImageSize / AllottedGeometry.Scale;
 
-		FVector2D CursorPositionLocalSpace = ViewportInterfacePin->GetSoftwareCursorPosition() / AllottedGeometry.Scale;
+		FVector2f CursorPositionLocalSpace = UE::Slate::CastToVector2f(ViewportInterfacePin->GetSoftwareCursorPosition()) / AllottedGeometry.Scale;
 
 		LayerId++;
 		FSlateDrawElement::MakeBox(
 			OutDrawElements,
 			LayerId,
-			AllottedGeometry.ToPaintGeometry( CursorPositionLocalSpace - (CursorSize / 2 ), CursorSize),
+			AllottedGeometry.ToPaintGeometry( CursorSize, FSlateLayoutTransform(CursorPositionLocalSpace - (CursorSize *.5f )) ),
 			Brush
 		);
 	}
diff --git a/Engine/Source/Runtime/Slate/Private/Widgets/Text/SMultiLineEditableText.cpp b/Engine/Source/Runtime/Slate/Private/Widgets/Text/SMultiLineEditableText.cpp
--- a/Engine/Source/Runtime/Slate/Private/Widgets/Text/SMultiLineEditableText.cpp
+++ b/Engine/Source/Runtime/Slate/Private/Widgets/Text/SMultiLineEditableText.cpp
@@ -604,7 +604,7 @@ void SMultiLineEditableText::Refresh()
 void SMultiLineEditableText::ForceScroll(int32 UserIndex, float ScrollAxisMagnitude)
 {
 	const FGeometry& CachedGeom = GetCachedGeometry();
-	FVector2D ScrollPos = (CachedGeom.LocalToAbsolute(FVector2D::ZeroVector) + CachedGeom.LocalToAbsolute(CachedGeom.GetLocalSize())) * 0.5f;
+	FVector2f ScrollPos = (CachedGeom.LocalToAbsolute(FVector2f::ZeroVector) + CachedGeom.LocalToAbsolute(CachedGeom.GetLocalSize())) * 0.5f;
 	TSet<FKey> PressedKeys;
 
 	OnMouseWheel(CachedGeom, FPointerEvent(UserIndex, 0, ScrollPos, ScrollPos, PressedKeys, EKeys::Invalid, ScrollAxisMagnitude, FModifierKeysState()));
@@ -635,12 +635,12 @@ int32 SMultiLineEditableText::OnPaint( const FPaintArgs& Args, const FGeometry&
 	if (bIsSoftwareCursor)
 	{
 		const FSlateBrush* Brush = FCoreStyle::Get().GetBrush(TEXT("SoftwareCursor_Grab"));
-		const FVector2D CursorSize = Brush->ImageSize / AllottedGeometry.Scale;
+		const FVector2f CursorSize = Brush->ImageSize / AllottedGeometry.Scale;
 
 		FSlateDrawElement::MakeBox(
 			OutDrawElements,
 			++LayerId,
-			AllottedGeometry.ToPaintGeometry(SoftwareCursorPosition - (CursorSize / 2), CursorSize),
+			AllottedGeometry.ToPaintGeometry(CursorSize, FSlateLayoutTransform(SoftwareCursorPosition - (CursorSize *.5f))),
 			Brush
 			);
 	}
@@ -735,7 +735,7 @@ FReply SMultiLineEditableText::OnMouseButtonUp( const FGeometry& MyGeometry, con
 		if (bWasRightClickScrolling)
 		{
 			bIsSoftwareCursor = false;
-			const FVector2D CursorPosition = MyGeometry.LocalToAbsolute(SoftwareCursorPosition);
+			const FVector2f CursorPosition = MyGeometry.LocalToAbsolute(SoftwareCursorPosition);
 			const FIntPoint OriginalMousePos(CursorPosition.X, CursorPosition.Y);
 			return FReply::Handled().ReleaseMouseCapture().SetMousePos(OriginalMousePos);
 		}
@@ -756,9 +756,9 @@ FReply SMultiLineEditableText::OnMouseMove( const FGeometry& MyGeometry, const F
 
 		if (IsRightClickScrolling())
 		{
-			const FVector2D PreviousScrollOffset = EditableTextLayout->GetScrollOffset();
+			const FVector2f PreviousScrollOffset = EditableTextLayout->GetScrollOffset();
 
-			FVector2D NewScrollOffset = PreviousScrollOffset;
+			FVector2f NewScrollOffset = PreviousScrollOffset;
 			NewScrollOffset.Y -= ScrollByAmount;
 			EditableTextLayout->SetScrollOffset(NewScrollOffset, MyGeometry);
 
@@ -789,9 +789,9 @@ FReply SMultiLineEditableText::OnMouseWheel( const FGeometry& MyGeometry, const
 	{
 		const float ScrollAmount = -MouseEvent.GetWheelDelta() * GetGlobalScrollAmount();
 
-		const FVector2D PreviousScrollOffset = EditableTextLayout->GetScrollOffset();
+		const FVector2f PreviousScrollOffset = EditableTextLayout->GetScrollOffset();
 		
-		FVector2D NewScrollOffset = PreviousScrollOffset;
+		FVector2f NewScrollOffset = PreviousScrollOffset;
 		NewScrollOffset.Y += ScrollAmount;
 		EditableTextLayout->SetScrollOffset(NewScrollOffset, MyGeometry);
 
diff --git a/Engine/Source/Runtime/Slate/Private/Widgets/Text/STextBlock.cpp b/Engine/Source/Runtime/Slate/Private/Widgets/Text/STextBlock.cpp
--- a/Engine/Source/Runtime/Slate/Private/Widgets/Text/STextBlock.cpp
+++ b/Engine/Source/Runtime/Slate/Private/Widgets/Text/STextBlock.cpp
@@ -143,9 +143,9 @@ FSlateColor STextBlock::GetColorAndOpacity() const
 	return bIsAttributeColorAndOpacitySet ? ColorAndOpacity.Get() : TextStyle.ColorAndOpacity;
 }
 
-FVector2D STextBlock::GetShadowOffset() const
+FVector2f STextBlock::GetShadowOffset() const
 {
-	return bIsAttributeShadowOffsetSet ? ShadowOffset.Get() : TextStyle.ShadowOffset;
+	return bIsAttributeShadowOffsetSet ? UE::Slate::CastToVector2f(ShadowOffset.Get()) : FVector2f(TextStyle.ShadowOffset);
 }
 
 FLinearColor STextBlock::GetShadowColorAndOpacity() const
@@ -199,7 +199,7 @@ int32 STextBlock::OnPaint( const FPaintArgs& Args, const FGeometry& AllottedGeom
 	{
 		// Draw the optional shadow
 		const FLinearColor LocalShadowColorAndOpacity = GetShadowColorAndOpacity();
-		const FVector2D LocalShadowOffset = GetShadowOffset();
+		const FVector2f LocalShadowOffset = GetShadowOffset();
 		const bool ShouldDropShadow = LocalShadowColorAndOpacity.A > 0.f && LocalShadowOffset.SizeSquared() > 0.f;
 
 		const bool bShouldBeEnabled = ShouldBeEnabled(bParentEnabled);
@@ -272,7 +272,7 @@ FVector2D STextBlock::ComputeDesiredSize(float LayoutScaleMultiplier) const
 
 	if (bSimpleTextMode)
 	{
-		const FVector2D LocalShadowOffset = GetShadowOffset();
+		const FVector2D LocalShadowOffset = FVector2D(GetShadowOffset());
 
 		const float LocalOutlineSize = (float)(GetFont().OutlineSettings.OutlineSize);
 
diff --git a/Engine/Source/Runtime/Slate/Private/Widgets/Text/SlateEditableTextLayout.cpp b/Engine/Source/Runtime/Slate/Private/Widgets/Text/SlateEditableTextLayout.cpp
--- a/Engine/Source/Runtime/Slate/Private/Widgets/Text/SlateEditableTextLayout.cpp
+++ b/Engine/Source/Runtime/Slate/Private/Widgets/Text/SlateEditableTextLayout.cpp
@@ -751,29 +751,29 @@ void FSlateEditableTextLayout::AdvanceSearch(const bool InReverse)
 	}
 }
 
-FVector2D FSlateEditableTextLayout::SetHorizontalScrollFraction(const float InScrollOffsetFraction)
+UE::Slate::FDeprecateVector2DResult FSlateEditableTextLayout::SetHorizontalScrollFraction(const float InScrollOffsetFraction)
 {
 	ScrollOffset.X = FMath::Clamp<float>(InScrollOffsetFraction, 0.0, 1.0) * TextLayout->GetSize().X;
-	return FVector2D(ScrollOffset);
+	return ScrollOffset;
 }
 
-FVector2D FSlateEditableTextLayout::SetVerticalScrollFraction(const float InScrollOffsetFraction)
+UE::Slate::FDeprecateVector2DResult FSlateEditableTextLayout::SetVerticalScrollFraction(const float InScrollOffsetFraction)
 {
 	ScrollOffset.Y = FMath::Clamp<float>(InScrollOffsetFraction, 0.0, 1.0) * TextLayout->GetSize().Y;
-	return FVector2D(ScrollOffset);
+	return ScrollOffset;
 }
 
-FVector2D FSlateEditableTextLayout::SetScrollOffset(const FVector2D& InScrollOffset, const FGeometry& InGeometry)
+UE::Slate::FDeprecateVector2DResult FSlateEditableTextLayout::SetScrollOffset(const UE::Slate::FDeprecateVector2DParameter& InScrollOffset, const FGeometry& InGeometry)
 {
-	const FVector2D ContentSize = TextLayout->GetSize();
+	const FVector2f ContentSize = UE::Slate::CastToVector2f(TextLayout->GetSize());
 	ScrollOffset.X = FMath::Clamp(InScrollOffset.X, 0.0f, ContentSize.X - InGeometry.GetLocalSize().X);
 	ScrollOffset.Y = FMath::Clamp(InScrollOffset.Y, 0.0f, ContentSize.Y - InGeometry.GetLocalSize().Y);
-	return FVector2D(ScrollOffset);
+	return ScrollOffset;
 }
 
-FVector2D FSlateEditableTextLayout::GetScrollOffset() const
+UE::Slate::FDeprecateVector2DResult FSlateEditableTextLayout::GetScrollOffset() const
 {
-	return FVector2D(ScrollOffset);
+	return ScrollOffset;
 }
 
 
@@ -1406,7 +1406,7 @@ FReply FSlateEditableTextLayout::HandleMouseButtonUp(const FGeometry& MyGeometry
 
 FReply FSlateEditableTextLayout::HandleMouseMove(const FGeometry& InMyGeometry, const FPointerEvent& InMouseEvent)
 {
-	if (bIsDragSelecting && OwnerWidget->GetSlateWidget()->HasMouseCapture() && InMouseEvent.GetCursorDelta() != FVector2D::ZeroVector)
+	if (bIsDragSelecting && OwnerWidget->GetSlateWidget()->HasMouseCapture() && InMouseEvent.GetCursorDelta() != FVector2f::ZeroVector)
 	{
 		MoveCursor(FMoveCursor::ViaScreenPointer(InMyGeometry.AbsoluteToLocal(InMouseEvent.GetScreenSpacePosition()), InMyGeometry.Scale, ECursorAction::SelectText));
 		bHasDragSelectedSinceFocused = true;
@@ -1750,13 +1750,13 @@ bool FSlateEditableTextLayout::AnyTextSelected() const
 	return SelectionPosition != CursorInteractionPosition;
 }
 
-bool FSlateEditableTextLayout::IsTextSelectedAt(const FGeometry& MyGeometry, const FVector2D& ScreenSpacePosition) const
+bool FSlateEditableTextLayout::IsTextSelectedAt(const FGeometry& MyGeometry, const UE::Slate::FDeprecateVector2DParameter& ScreenSpacePosition) const
 {
-	const FVector2D LocalPosition = MyGeometry.AbsoluteToLocal(ScreenSpacePosition);
+	const FVector2f LocalPosition = MyGeometry.AbsoluteToLocal(ScreenSpacePosition);
 	return IsTextSelectedAt(LocalPosition * MyGeometry.Scale);
 }
 
-bool FSlateEditableTextLayout::IsTextSelectedAt(const FVector2D& InLocalPosition) const
+bool FSlateEditableTextLayout::IsTextSelectedAt(const UE::Slate::FDeprecateVector2DParameter& InLocalPosition) const
 {
 	const FTextLocation CursorInteractionPosition = CursorInfo.GetCursorInteractionLocation();
 	const FTextLocation SelectionPosition = SelectionStart.Get(CursorInteractionPosition);
@@ -1766,7 +1766,7 @@ bool FSlateEditableTextLayout::IsTextSelectedAt(const FVector2D& InLocalPosition
 		return false;
 	}
 
-	const FTextLocation ClickedPosition = TextLayout->GetTextLocationAt(InLocalPosition);
+	const FTextLocation ClickedPosition = TextLayout->GetTextLocationAt(FVector2D(InLocalPosition));
 
 	FTextLocation SelectionLocation = SelectionStart.Get(CursorInteractionPosition);
 	FTextSelection Selection(SelectionLocation, CursorInteractionPosition);
@@ -1819,15 +1819,15 @@ void FSlateEditableTextLayout::SelectAllText()
 	UpdateCursorHighlight();
 }
 
-void FSlateEditableTextLayout::SelectWordAt(const FGeometry& MyGeometry, const FVector2D& ScreenSpacePosition)
+void FSlateEditableTextLayout::SelectWordAt(const FGeometry& MyGeometry, const UE::Slate::FDeprecateVector2DParameter& ScreenSpacePosition)
 {
-	const FVector2D LocalPosition = MyGeometry.AbsoluteToLocal(ScreenSpacePosition);
+	const FVector2f LocalPosition = MyGeometry.AbsoluteToLocal(ScreenSpacePosition);
 	SelectWordAt(LocalPosition * MyGeometry.Scale);
 }
 
-void FSlateEditableTextLayout::SelectWordAt(const FVector2D& InLocalPosition)
+void FSlateEditableTextLayout::SelectWordAt(const UE::Slate::FDeprecateVector2DParameter& InLocalPosition)
 {
-	FTextLocation InitialLocation = TextLayout->GetTextLocationAt(InLocalPosition);
+	FTextLocation InitialLocation = TextLayout->GetTextLocationAt(FVector2d(InLocalPosition));
 	FTextSelection WordSelection = TextLayout->GetWordAt(InitialLocation);
 
 	FTextLocation WordStart = WordSelection.GetBeginning();
@@ -3471,7 +3471,7 @@ void FSlateEditableTextLayout::Tick(const FGeometry& AllottedGeometry, const dou
 			const FSlateRect LocalLineViewRect(LineView.Offset / TextLayout->GetScale(), (LineView.Offset + LineView.Size) / TextLayout->GetScale());
 
 			const FVector2D LocalCursorLocation = TextLayout->GetLocationAt(ScrollInfo.Position, ScrollInfo.Alignment == SlateEditableTextTypes::ECursorAlignment::Right) / TextLayout->GetScale();
-			const FSlateRect LocalCursorRect(LocalCursorLocation, FVector2D(LocalCursorLocation.X + CaretWidth, LocalCursorLocation.Y + FontMaxCharHeight));
+			const FSlateRect LocalCursorRect(LocalCursorLocation, FVector2f(LocalCursorLocation.X + CaretWidth, LocalCursorLocation.Y + FontMaxCharHeight));
 
 			if (LocalCursorRect.Left < 0.0f)
 			{
@@ -3651,9 +3651,9 @@ void FSlateEditableTextLayout::OnArrangeChildren(const FGeometry& AllottedGeomet
 	}
 }
 
-FVector2D FSlateEditableTextLayout::GetSize() const
+UE::Slate::FDeprecateVector2DResult FSlateEditableTextLayout::GetSize() const
 {
-	return TextLayout->GetSize();
+	return UE::Slate::CastToVector2f(TextLayout->GetSize());
 }
 
 TSharedRef<SWidget> FSlateEditableTextLayout::BuildDefaultContextMenu(const TSharedPtr<FExtender>& InMenuExtender) const
diff --git a/Engine/Source/Runtime/Slate/Private/Widgets/Views/SHeaderRow.cpp b/Engine/Source/Runtime/Slate/Private/Widgets/Views/SHeaderRow.cpp
--- a/Engine/Source/Runtime/Slate/Private/Widgets/Views/SHeaderRow.cpp
+++ b/Engine/Source/Runtime/Slate/Private/Widgets/Views/SHeaderRow.cpp
@@ -284,9 +284,9 @@ public:
 		return EVisibility::Visible;
 	}
 
-	FVector2D GetMenuOverlaySize() const 
+	UE::Slate::FDeprecateVector2DResult GetMenuOverlaySize() const 
 	{ 
-		return MenuOverlay.IsValid() ? MenuOverlay->GetDesiredSize() : FVector2D::ZeroVector; 
+		return MenuOverlay.IsValid() ? MenuOverlay->GetDesiredSize() : FVector2f::ZeroVector; 
 	}
 
 private:
@@ -421,7 +421,7 @@ private:
 	{
 		if ( ContextMenuContent != SNullWidget::NullWidget )
 		{
-			const FVector2D& SummonLocation = MouseEvent.GetScreenSpacePosition();
+			FVector2f SummonLocation = MouseEvent.GetScreenSpacePosition();
 			FWidgetPath WidgetPath = MouseEvent.GetEventPath() != nullptr ? *MouseEvent.GetEventPath() : FWidgetPath();
 
 			FSlateApplication::Get().CloseToolTip();
@@ -462,7 +462,7 @@ void SHeaderRow::Construct( const FArguments& InArgs )
 {
 	check(InArgs._Style);
 
-	ScrollBarThickness = FVector2D::ZeroVector;
+	ScrollBarThickness = FVector2f::ZeroVector;
 	ScrollBarVisibility = EVisibility::Collapsed;
 	Style = InArgs._Style;
 	OnGetMaxRowSizeForColumn = InArgs._OnGetMaxRowSizeForColumn;
@@ -608,7 +608,7 @@ FVector2D SHeaderRow::GetRowSizeForSlotIndex(int32 SlotIndex) const
 		const TSharedPtr<STableColumnHeader>& HeaderWidget = HeaderWidgets[SlotIndex];
 		const FColumn& Column = Columns[SlotIndex];
 
-		FVector2D HeaderSize = HeaderWidget->GetDesiredSize();
+		FVector2D HeaderSize = FVector2D(HeaderWidget->GetDesiredSize());
 
 		if (Column.HeaderMenuContent.Widget != SNullWidget::NullWidget && HeaderWidget->GetMenuOverlayVisibility() != EVisibility::Visible)
 		{
@@ -681,7 +681,7 @@ FReply SHeaderRow::OnMouseButtonUp(const FGeometry& MyGeometry, const FPointerEv
 {
 	if (bCanSelectGeneratedColumn && MouseEvent.GetEffectingButton() == EKeys::RightMouseButton)
 	{
-		const FVector2D& SummonLocation = MouseEvent.GetScreenSpacePosition();
+		FVector2f SummonLocation = MouseEvent.GetScreenSpacePosition();
 		FWidgetPath WidgetPath = MouseEvent.GetEventPath() != nullptr ? *MouseEvent.GetEventPath() : FWidgetPath();
 
 		const bool CloseAfterSelection = true;
diff --git a/Engine/Source/Runtime/Slate/Private/Widgets/Views/STableViewBase.cpp b/Engine/Source/Runtime/Slate/Private/Widgets/Views/STableViewBase.cpp
--- a/Engine/Source/Runtime/Slate/Private/Widgets/Views/STableViewBase.cpp
+++ b/Engine/Source/Runtime/Slate/Private/Widgets/Views/STableViewBase.cpp
@@ -33,11 +33,11 @@ FTableViewDimensions::FTableViewDimensions(EOrientation InOrientation)
 
 
 FTableViewDimensions::FTableViewDimensions(EOrientation InOrientation, float X, float Y)
-	: FTableViewDimensions(InOrientation, FVector2D(X, Y))
+	: FTableViewDimensions(InOrientation, FVector2f(X, Y))
 {
 }
 
-FTableViewDimensions::FTableViewDimensions(EOrientation InOrientation, const FVector2D& Size)
+FTableViewDimensions::FTableViewDimensions(EOrientation InOrientation, const UE::Slate::FDeprecateVector2DParameter& Size)
 	: FTableViewDimensions(InOrientation)
 {
 	if (InOrientation == Orient_Vertical)
@@ -473,7 +473,7 @@ FReply STableViewBase::OnMouseButtonUp( const FGeometry& MyGeometry, const FPoin
 		if ( HasMouseCapture() )
 		{
 			FSlateRect ListScreenSpaceRect = MyGeometry.GetLayoutBoundingRect();
-			FVector2D CursorPosition = MyGeometry.LocalToAbsolute( SoftwareCursorPosition );
+			FVector2f CursorPosition = MyGeometry.LocalToAbsolute( SoftwareCursorPosition );
 
 			FIntPoint BestPositionInList(
 				FMath::RoundToInt( FMath::Clamp( CursorPosition.X, ListScreenSpaceRect.Left, ListScreenSpaceRect.Right ) ),
@@ -763,12 +763,12 @@ int32 STableViewBase::OnPaint( const FPaintArgs& Args, const FGeometry& Allotted
 	}
 
 	const FSlateBrush* Brush = FCoreStyle::Get().GetBrush(TEXT("SoftwareCursor_Grab"));
-	const FVector2D CursorSize = Brush->ImageSize / AllottedGeometry.Scale;
+	const FVector2f CursorSize = Brush->ImageSize / AllottedGeometry.Scale;
 
 	FSlateDrawElement::MakeBox(
 		OutDrawElements,
 		++NewLayerId,
-		AllottedGeometry.ToPaintGeometry( SoftwareCursorPosition - (CursorSize / 2 ), CursorSize),
+		AllottedGeometry.ToPaintGeometry(CursorSize, FSlateLayoutTransform(SoftwareCursorPosition - (CursorSize / .5f ))),
 		Brush
 		);
 
@@ -984,7 +984,7 @@ float STableViewBase::GetItemHeight() const
 	return GetItemSize().Y;
 }
 
-FVector2D STableViewBase::GetItemSize() const
+UE::Slate::FDeprecateVector2DResult STableViewBase::GetItemSize() const
 {
 	FTableViewDimensions ItemDimensions = ItemsPanel->GetItemSize(PanelGeometryLastTick);
 	ItemDimensions.LineAxis += ItemsPanel->GetItemPadding(PanelGeometryLastTick);
@@ -1026,7 +1026,7 @@ void STableViewBase::NavigateToWidget(const uint32 UserIndex, const TSharedPtr<S
 
 void STableViewBase::OnRightMouseButtonUp(const FPointerEvent& MouseEvent)
 {
-	const FVector2D& SummonLocation = MouseEvent.GetScreenSpacePosition();
+	FVector2f SummonLocation = MouseEvent.GetScreenSpacePosition();
 	const bool bShouldOpenContextMenu = !IsRightClickScrolling();
 	const bool bContextMenuOpeningBound = OnContextMenuOpening.IsBound();
 
diff --git a/Engine/Source/Runtime/Slate/Public/Framework/Application/MenuStack.h b/Engine/Source/Runtime/Slate/Public/Framework/Application/MenuStack.h
--- a/Engine/Source/Runtime/Slate/Public/Framework/Application/MenuStack.h
+++ b/Engine/Source/Runtime/Slate/Public/Framework/Application/MenuStack.h
@@ -84,7 +84,7 @@ public:
 	 * @param bIsCollapsedByParent	Is this menu collapsed when a parent menu receives focus/activation? If false, only focus/activation outside the entire stack will auto collapse it.
 	 * @param bEnablePerPixelTransparency Does the menu's content require per pixel transparency?
 	 */
-	TSharedRef<IMenu> Push(const FWidgetPath& InOwnerPath, const TSharedRef<SWidget>& InContent, const FVector2D& SummonLocation, const FPopupTransitionEffect& TransitionEffect, const bool bFocusImmediately = true, const FVector2D& SummonLocationSize = FVector2D::ZeroVector, TOptional<EPopupMethod> InMethod = TOptional<EPopupMethod>(), const bool bIsCollapsedByParent = true, const bool bEnablePerPixelTransparency = false);
+	TSharedRef<IMenu> Push(const FWidgetPath& InOwnerPath, const TSharedRef<SWidget>& InContent, const UE::Slate::FDeprecateVector2DParameter& SummonLocation, const FPopupTransitionEffect& TransitionEffect, const bool bFocusImmediately = true, const UE::Slate::FDeprecateVector2DParameter& SummonLocationSize = FVector2f::ZeroVector, TOptional<EPopupMethod> InMethod = TOptional<EPopupMethod>(), const bool bIsCollapsedByParent = true, const bool bEnablePerPixelTransparency = false);
 	
 	/**
 	 * Pushes a new child menu onto the stack.
@@ -99,7 +99,7 @@ public:
 	 * @param bIsCollapsedByParent	Is this menu collapsed when a parent menu receives focus/activation? If false, only focus/activation outside the entire stack will auto collapse it.
 	 * @param bEnablePerPixelTransparency Does the menu's content require per pixel transparency?
 	 */
-	TSharedRef<IMenu> Push(const TSharedPtr<IMenu>& InParentMenu, const TSharedRef<SWidget>& InContent, const FVector2D& SummonLocation, const FPopupTransitionEffect& TransitionEffect, const bool bFocusImmediately = true, const FVector2D& SummonLocationSize = FVector2D::ZeroVector, const bool bIsCollapsedByParent = true, const bool bEnablePerPixelTransparency = false);
+	TSharedRef<IMenu> Push(const TSharedPtr<IMenu>& InParentMenu, const TSharedRef<SWidget>& InContent, const UE::Slate::FDeprecateVector2DParameter& SummonLocation, const FPopupTransitionEffect& TransitionEffect, const bool bFocusImmediately = true, const UE::Slate::FDeprecateVector2DParameter& SummonLocationSize = FVector2f::ZeroVector, const bool bIsCollapsedByParent = true, const bool bEnablePerPixelTransparency = false);
 
 	/**
 	 * Pushes a new menu onto the stack that is drawn by an external host widget.
@@ -266,10 +266,10 @@ private:
 	{
 		TSharedPtr<SWidget> WrappedContent;
 		TSharedPtr<SWidget> WidgetToFocus;
-		FVector2D AnimStartLocation;
-		FVector2D AnimFinalLocation;
+		FVector2f AnimStartLocation;
+		FVector2f AnimFinalLocation;
 		bool bAnchorSetsMinWidth;
-		FVector2D ExpectedSize;
+		FVector2f ExpectedSize;
 		bool bAllowAnimations;
 		bool bFocusImmediately;
 		bool bIsCollapsedByParent;
diff --git a/Engine/Source/Runtime/Slate/Public/Framework/Application/SlateApplication.h b/Engine/Source/Runtime/Slate/Public/Framework/Application/SlateApplication.h
--- a/Engine/Source/Runtime/Slate/Public/Framework/Application/SlateApplication.h
+++ b/Engine/Source/Runtime/Slate/Public/Framework/Application/SlateApplication.h
@@ -299,7 +299,7 @@ public:
 	 *
 	 * @param MouseCoordinate		The new position.
 	 */
-	void SetCursorPos( const FVector2D& MouseCoordinate );
+	void SetCursorPos( const FVector2D& MouseCoordinate ) override;
 
 	/** 
 	 *	Updates the cursor user's cursor to either the platform cursor or fake cursor
@@ -381,7 +381,7 @@ public:
 	 * @param Method				An optional popup method override. If not set, the widgets in the InOwnerPath will be queried for this.
 	 * @param bIsCollapsedByParent	Is this menu collapsed when a parent menu receives focus/activation? If false, only focus/activation outside the entire stack will auto collapse it.
 	 */
-	TSharedPtr<IMenu> PushMenu(const TSharedRef<SWidget>& InParentWidget, const FWidgetPath& InOwnerPath, const TSharedRef<SWidget>& InContent, const FVector2D& SummonLocation, const FPopupTransitionEffect& TransitionEffect, const bool bFocusImmediately = true, const FVector2D& SummonLocationSize = FVector2D::ZeroVector, TOptional<EPopupMethod> Method = TOptional<EPopupMethod>(), const bool bIsCollapsedByParent = true);
+	TSharedPtr<IMenu> PushMenu(const TSharedRef<SWidget>& InParentWidget, const FWidgetPath& InOwnerPath, const TSharedRef<SWidget>& InContent, const UE::Slate::FDeprecateVector2DParameter& SummonLocation, const FPopupTransitionEffect& TransitionEffect, const bool bFocusImmediately = true, const UE::Slate::FDeprecateVector2DParameter& SummonLocationSize = FVector2f::ZeroVector, TOptional<EPopupMethod> Method = TOptional<EPopupMethod>(), const bool bIsCollapsedByParent = true);
 
 	/**
 	 * Creates a new Menu and adds it to the menu stack under the specified parent menu.
@@ -395,7 +395,7 @@ public:
 	 * @param SummonLocationSize	An optional rect which describes an area in which the menu may not appear
 	 * @param bIsCollapsedByParent	Is this menu collapsed when a parent menu receives focus/activation? If false, only focus/activation outside the entire stack will auto collapse it.
 	 */
-	TSharedPtr<IMenu> PushMenu(const TSharedPtr<IMenu>& InParentMenu, const TSharedRef<SWidget>& InContent, const FVector2D& SummonLocation, const FPopupTransitionEffect& TransitionEffect, const bool bFocusImmediately = true, const FVector2D& SummonLocationSize = FVector2D::ZeroVector, const bool bIsCollapsedByParent = true);
+	TSharedPtr<IMenu> PushMenu(const TSharedPtr<IMenu>& InParentMenu, const TSharedRef<SWidget>& InContent, const UE::Slate::FDeprecateVector2DParameter& SummonLocation, const FPopupTransitionEffect& TransitionEffect, const bool bFocusImmediately = true, const UE::Slate::FDeprecateVector2DParameter& SummonLocationSize = FVector2f::ZeroVector, const bool bIsCollapsedByParent = true);
 
 	/**
 	 * Creates a new hosted Menu and adds it to the menu stack.
@@ -833,7 +833,7 @@ public:
 	 *								If horizontal it will attempt to open below the anchor but will open above if there is no room.
 	 * @return The adjusted position
 	 */
-	virtual FVector2D CalculatePopupWindowPosition( const FSlateRect& InAnchor, const FVector2D& InSize, bool bAutoAdjustForDPIScale = true, const FVector2D& InProposedPlacement = FVector2D::ZeroVector, const EOrientation Orientation = Orient_Vertical) const;
+	virtual UE::Slate::FDeprecateVector2DResult CalculatePopupWindowPosition( const FSlateRect& InAnchor, const UE::Slate::FDeprecateVector2DParameter& InSize, bool bAutoAdjustForDPIScale = true, const UE::Slate::FDeprecateVector2DParameter& InProposedPlacement = FVector2f::ZeroVector, const EOrientation Orientation = Orient_Vertical) const;
 
 	/**
 	 * Calculates the tooltip window position.
@@ -842,7 +842,7 @@ public:
 	 * @param InSize The size of the tooltip window.
 	 * @return The suggested position.
 	 */
-	virtual FVector2D CalculateTooltipWindowPosition( const FSlateRect& InAnchorRect, const FVector2D& InSize, bool bAutoAdjustForDPIScale) const;
+	virtual UE::Slate::FDeprecateVector2DResult CalculateTooltipWindowPosition( const FSlateRect& InAnchorRect, const UE::Slate::FDeprecateVector2DParameter& InSize, bool bAutoAdjustForDPIScale) const;
 
 	/**
 	 * Is the window in the app's destroy queue? If so it will be destroyed next tick.
@@ -1114,7 +1114,7 @@ protected:
 	/** Engages or disengages application throttling based on user behavior */
 	void ThrottleApplicationBasedOnMouseMovement();
 
-	virtual FWidgetPath LocateWidgetInWindow(FVector2D ScreenspaceMouseCoordinate, const TSharedRef<SWindow>& Window, bool bIgnoreEnabledStatus, int32 UserIndex) const override;
+	virtual FWidgetPath LocateWidgetInWindow(UE::Slate::FDeprecateVector2DParameter ScreenspaceMouseCoordinate, const TSharedRef<SWindow>& Window, bool bIgnoreEnabledStatus, int32 UserIndex) const override;
 
 	/**
 	 * Sets up any values that need to be based on the physical dimensions of the device.  
@@ -1323,7 +1323,7 @@ public:
 	 * @param	InToolTip           Widget to display.
 	 * @param	InSpawnLocation     Screen space location to show the tool tip (window top left)
 	 */
-	void SpawnToolTip( const TSharedRef<IToolTip>& InToolTip, const FVector2D& InSpawnLocation );
+	void SpawnToolTip( const TSharedRef<IToolTip>& InToolTip, const UE::Slate::FDeprecateVector2DParameter& InSpawnLocation );
 
 	/** Closes the open tool-tip, if a tool-tip is open */
 	void CloseToolTip();
@@ -1362,8 +1362,8 @@ public:
 	float GetDragTriggerDistanceSquared() const;
 
 	/** @return true if the difference between the ScreenSpaceOrigin and the ScreenSpacePosition is larger than the trigger distance for dragging in Slate. */
-	bool HasTraveledFarEnoughToTriggerDrag(const FPointerEvent& PointerEvent, const FVector2D ScreenSpaceOrigin) const;
-	bool HasTraveledFarEnoughToTriggerDrag(const FPointerEvent& PointerEvent, const FVector2D ScreenSpaceOrigin, EOrientation Orientation) const;
+	bool HasTraveledFarEnoughToTriggerDrag(const FPointerEvent& PointerEvent, const UE::Slate::FDeprecateVector2DParameter ScreenSpaceOrigin) const;
+	bool HasTraveledFarEnoughToTriggerDrag(const FPointerEvent& PointerEvent, const UE::Slate::FDeprecateVector2DParameter ScreenSpaceOrigin, EOrientation Orientation) const;
 
 	/** Set the size of the deadzone for dragging in screen pixels */
 	void SetDragTriggerDistance( float ScreenPixels );
@@ -1442,9 +1442,9 @@ public:
 	virtual bool GetSoftwareCursorAvailable() const override { return bSoftwareCursorAvailable; }
 	virtual EVisibility GetSoftwareCursorVis() const override;
 
-	virtual FVector2D GetCursorPos() const override;
-	virtual FVector2D GetLastCursorPos() const override;
-	virtual FVector2D GetCursorSize() const override;
+	virtual UE::Slate::FDeprecateVector2DResult GetCursorPos() const override;
+	virtual UE::Slate::FDeprecateVector2DResult GetLastCursorPos() const override;
+	virtual UE::Slate::FDeprecateVector2DResult GetCursorSize() const override;
 
 	virtual TSharedPtr<SWidget> GetKeyboardFocusedWidget() const override;
 
@@ -1467,7 +1467,7 @@ public:
 	virtual bool HasFocusedDescendants( const TSharedRef<const SWidget>& Widget ) const override;
 	virtual bool HasUserFocusedDescendants(const TSharedRef< const SWidget >& Widget, int32 UserIndex) const override;
 	virtual bool IsExternalUIOpened() override;
-	virtual FWidgetPath LocateWindowUnderMouse( FVector2D ScreenspaceMouseCoordinate, const TArray<TSharedRef<SWindow>>& Windows, bool bIgnoreEnabledStatus = false, int32 UserIndex = INDEX_NONE) override;
+	virtual FWidgetPath LocateWindowUnderMouse( UE::Slate::FDeprecateVector2DParameter ScreenspaceMouseCoordinate, const TArray<TSharedRef<SWindow>>& Windows, bool bIgnoreEnabledStatus = false, int32 UserIndex = INDEX_NONE) override;
 	virtual bool IsWindowHousingInteractiveTooltip(const TSharedRef<const SWindow>& WindowToTest) const override;
 	virtual TSharedRef<SImage> MakeImage( const TAttribute<const FSlateBrush*>& Image, const TAttribute<FSlateColor>& Color, const TAttribute<EVisibility>& Visibility ) const override;
 	virtual TSharedRef<SWidget> MakeWindowTitleBar(const FWindowTitleBarArgs& InArgs, TSharedPtr<IWindowTitleBar>& OutTitleBar) const override;
diff --git a/Engine/Source/Runtime/Slate/Public/Framework/Application/SlateUser.h b/Engine/Source/Runtime/Slate/Public/Framework/Application/SlateUser.h
--- a/Engine/Source/Runtime/Slate/Public/Framework/Application/SlateUser.h
+++ b/Engine/Source/Runtime/Slate/Public/Framework/Application/SlateUser.h
@@ -89,15 +89,15 @@ public:
 	void SetCursorVisibility(bool bDrawCursor);
 
 	void SetCursorPosition(int32 PosX, int32 PosY);
-	void SetCursorPosition(const FVector2D& NewCursorPos);
+	void SetCursorPosition(const UE::Slate::FDeprecateVector2DParameter& NewCursorPos);
 	void SetPointerPosition(uint32 PointerIndex, int32 PosX, int32 PosY);
-	void SetPointerPosition(uint32 PointerIndex, const FVector2D& NewPointerPos);
+	void SetPointerPosition(uint32 PointerIndex, const UE::Slate::FDeprecateVector2DParameter& NewPointerPos);
 	
-	FVector2D GetCursorPosition() const;
-	FVector2D GetPreviousCursorPosition() const;
+	UE::Slate::FDeprecateVector2DResult GetCursorPosition() const;
+	UE::Slate::FDeprecateVector2DResult GetPreviousCursorPosition() const;
 	
-	FVector2D GetPointerPosition(uint32 PointerIndex) const;
-	FVector2D GetPreviousPointerPosition(uint32 PointerIndex) const;
+	UE::Slate::FDeprecateVector2DResult GetPointerPosition(uint32 PointerIndex) const;
+	UE::Slate::FDeprecateVector2DResult GetPreviousPointerPosition(uint32 PointerIndex) const;
 
 	bool IsWidgetUnderCursor(TSharedPtr<const SWidget> Widget) const;
 	bool IsWidgetUnderPointer(TSharedPtr<const SWidget> Widget, uint32 PointerIndex) const;
@@ -116,7 +116,7 @@ public:
 	bool IsDragDroppingAffected(const FPointerEvent& InPointerEvent) const;
 	void CancelDragDrop();
 	
-	void ShowTooltip(const TSharedRef<IToolTip>& InTooltip, const FVector2D& InSpawnLocation);
+	void ShowTooltip(const TSharedRef<IToolTip>& InTooltip, const UE::Slate::FDeprecateVector2DParameter& InSpawnLocation);
 	void CloseTooltip();
 
 	const FGestureDetector& GetGestureDetector() const { return GestureDetector; }
@@ -168,7 +168,7 @@ SLATE_SCOPE:
 	void NotifyPointerReleased(const FPointerEvent& PointerEvent, const FWidgetPath& WidgetsUnderCursor, TSharedPtr<FDragDropOperation> DroppedContent, bool bWasHandled);
 	void UpdatePointerPosition(const FPointerEvent& PointerEvent);
 
-	void StartDragDetection(const FWidgetPath& PathToWidget, int32 PointerIndex, FKey DragButton, FVector2D StartLocation);
+	void StartDragDetection(const FWidgetPath& PathToWidget, int32 PointerIndex, FKey DragButton, UE::Slate::FDeprecateVector2DParameter StartLocation);
 	FWidgetPath DetectDrag(const FPointerEvent& PointerEvent, float DragTriggerDistance);
 	bool IsDetectingDrag(uint32 PointerIndex) const;
 	void ResetDragDetection();
@@ -195,7 +195,7 @@ private:
 	FSlateUser(int32 InUserIndex, TSharedPtr<ICursor> InCursor);
 	
 	FSlateUser(FPlatformUserId InPlatformUser, TSharedPtr<ICursor> InCursor);
-	void UpdatePointerPosition(uint32 PointerIndex, const FVector2D& Position);
+	void UpdatePointerPosition(uint32 PointerIndex, const FVector2f& Position);
 	void LockCursorInternal(const FWidgetPath& WidgetPath);
 	TSharedRef<SWindow> GetOrCreateTooltipWindow();
 
@@ -246,15 +246,15 @@ private:
 	int32 FocusVersion = 0;
 
 	/** Current position of all pointers controlled by this user */
-	TMap<uint32, FVector2D> PointerPositionsByIndex;
-	TMap<uint32, FVector2D> PreviousPointerPositionsByIndex;
+	TMap<uint32, FVector2f> PointerPositionsByIndex;
+	TMap<uint32, FVector2f> PreviousPointerPositionsByIndex;
 
 	/** Weak paths to widgets that are currently capturing a particular pointer */
 	TMap<uint32, FWeakWidgetPath> PointerCaptorPathsByIndex;
 
 	struct FDragDetectionState
 	{
-		FDragDetectionState(const FWidgetPath& PathToWidget, int32 PointerIndex, FKey DragButton, const FVector2D& StartLocation)
+		FDragDetectionState(const FWidgetPath& PathToWidget, int32 PointerIndex, FKey DragButton, const FVector2f& StartLocation)
 			: DetectDragForWidget(PathToWidget)
 			, DragStartLocation(StartLocation)
 			, TriggerButton(DragButton)
@@ -265,7 +265,7 @@ private:
 		/** If not null, a widget has requested that we detect a drag being triggered in this widget and send an OnDragDetected() event*/
 		FWeakWidgetPath DetectDragForWidget;
 
-		FVector2D DragStartLocation = FVector2D::ZeroVector;
+		FVector2f DragStartLocation = FVector2f::ZeroVector;
 		FKey TriggerButton = EKeys::Invalid;
 		int32 PointerIndex = INDEX_NONE;
 	};
@@ -311,7 +311,7 @@ private:
 		TWeakPtr<SWidget> SourceWidget;
 
 		/** Desired position of the tooltip in screen space, updated whenever the mouse moves */
-		FVector2D DesiredLocation = FVector2D::ZeroVector;
+		FVector2f DesiredLocation = FVector2f::ZeroVector;
 
 		/** The time at which the tooltip was summoned */
 		double SummonTime = 0.0;
diff --git a/Engine/Source/Runtime/Slate/Public/Framework/Layout/ScrollyZoomy.h b/Engine/Source/Runtime/Slate/Public/Framework/Layout/ScrollyZoomy.h
--- a/Engine/Source/Runtime/Slate/Public/Framework/Layout/ScrollyZoomy.h
+++ b/Engine/Source/Runtime/Slate/Public/Framework/Layout/ScrollyZoomy.h
@@ -139,7 +139,7 @@ public:
 	 * @return Cursor position.
 	 * @see NeedsSoftwareCursor
 	 */
-	const FVector2D& GetSoftwareCursorPosition() const;
+	UE::Slate::FDeprecateVector2DResult GetSoftwareCursorPosition() const;
 
 	/**
 	 * Call this from your widget's OnPaint to paint a software cursor, if needed 
@@ -161,7 +161,7 @@ private:
 	bool bShowSoftwareCursor;
 
 	/** The current position of the software cursor. */
-	FVector2D SoftwareCursorPosition;
+	FVector2f SoftwareCursorPosition;
 
 	/** Whether inertial scrolling is enabled. */
 	bool UseInertialScrolling;
diff --git a/Engine/Source/Runtime/Slate/Public/Framework/Text/SlateTextUnderlineLineHighlighter.h b/Engine/Source/Runtime/Slate/Public/Framework/Text/SlateTextUnderlineLineHighlighter.h
--- a/Engine/Source/Runtime/Slate/Public/Framework/Text/SlateTextUnderlineLineHighlighter.h
+++ b/Engine/Source/Runtime/Slate/Public/Framework/Text/SlateTextUnderlineLineHighlighter.h
@@ -23,7 +23,7 @@ public:
 	static const int32 DefaultZIndex = 1;
 
 protected:
-	ISlateTextLineHighlighter(const FSlateBrush& InLineBrush, const FSlateFontInfo& InFontInfo, const FSlateColor InColorAndOpacity, const FVector2D InShadowOffset, const FLinearColor InShadowColorAndOpacity);
+	ISlateTextLineHighlighter(const FSlateBrush& InLineBrush, const FSlateFontInfo& InFontInfo, const FSlateColor InColorAndOpacity, const UE::Slate::FDeprecateVector2DParameter InShadowOffset, const FLinearColor InShadowColorAndOpacity);
 
 	virtual void GetLineMetrics(const float InFontScale, int16& OutLinePos, int16& OutLineThickness) const = 0;
 
@@ -37,7 +37,7 @@ protected:
 	FSlateColor ColorAndOpacity;
 
 	/** Offset at which to draw the shadow (if any) */
-	FVector2D ShadowOffset;
+	UE::Slate::FDeprecateVector2DResult ShadowOffset;
 
 	/** The color to draw the shadow */
 	FLinearColor ShadowColorAndOpacity;
@@ -47,10 +47,10 @@ protected:
 class SLATE_API FSlateTextUnderlineLineHighlighter : public ISlateTextLineHighlighter
 {
 public:
-	static TSharedRef<FSlateTextUnderlineLineHighlighter> Create(const FSlateBrush& InUnderlineBrush, const FSlateFontInfo& InFontInfo, const FSlateColor InColorAndOpacity, const FVector2D InShadowOffset, const FLinearColor InShadowColorAndOpacity);
+	static TSharedRef<FSlateTextUnderlineLineHighlighter> Create(const FSlateBrush& InUnderlineBrush, const FSlateFontInfo& InFontInfo, const FSlateColor InColorAndOpacity, const UE::Slate::FDeprecateVector2DParameter InShadowOffset, const FLinearColor InShadowColorAndOpacity);
 
 protected:
-	FSlateTextUnderlineLineHighlighter(const FSlateBrush& InUnderlineBrush, const FSlateFontInfo& InFontInfo, const FSlateColor InColorAndOpacity, const FVector2D InShadowOffset, const FLinearColor InShadowColorAndOpacity);
+	FSlateTextUnderlineLineHighlighter(const FSlateBrush& InUnderlineBrush, const FSlateFontInfo& InFontInfo, const FSlateColor InColorAndOpacity, const UE::Slate::FDeprecateVector2DParameter InShadowOffset, const FLinearColor InShadowColorAndOpacity);
 
 	virtual void GetLineMetrics(const float InFontScale, int16& OutLinePos, int16& OutLineThickness) const override;
 };
@@ -59,10 +59,10 @@ protected:
 class SLATE_API FSlateTextStrikeLineHighlighter : public ISlateTextLineHighlighter
 {
 public:
-	static TSharedRef<FSlateTextStrikeLineHighlighter> Create(const FSlateBrush& InUnderlineBrush, const FSlateFontInfo& InFontInfo, const FSlateColor InColorAndOpacity, const FVector2D InShadowOffset, const FLinearColor InShadowColorAndOpacity);
+	static TSharedRef<FSlateTextStrikeLineHighlighter> Create(const FSlateBrush& InUnderlineBrush, const FSlateFontInfo& InFontInfo, const FSlateColor InColorAndOpacity, const UE::Slate::FDeprecateVector2DParameter InShadowOffset, const FLinearColor InShadowColorAndOpacity);
 
 protected:
-	FSlateTextStrikeLineHighlighter(const FSlateBrush& InStrikeBrush, const FSlateFontInfo& InFontInfo, const FSlateColor InColorAndOpacity, const FVector2D InShadowOffset, const FLinearColor InShadowColorAndOpacity);
+	FSlateTextStrikeLineHighlighter(const FSlateBrush& InStrikeBrush, const FSlateFontInfo& InFontInfo, const FSlateColor InColorAndOpacity, const UE::Slate::FDeprecateVector2DParameter InShadowOffset, const FLinearColor InShadowColorAndOpacity);
 
 	virtual void GetLineMetrics(const float InFontScale, int16& OutLinePos, int16& OutLineThickness) const override;
 };
diff --git a/Engine/Source/Runtime/Slate/Public/SlateMaterialBrush.h b/Engine/Source/Runtime/Slate/Public/SlateMaterialBrush.h
--- a/Engine/Source/Runtime/Slate/Public/SlateMaterialBrush.h
+++ b/Engine/Source/Runtime/Slate/Public/SlateMaterialBrush.h
@@ -23,7 +23,7 @@ struct FSlateMaterialBrush : public FSlateBrush
 	 * @param InMaterial The material to use.
 	 * @param InImageSize The material's dimensions.
 	 */
-	FSlateMaterialBrush( class UMaterialInterface& InMaterial, const FVector2D& InImageSize )
+	FSlateMaterialBrush( class UMaterialInterface& InMaterial, const UE::Slate::FDeprecateVector2DParameter& InImageSize )
 		: FSlateBrush( ESlateBrushDrawType::Image, NAME_None, FMargin(0), ESlateBrushTileType::NoTile, ESlateBrushImageType::FullColor, InImageSize, FLinearColor::White, &InMaterial )
 	{
 		ResourceName = FName( *InMaterial.GetFullName() );
@@ -34,7 +34,7 @@ struct FSlateMaterialBrush : public FSlateBrush
 	 *
 	 * @param InImageSize The material's dimensions.
 	 */
-	FSlateMaterialBrush(const FVector2D& InImageSize)
+	FSlateMaterialBrush(const UE::Slate::FDeprecateVector2DParameter& InImageSize)
 		: FSlateBrush(ESlateBrushDrawType::Image, FName(TEXT("None")), FMargin(0), ESlateBrushTileType::NoTile, ESlateBrushImageType::FullColor, InImageSize, FLinearColor::White)
 	{}
 
diff --git a/Engine/Source/Runtime/Slate/Public/Widgets/Colors/SColorGradingWheel.h b/Engine/Source/Runtime/Slate/Public/Widgets/Colors/SColorGradingWheel.h
--- a/Engine/Source/Runtime/Slate/Public/Widgets/Colors/SColorGradingWheel.h
+++ b/Engine/Source/Runtime/Slate/Public/Widgets/Colors/SColorGradingWheel.h
@@ -77,7 +77,7 @@ protected:
 	 *
 	 * @return The position relative to the widget.
 	 */
-	FVector2D CalcRelativePositionFromCenter() const;
+	UE::Slate::FDeprecateVector2DResult CalcRelativePositionFromCenter() const;
 
 	/**
 	 * Performs actions according to mouse click / move
diff --git a/Engine/Source/Runtime/Slate/Public/Widgets/Colors/SColorWheel.h b/Engine/Source/Runtime/Slate/Public/Widgets/Colors/SColorWheel.h
--- a/Engine/Source/Runtime/Slate/Public/Widgets/Colors/SColorWheel.h
+++ b/Engine/Source/Runtime/Slate/Public/Widgets/Colors/SColorWheel.h
@@ -70,7 +70,7 @@ protected:
 	 *
 	 * @return The position relative to the widget.
 	 */
-	FVector2D CalcRelativePositionFromCenter() const;
+	UE::Slate::FDeprecateVector2DResult CalcRelativePositionFromCenter() const;
 
 	/**
 	 * Performs actions according to mouse click / move
diff --git a/Engine/Source/Runtime/Slate/Public/Widgets/Docking/SDockTab.h b/Engine/Source/Runtime/Slate/Public/Widgets/Docking/SDockTab.h
--- a/Engine/Source/Runtime/Slate/Public/Widgets/Docking/SDockTab.h
+++ b/Engine/Source/Runtime/Slate/Public/Widgets/Docking/SDockTab.h
@@ -478,7 +478,7 @@ protected:
 	TAttribute<FLinearColor> IconColor;
 
 	/** @return the scaling of the tab based on the opening/closing animation */
-	FVector2D GetAnimatedScale() const;
+	UE::Slate::FDeprecateVector2DResult GetAnimatedScale() const;
 
 	/** Animation that shows the tab opening up */
 	FCurveSequence SpawnAnimCurve;
diff --git a/Engine/Source/Runtime/Slate/Public/Widgets/Input/SSlider.h b/Engine/Source/Runtime/Slate/Public/Widgets/Input/SSlider.h
--- a/Engine/Source/Runtime/Slate/Public/Widgets/Input/SSlider.h
+++ b/Engine/Source/Runtime/Slate/Public/Widgets/Input/SSlider.h
@@ -200,7 +200,7 @@ protected:
 	 * @param AbsolutePosition The absolute position of the slider.
 	 * @return The new value.
 	 */
-	float PositionToValue( const FGeometry& MyGeometry, const FVector2D& AbsolutePosition );
+	float PositionToValue( const FGeometry& MyGeometry, const UE::Slate::FDeprecateVector2DParameter& AbsolutePosition );
 
 	const FSlateBrush* GetBarImage() const;
 	const FSlateBrush* GetThumbImage() const;
@@ -259,7 +259,7 @@ protected:
 	EMouseCursor::Type CachedCursor;
 
 	/** The location in screenspace the slider was pressed by a touch */
-	FVector2D PressedScreenSpaceTouchDownPosition;
+	UE::Slate::FDeprecateVector2DResult PressedScreenSpaceTouchDownPosition;
 
 	/** Holds the amount to adjust the value by when using a controller or keyboard */
 	TAttribute<float> StepSize;
diff --git a/Engine/Source/Runtime/Slate/Public/Widgets/Input/SSpinBox.h b/Engine/Source/Runtime/Slate/Public/Widgets/Input/SSpinBox.h
--- a/Engine/Source/Runtime/Slate/Public/Widgets/Input/SSpinBox.h
+++ b/Engine/Source/Runtime/Slate/Public/Widgets/Input/SSpinBox.h
@@ -367,7 +367,7 @@ public:
 				FSlateDrawElement::MakeBox(
 					OutDrawElements,
 					FilledLayer,
-					AllottedGeometry.ToPaintGeometry(Style->InsetPadding.GetTopLeft(), FillSize - FVector2D(Style->InsetPadding.GetTotalSpaceAlong<Orient_Horizontal>(), Style->InsetPadding.GetTotalSpaceAlong<Orient_Vertical>())),
+					AllottedGeometry.ToPaintGeometry(FillSize - FVector2D(Style->InsetPadding.GetTotalSpaceAlong<Orient_Horizontal>(), Style->InsetPadding.GetTotalSpaceAlong<Orient_Vertical>()), FSlateLayoutTransform(Style->InsetPadding.GetTopLeft())),
 					FillImage,
 					DrawEffects,
 					FillImage->GetTint(InWidgetStyle) * InWidgetStyle.GetColorAndOpacityTint()
diff --git a/Engine/Source/Runtime/Slate/Public/Widgets/Layout/SScrollBox.h b/Engine/Source/Runtime/Slate/Public/Widgets/Layout/SScrollBox.h
--- a/Engine/Source/Runtime/Slate/Public/Widgets/Layout/SScrollBox.h
+++ b/Engine/Source/Runtime/Slate/Public/Widgets/Layout/SScrollBox.h
@@ -181,7 +181,7 @@ public:
 		, _ScrollBarVisibility(EVisibility::Visible)
 		, _ScrollBarAlwaysVisible(false)
 		, _ScrollBarDragFocusCause(EFocusCause::Mouse)
-		, _ScrollBarThickness(FVector2D(_Style->BarThickness, _Style->BarThickness))
+		, _ScrollBarThickness(FVector2f(_Style->BarThickness, _Style->BarThickness))
 		, _ScrollBarPadding(2.0f)
 		, _AllowOverscroll(EAllowOverscroll::Yes)
 		, _BackPadScrolling(false)
@@ -217,7 +217,7 @@ public:
 
 		SLATE_ARGUMENT( EFocusCause, ScrollBarDragFocusCause )
 
-		SLATE_ARGUMENT( FVector2D, ScrollBarThickness )
+		SLATE_ARGUMENT( UE::Slate::FDeprecateVector2DParameter, ScrollBarThickness )
 
 		/** This accounts for total internal scroll bar padding; default 2.0f padding from the scroll bar itself is removed */
 		SLATE_ARGUMENT( FMargin, ScrollBarPadding )
@@ -322,7 +322,7 @@ public:
 	
 	void SetScrollBarTrackAlwaysVisible(bool InAlwaysVisible);
 
-	void SetScrollBarThickness(FVector2D InThickness);
+	void SetScrollBarThickness(UE::Slate::FDeprecateVector2DParameter InThickness);
 
 	void SetScrollBarPadding(const FMargin& InPadding);
 
@@ -362,13 +362,13 @@ private:
 	void ConstructHorizontalLayout();
 
 	/** Gets the component of a vector in the direction of scrolling based on the Orientation property. */
-	FORCEINLINE float GetScrollComponentFromVector(FVector2D Vector) const
+	FORCEINLINE float GetScrollComponentFromVector(FVector2f Vector) const
 	{
 		return float(Orientation == Orient_Vertical ? Vector.Y : Vector.X);
 	}
 
 	/** Sets the component of a vector in the direction of scrolling based on the Orientation property. */
-	inline void SetScrollComponentOnVector(FVector2D& InVector, float Value) const
+	inline void SetScrollComponentOnVector(FVector2f& InVector, float Value) const
 	{
 		if (Orientation == Orient_Vertical)
 		{
@@ -479,7 +479,7 @@ protected:
 	EScrollWhenFocusChanges ScrollWhenFocusChanges;
 
 	/**	The current position of the software cursor */
-	FVector2D SoftwareCursorPosition;
+	FVector2f SoftwareCursorPosition;
 
 	/** Fired when the user scrolls the scrollbox */
 	FOnUserScrolled OnUserScrolled;
diff --git a/Engine/Source/Runtime/Slate/Public/Widgets/SViewport.h b/Engine/Source/Runtime/Slate/Public/Widgets/SViewport.h
--- a/Engine/Source/Runtime/Slate/Public/Widgets/SViewport.h
+++ b/Engine/Source/Runtime/Slate/Public/Widgets/SViewport.h
@@ -78,7 +78,7 @@ public:
 		/** Size of the viewport widget. */
 		SLATE_ATTRIBUTE(FVector2D, ViewportSize);
 
-		SLATE_API static FVector2D GetDefaultViewportSize();
+		SLATE_API static UE::Slate::FDeprecateVector2DResult GetDefaultViewportSize();
 
 	SLATE_END_ARGS()
 
diff --git a/Engine/Source/Runtime/Slate/Public/Widgets/Text/SMultiLineEditableText.h b/Engine/Source/Runtime/Slate/Public/Widgets/Text/SMultiLineEditableText.h
--- a/Engine/Source/Runtime/Slate/Public/Widgets/Text/SMultiLineEditableText.h
+++ b/Engine/Source/Runtime/Slate/Public/Widgets/Text/SMultiLineEditableText.h
@@ -605,7 +605,7 @@ protected:
 	bool bIsSoftwareCursor;
 
 	/**	The current position of the software cursor */
-	FVector2D SoftwareCursorPosition;
+	FVector2f SoftwareCursorPosition;
 
 	/** Callback delegate to have first chance handling of the OnKeyChar event */
 	FOnKeyChar OnKeyCharHandler;
diff --git a/Engine/Source/Runtime/Slate/Public/Widgets/Text/STextBlock.h b/Engine/Source/Runtime/Slate/Public/Widgets/Text/STextBlock.h
--- a/Engine/Source/Runtime/Slate/Public/Widgets/Text/STextBlock.h
+++ b/Engine/Source/Runtime/Slate/Public/Widgets/Text/STextBlock.h
@@ -272,7 +272,7 @@ private:
 	ETextTransformPolicy GetTransformPolicyImpl() const;
 	
 	/** Gets the current shadow offset */
-	FVector2D GetShadowOffset() const;
+	FVector2f GetShadowOffset() const;
 
 	/** Gets the current shadow color and opacity */
 	FLinearColor GetShadowColorAndOpacity() const;
diff --git a/Engine/Source/Runtime/Slate/Public/Widgets/Text/SlateEditableTextLayout.h b/Engine/Source/Runtime/Slate/Public/Widgets/Text/SlateEditableTextLayout.h
--- a/Engine/Source/Runtime/Slate/Public/Widgets/Text/SlateEditableTextLayout.h
+++ b/Engine/Source/Runtime/Slate/Public/Widgets/Text/SlateEditableTextLayout.h
@@ -142,16 +142,16 @@ public:
 	void AdvanceSearch(const bool InReverse = false);
 
 	/** Update the horizontal scroll amount from the given fraction */
-	FVector2D SetHorizontalScrollFraction(const float InScrollOffsetFraction);
+	UE::Slate::FDeprecateVector2DResult SetHorizontalScrollFraction(const float InScrollOffsetFraction);
 
 	/** Update the vertical scroll amount from the given fraction */
-	FVector2D SetVerticalScrollFraction(const float InScrollOffsetFraction);
+	UE::Slate::FDeprecateVector2DResult SetVerticalScrollFraction(const float InScrollOffsetFraction);
 
 	/** Set the absolute scroll offset value */
-	FVector2D SetScrollOffset(const FVector2D& InScrollOffset, const FGeometry& InGeometry);
+	UE::Slate::FDeprecateVector2DResult SetScrollOffset(const UE::Slate::FDeprecateVector2DParameter& InScrollOffset, const FGeometry& InGeometry);
 
 	/** Get the absolute scroll offset value */
-	FVector2D GetScrollOffset() const;
+	UE::Slate::FDeprecateVector2DResult GetScrollOffset() const;
 
 	/** Returns the computed wrap location for this layout */
 	float GetComputedWrappingWidth() const;
@@ -211,10 +211,10 @@ public:
 	bool AnyTextSelected() const;
 
 	/** Query to see if the text under the given position is currently selected */
-	bool IsTextSelectedAt(const FGeometry& MyGeometry, const FVector2D& ScreenSpacePosition) const;
+	bool IsTextSelectedAt(const FGeometry& MyGeometry, const UE::Slate::FDeprecateVector2DParameter& ScreenSpacePosition) const;
 
 	/** Query to see if the text under the given position is currently selected (the position is local to the text layout space) */
-	bool IsTextSelectedAt(const FVector2D& InLocalPosition) const;
+	bool IsTextSelectedAt(const UE::Slate::FDeprecateVector2DParameter& InLocalPosition) const;
 
 	/** Are we able to execute the "Select All" command? */
 	bool CanExecuteSelectAll() const;
@@ -223,10 +223,10 @@ public:
 	void SelectAllText();
 
 	/** Select the word under the given position */
-	void SelectWordAt(const FGeometry& MyGeometry, const FVector2D& ScreenSpacePosition);
+	void SelectWordAt(const FGeometry& MyGeometry, const UE::Slate::FDeprecateVector2DParameter& ScreenSpacePosition);
 
 	/** Select the word under the given position (the position is local to the text layout space) */
-	void SelectWordAt(const FVector2D& InLocalPosition);
+	void SelectWordAt(const UE::Slate::FDeprecateVector2DParameter& InLocalPosition);
 
 	/** Select a block of text */
 	void SelectText(const FTextLocation& InSelectionStart, const FTextLocation& InCursorLocation);
@@ -388,7 +388,7 @@ public:
 
 	void OnArrangeChildren(const FGeometry& AllottedGeometry, FArrangedChildren& ArrangedChildren) const;
 
-	FVector2D GetSize() const;
+	UE::Slate::FDeprecateVector2DResult GetSize() const;
 
 	TSharedRef<SWidget> BuildDefaultContextMenu(const TSharedPtr<FExtender>& InMenuExtender) const;
 
diff --git a/Engine/Source/Runtime/Slate/Public/Widgets/Views/SHeaderRow.h b/Engine/Source/Runtime/Slate/Public/Widgets/Views/SHeaderRow.h
--- a/Engine/Source/Runtime/Slate/Public/Widgets/Views/SHeaderRow.h
+++ b/Engine/Source/Runtime/Slate/Public/Widgets/Views/SHeaderRow.h
@@ -371,7 +371,7 @@ private:
 	TIndirectArray<FColumn> Columns;
 	TArray<TSharedPtr<class STableColumnHeader>> HeaderWidgets;
 
-	FVector2D ScrollBarThickness;
+	FVector2f ScrollBarThickness;
 	TAttribute< EVisibility > ScrollBarVisibility;
 	const FHeaderRowStyle* Style;
 	FColumnsChanged ColumnsChanged;
