Use Ceil instead of Trunc to calculate if SClippingHorizontalBox should add a wrap button.

diff --git a/Engine/Source/Runtime/Slate/Private/Framework/MultiBox/SClippingHorizontalBox.cpp b/Engine/Source/Runtime/Slate/Private/Framework/MultiBox/SClippingHorizontalBox.cpp
--- a/Engine/Source/Runtime/Slate/Private/Framework/MultiBox/SClippingHorizontalBox.cpp
+++ b/Engine/Source/Runtime/Slate/Private/Framework/MultiBox/SClippingHorizontalBox.cpp
@@ -26,9 +26,12 @@ void SClippingHorizontalBox::OnArrangeChildren( const FGeometry& AllottedGeometr
 	for (int32 ChildIdx = NumChildren - 2; ChildIdx >= 0; --ChildIdx)
 	{
 		const FArrangedWidget& CurWidget = ArrangedChildren[ChildIdx];
+		// Ceil (Minus a tad for float precision) to ensure contents are not a sub-pixel larger than the box, which will create an unnecessary wrap button
 		FVector2D AbsWidgetPos = CurWidget.Geometry.LocalToAbsolute(CurWidget.Geometry.GetLocalSize());
 		FVector2D AbsBoxPos = FVector2D(AllottedGeometry.AbsolutePosition) + AllottedGeometry.GetLocalSize() * AllottedGeometry.Scale;
-		if (FMath::TruncToInt(AbsWidgetPos.X) > FMath::TruncToInt(AbsBoxPos.X))
+		const FVector2D WidgetPosRounded = FVector2D(FMath::CeilToInt(AbsWidgetPos.X - KINDA_SMALL_NUMBER), FMath::CeilToInt(AbsWidgetPos.Y - KINDA_SMALL_NUMBER));
+		const FVector2D BoxPosRounded = FVector2D(FMath::CeilToInt(AbsBoxPos.X - KINDA_SMALL_NUMBER), FMath::CeilToInt(AbsBoxPos.Y - KINDA_SMALL_NUMBER));
+		if (WidgetPosRounded.X > BoxPosRounded.X)
 		{
 			++NumClippedChildren;
 			ArrangedChildren.Remove(ChildIdx);
