Slate:

-Add option to hide header menu combo button on table view headers

diff --git a/Engine/Source/Runtime/Slate/Private/Widgets/Views/SHeaderRow.cpp b/Engine/Source/Runtime/Slate/Private/Widgets/Views/SHeaderRow.cpp
--- a/Engine/Source/Runtime/Slate/Private/Widgets/Views/SHeaderRow.cpp
+++ b/Engine/Source/Runtime/Slate/Private/Widgets/Views/SHeaderRow.cpp
@@ -148,54 +148,57 @@ public:
 		if( Column.HeaderMenuContent.Widget != SNullWidget::NullWidget ||
 			Column.OnGetMenuContent.IsBound())
 		{
-			// Add Drop down menu button (only if menu content has been specified)
-			Box->AddSlot()
-			.AutoWidth()
-			[
-				SAssignNew( MenuOverlay, SOverlay )
-				.Visibility( this, &STableColumnHeader::GetMenuOverlayVisibility )
-
-				+SOverlay::Slot()
+			if (!Column.bHideHeaderMenuButton)
+			{
+				// Add Drop down menu button (only if menu content has been specified)
+				Box->AddSlot()
+				.AutoWidth()
 				[
-					SNew( SBorder )
-					.Padding( FMargin( 0.0f ) )
-					.BorderImage( this, &STableColumnHeader::GetComboButtonBorderBrush )
+					SAssignNew( MenuOverlay, SOverlay )
+					.Visibility( this, &STableColumnHeader::GetMenuOverlayVisibility )
+
+					+SOverlay::Slot()
 					[
-						SAssignNew( ComboButton, SComboButton )
-						.HasDownArrow(false)
-						.ButtonStyle( FAppStyle::Get(), "NoBorder" )
-						.ContentPadding( FMargin(0) )
-						.ButtonContent()
+						SNew( SBorder )
+						.Padding( FMargin( 0.0f ) )
+						.BorderImage( this, &STableColumnHeader::GetComboButtonBorderBrush )
 						[
-							SNew( SSpacer )
-							.Size( FVector2D( 14.0f, 0 ) )
+							SAssignNew( ComboButton, SComboButton )
+							.HasDownArrow(false)
+							.ButtonStyle( FAppStyle::Get(), "NoBorder" )
+							.ContentPadding( FMargin(0) )
+							.ButtonContent()
+							[
+								SNew( SSpacer )
+								.Size( FVector2D( 14.0f, 0 ) )
+							]
 						]
 					]
-				]
 
-				+SOverlay::Slot()
-				.HAlign( HAlign_Center )
-				.VAlign( VAlign_Center )
-				[
-					SNew( SBox )
-					.HeightOverride( 18.0f )
-					.Padding( FMargin( 0.0f, -2.0f ) )
+					+SOverlay::Slot()
+					.HAlign( HAlign_Center )
+					.VAlign( VAlign_Center )
 					[
-						SNew( SImage )
-						.Image( &Style->MenuDropdownImage )
-						.ColorAndOpacity( this, &STableColumnHeader::GetComboButtonTint )
-						.Visibility( EVisibility::HitTestInvisible )
+						SNew( SBox )
+						.HeightOverride( 18.0f )
+						.Padding( FMargin( 0.0f, -2.0f ) )
+						[
+							SNew( SImage )
+							.Image( &Style->MenuDropdownImage )
+							.ColorAndOpacity( this, &STableColumnHeader::GetComboButtonTint )
+							.Visibility( EVisibility::HitTestInvisible )
+						]
 					]
-				]
-			];
+				];
 
-			if (Column.HeaderMenuContent.Widget != SNullWidget::NullWidget)
-			{
-				ComboButton->SetMenuContent( ContextMenuContent );
-			}
-			else if (Column.OnGetMenuContent.IsBound())
-			{
-				ComboButton->SetOnGetMenuContent( Column.OnGetMenuContent );
+				if (Column.HeaderMenuContent.Widget != SNullWidget::NullWidget)
+				{
+					ComboButton->SetMenuContent( ContextMenuContent );
+				}
+				else if (Column.OnGetMenuContent.IsBound())
+				{
+					ComboButton->SetOnGetMenuContent( Column.OnGetMenuContent );
+				}
 			}
 		}
 
diff --git a/Engine/Source/Runtime/Slate/Public/Widgets/Views/SHeaderRow.h b/Engine/Source/Runtime/Slate/Public/Widgets/Views/SHeaderRow.h
--- a/Engine/Source/Runtime/Slate/Public/Widgets/Views/SHeaderRow.h
+++ b/Engine/Source/Runtime/Slate/Public/Widgets/Views/SHeaderRow.h
@@ -2,7 +2,6 @@
 
 #pragma once
 
-#include "CoreMinimal.h"
 #include "Containers/IndirectArray.h"
 #include "Misc/Attribute.h"
 #include "Layout/Visibility.h"
@@ -114,6 +113,7 @@ public:
 			, _HeaderContentPadding()
 			, _HeaderComboVisibility(EHeaderComboVisibility::OnHover)
 			, _MenuContent()
+			, _HideHeaderMenuButton(false)
 			, _HAlignCell( HAlign_Fill )
 			, _VAlignCell( VAlign_Fill )
 			, _InitialSortMode( EColumnSortMode::Ascending )
@@ -144,6 +144,8 @@ public:
 			SLATE_NAMED_SLOT( FArguments, MenuContent )
 			SLATE_EVENT( FOnGetContent, OnGetMenuContent )
 
+			SLATE_ARGUMENT(bool, HideHeaderMenuButton)
+
 			SLATE_ARGUMENT( EHorizontalAlignment, HAlignCell )
 			SLATE_ARGUMENT( EVerticalAlignment, VAlignCell )
 
@@ -167,6 +169,7 @@ public:
 			, HeaderContent( InArgs._HeaderContent )
 			, HeaderMenuContent( InArgs._MenuContent )
 			, OnGetMenuContent( InArgs._OnGetMenuContent )
+			, bHideHeaderMenuButton(InArgs._HideHeaderMenuButton)
 			, HeaderHAlignment( InArgs._HAlignHeader )
 			, HeaderVAlignment( InArgs._VAlignHeader )
 			, HeaderContentPadding( InArgs._HeaderContentPadding )
@@ -250,6 +253,8 @@ public:
 
 		FOnGetContent OnGetMenuContent;
 
+		bool bHideHeaderMenuButton;
+
 		EHorizontalAlignment HeaderHAlignment;
 		EVerticalAlignment HeaderVAlignment;
 		TOptional< FMargin > HeaderContentPadding;
