SNumericEntryBox: add support for configurable min/max fractionnal digit.
Necessary for the Font DPI settings feature.

Tests: used by [REDACTED_PERSON] settings (CL to come).

diff --git a/Engine/Source/Runtime/Slate/Public/Widgets/Input/SNumericEntryBox.h b/Engine/Source/Runtime/Slate/Public/Widgets/Input/SNumericEntryBox.h
--- a/Engine/Source/Runtime/Slate/Public/Widgets/Input/SNumericEntryBox.h
+++ b/Engine/Source/Runtime/Slate/Public/Widgets/Input/SNumericEntryBox.h
@@ -84,6 +84,8 @@ public:
 		, _SupportDynamicSliderMaxValue(false)
 		, _SupportDynamicSliderMinValue(false)
 		, _Delta(NumericType(0))
+		, _MinFractionalDigits(DefaultMinFractionalDigits)
+		, _MaxFractionalDigits(DefaultMaxFractionalDigits)
 		, _MinValue(TNumericLimits<NumericType>::Lowest())
 		, _MaxValue(TNumericLimits<NumericType>::Max())
 		, _MinSliderValue(NumericType(0))
@@ -138,6 +140,10 @@ public:
 		SLATE_EVENT(FOnDynamicSliderMinMaxValueChanged, OnDynamicSliderMinValueChanged)
 		/** Delta to increment the value as the slider moves.  If not specified will determine automatically */
 		SLATE_ATTRIBUTE( NumericType, Delta )
+		/** The minimum fractional digits the spin box displays, defaults to 1 */
+		SLATE_ATTRIBUTE(TOptional< int32 >, MinFractionalDigits)
+		/** The maximum fractional digits the spin box displays, defaults to 6 */
+		SLATE_ATTRIBUTE(TOptional< int32 >, MaxFractionalDigits)
 		/** The minimum value that can be entered into the text edit box */
 		SLATE_ATTRIBUTE( TOptional< NumericType >, MinValue )
 		/** The maximum value that can be entered into the text edit box */
@@ -204,11 +210,17 @@ public:
 		BorderImageFocused = &InArgs._EditableTextBoxStyle->BackgroundImageFocused;
 		Interface = InArgs._TypeInterface.IsValid() ? InArgs._TypeInterface : MakeShareable( new TDefaultNumericTypeInterface<NumericType> );
 
+		MinFractionalDigits = (InArgs._MinFractionalDigits.Get().IsSet()) ? InArgs._MinFractionalDigits : DefaultMinFractionalDigits;
+		MaxFractionalDigits = (InArgs._MaxFractionalDigits.Get().IsSet()) ? InArgs._MaxFractionalDigits : DefaultMaxFractionalDigits;
+		SetMinFractionalDigits(MinFractionalDigits);
+		SetMaxFractionalDigits(MaxFractionalDigits);
+
 		CachedExternalValue = ValueAttribute.Get();
 		if (CachedExternalValue.IsSet())
 		{
 			CachedValueString = Interface->ToString(CachedExternalValue.GetValue());
 		}
+		bCachedValueStringDirty = false;
 
 		const bool bDisplayToggle = InArgs._DisplayToggle;
 		if(bDisplayToggle)
@@ -237,6 +249,8 @@ public:
 				.OnDynamicSliderMinValueChanged(InArgs._OnDynamicSliderMinValueChanged)
 				.OnValueChanged(OnValueChanged)
 				.OnValueCommitted(OnValueCommitted)
+				.MinFractionalDigits(MinFractionalDigits)
+				.MaxFractionalDigits(MaxFractionalDigits)
 				.MinSliderValue(InArgs._MinSliderValue)
 				.MaxSliderValue(InArgs._MaxSliderValue)
 				.MaxValue(InArgs._MaxValue)
@@ -430,6 +444,22 @@ public:
 	/** Return the internally created SpinBox if bAllowSpin is true */
 	TSharedPtr<SWidget> GetSpinBox() const { return SpinBox; }
 
+	/** See the MinFractionalDigits attribute */
+	int32 GetMinFractionalDigits() const { return Interface->GetMinFractionalDigits(); }
+	void SetMinFractionalDigits(const TAttribute<TOptional<int32>>& InMinFractionalDigits)
+	{
+		Interface->SetMinFractionalDigits((InMinFractionalDigits.Get().IsSet()) ? InMinFractionalDigits.Get() : MinFractionalDigits);
+		bCachedValueStringDirty = true;
+	}
+
+	/** See the MaxFractionalDigits attribute */
+	int32 GetMaxFractionalDigits() const { return Interface->GetMaxFractionalDigits(); }
+	void SetMaxFractionalDigits(const TAttribute<TOptional<int32>>& InMaxFractionalDigits)
+	{
+		Interface->SetMaxFractionalDigits((InMaxFractionalDigits.Get().IsSet()) ? InMaxFractionalDigits.Get() : MaxFractionalDigits);
+		bCachedValueStringDirty = true;
+	}
+
 private:
 
 	//~ SWidget Interface
@@ -504,17 +534,23 @@ private:
 
 	void SetCachedString(const NumericType CurrentValue)
 	{
-		if (!CachedExternalValue.IsSet() || CachedExternalValue.GetValue() != CurrentValue)
+		if (!CachedExternalValue.IsSet() || CachedExternalValue.GetValue() != CurrentValue || bCachedValueStringDirty)
 		{
 			CachedExternalValue = CurrentValue;
 			CachedValueString = Interface->ToString(CurrentValue);
+			bCachedValueStringDirty = false;
 		}
 	}
 
 	FString GetCachedString(const NumericType CurrentValue) const
 	{
-		bool bUseCachedString = CachedExternalValue.IsSet() && CurrentValue == CachedExternalValue.GetValue();
-		return bUseCachedString ? CachedValueString : Interface->ToString(CurrentValue);
+		bool bUseCachedString = CachedExternalValue.IsSet() && CurrentValue == CachedExternalValue.GetValue() && !bCachedValueStringDirty;
+		if (!bUseCachedString)
+		{
+			CachedValueString = Interface->ToString(CurrentValue);
+			bCachedValueStringDirty = false;
+		}
+		return CachedValueString;
 	}
 
 	FText GetValueAsText() const
@@ -709,6 +745,12 @@ private:
 
 private:
 
+	/** The default minimum fractional digits */
+	static const int32 DefaultMinFractionalDigits;
+
+	/** The default maximum fractional digits */
+	static const int32 DefaultMaxFractionalDigits;
+
 	/** Attribute for getting the label */
 	TAttribute< TOptional<FString > > LabelAttribute;
 	/** Attribute for getting the value.  If the value is not set we display the undetermined string */
@@ -742,7 +784,11 @@ private:
 	/** Cached value of entry box, updated on set & per tick */
 	TOptional<NumericType> CachedExternalValue;
 	/** Used to prevent per-frame re-conversion of the cached numeric value to a string. */
-	FString CachedValueString;
+	mutable FString CachedValueString;
+	/** Whetever the interfaced setting changed and the CachedValueString needs to be recomputed. */
+	mutable bool bCachedValueStringDirty;
+	TAttribute< TOptional<int32> > MinFractionalDigits;
+	TAttribute< TOptional<int32> > MaxFractionalDigits;
 };
 
 
@@ -760,3 +806,10 @@ const FLinearColor SNumericEntryBox<NumericType>::LilacLabelBackgroundColor(0.8f
 
 template <typename NumericType>
 const FText SNumericEntryBox<NumericType>::DefaultUndeterminedString = FText::FromString(TEXT("---"));
+
+template<typename NumericType>
+const int32 SNumericEntryBox<NumericType>::DefaultMinFractionalDigits = 1;
+
+template<typename NumericType>
+const int32 SNumericEntryBox<NumericType>::DefaultMaxFractionalDigits = 6;
+
