[BugFix] Fixed excessive recursion in SLinkableScrollBar::SetState

diff --git a/Engine/Source/Runtime/Slate/Private/Widgets/Layout/LinkableScrollBar.cpp b/Engine/Source/Runtime/Slate/Private/Widgets/Layout/LinkableScrollBar.cpp
--- a/Engine/Source/Runtime/Slate/Private/Widgets/Layout/LinkableScrollBar.cpp
+++ b/Engine/Source/Runtime/Slate/Private/Widgets/Layout/LinkableScrollBar.cpp
@@ -21,7 +21,8 @@ void SLinkableScrollBar::SetState(float InOffsetFraction, float InThumbSizeFract
 	
 	const float NewDistanceFromTop = FMath::Clamp(Track->DistanceFromTop(), 0.f, 1.f);
 	const float NewDistanceFromBottom = FMath::Clamp(Track->DistanceFromBottom(), 0.f, 1.f);
-	const bool bDirty = NewDistanceFromTop != PrevDistanceFromTop && NewDistanceFromBottom != PrevDistanceFromBottom;
+	
+	const bool bDirty = !FMath::IsNearlyEqual(NewDistanceFromTop, PrevDistanceFromTop, 1E-04f) && !FMath::IsNearlyEqual(NewDistanceFromBottom, PrevDistanceFromBottom, 1E-04f);
 	
 	// if another scrollbar is linked to this one, Scroll it aswell
 	if (bDirty)
