Slate:

-Add option to hide header menu combo button on table view headers

diff --git a/Engine/Source/Runtime/Slate/Private/Widgets/Views/SHeaderRow.cpp b/Engine/Source/Runtime/Slate/Private/Widgets/Views/SHeaderRow.cpp
--- a/Engine/Source/Runtime/Slate/Private/Widgets/Views/SHeaderRow.cpp
+++ b/Engine/Source/Runtime/Slate/Private/Widgets/Views/SHeaderRow.cpp
@@ -145,8 +145,9 @@ public:
 			];
 		}
 
-		if( Column.HeaderMenuContent.Widget != SNullWidget::NullWidget ||
-			Column.OnGetMenuContent.IsBound())
+		if( ComboVisibility != EHeaderComboVisibility::Never &&
+			(Column.HeaderMenuContent.Widget != SNullWidget::NullWidget ||
+			Column.OnGetMenuContent.IsBound()))
 		{
 			// Add Drop down menu button (only if menu content has been specified)
 			Box->AddSlot()
@@ -281,6 +282,11 @@ public:
 			}
 		}
 
+		if (ComboVisibility == EHeaderComboVisibility::Never)
+		{
+			return EVisibility::Collapsed;
+		}
+
 		return EVisibility::Visible;
 	}
 
@@ -304,14 +310,17 @@ private:
 	
 	const FSlateBrush* GetComboButtonBorderBrush() const
 	{
-		if ( ComboButton.IsValid() && ( ComboButton->IsHovered() || ComboButton->IsOpen() ) )
+		if (ComboVisibility != EHeaderComboVisibility::Never)
 		{
-			return &Style->MenuDropdownHoveredBorderBrush;
-		}
+			if ( ComboButton.IsValid() && ( ComboButton->IsHovered() || ComboButton->IsOpen() ) )
+			{
+				return &Style->MenuDropdownHoveredBorderBrush;
+			}
 
-		if ( IsHovered() || ComboVisibility == EHeaderComboVisibility::Always )
-		{
-			return &Style->MenuDropdownNormalBorderBrush;
+			if ( IsHovered() || ComboVisibility == EHeaderComboVisibility::Always )
+			{
+				return &Style->MenuDropdownNormalBorderBrush;
+			}
 		}
 
 		return FStyleDefaults::GetNoBrush();
@@ -348,6 +357,11 @@ private:
 			}
 			break;
 
+		case EHeaderComboVisibility::Never:
+			{
+				return FLinearColor::White;
+			}
+
 		default:
 			break;
 		}
diff --git a/Engine/Source/Runtime/Slate/Public/Widgets/Views/SHeaderRow.h b/Engine/Source/Runtime/Slate/Public/Widgets/Views/SHeaderRow.h
--- a/Engine/Source/Runtime/Slate/Public/Widgets/Views/SHeaderRow.h
+++ b/Engine/Source/Runtime/Slate/Public/Widgets/Views/SHeaderRow.h
@@ -2,7 +2,6 @@
 
 #pragma once
 
-#include "CoreMinimal.h"
 #include "Containers/IndirectArray.h"
 #include "Misc/Attribute.h"
 #include "Layout/Visibility.h"
@@ -76,6 +75,9 @@ enum class EHeaderComboVisibility
 
 	/** Only show the drop down when hovered */
 	OnHover,
+
+	/** Never show the drop down. Context Menu can still be opened with a context click (such as right-click) */
+	Never
 };
 
 
