replaced 1E-04f with UE_KINDA_SMALL_NUMBER

diff --git a/Engine/Source/Runtime/Slate/Private/Widgets/Layout/LinkableScrollBar.cpp b/Engine/Source/Runtime/Slate/Private/Widgets/Layout/LinkableScrollBar.cpp
--- a/Engine/Source/Runtime/Slate/Private/Widgets/Layout/LinkableScrollBar.cpp
+++ b/Engine/Source/Runtime/Slate/Private/Widgets/Layout/LinkableScrollBar.cpp
@@ -22,7 +22,7 @@ void SLinkableScrollBar::SetState(float InOffsetFraction, float InThumbSizeFract
 	const float NewDistanceFromTop = FMath::Clamp(Track->DistanceFromTop(), 0.f, 1.f);
 	const float NewDistanceFromBottom = FMath::Clamp(Track->DistanceFromBottom(), 0.f, 1.f);
 	
-	const bool bDirty = !FMath::IsNearlyEqual(NewDistanceFromTop, PrevDistanceFromTop, 1E-04f) && !FMath::IsNearlyEqual(NewDistanceFromBottom, PrevDistanceFromBottom, 1E-04f);
+	const bool bDirty = !FMath::IsNearlyEqual(NewDistanceFromTop, PrevDistanceFromTop, UE_KINDA_SMALL_NUMBER) && !FMath::IsNearlyEqual(NewDistanceFromBottom, PrevDistanceFromBottom, UE_KINDA_SMALL_NUMBER);
 	
 	// if another scrollbar is linked to this one, Scroll it aswell
 	if (bDirty)
