Optimize UpdateTooltip:
- no longer call it from ProcessMouseMoveEvent, there is no point doing it here, in addition to Tick => Tick should be called at a sufficient rate to have a smooth display, and if the frame rate is slow, doing a costly process at each mouse move event won't help.
- no longer try to create a new tooltip each time the function is called, but only when the mouse is stopped for a short amount of time. However, an existing tool tip will still be moved or hidden if necessary.

Tests:
- compared before/after, did not perceived any behavior difference.
- tested tooltip in various situation in the editor, test the ctrl + alt behavior.
- frame rate no longer drop when moving frantically the mouse over list like content browser.
Some metrics, in ms:
                                Idle        MovingMouse(before)        MovingMouse(now)
Total Slate Tick          [REDACTED_PHONE]
Update Tooltip            [REDACTED_PHONE]
Slate App Input           [REDACTED_PHONE]
ProcessMouseMove      [REDACTED_PHONE]

diff --git a/Engine/Source/Runtime/Slate/Private/Framework/Application/SlateApplication.cpp b/Engine/Source/Runtime/Slate/Private/Framework/Application/SlateApplication.cpp
--- a/Engine/Source/Runtime/Slate/Private/Framework/Application/SlateApplication.cpp
+++ b/Engine/Source/Runtime/Slate/Private/Framework/Application/SlateApplication.cpp
@@ -1553,9 +1553,10 @@ void FSlateApplication::TickPlatform(float DeltaTime)
 	{
 		SCOPE_CYCLE_COUNTER(STAT_SlateApplicationInput);
 
+		const bool bCanSpawnNewTooltip = true;
 		ForEachUser([this](FSlateUser& User) {
 			User.UpdateCursor();
-			User.UpdateTooltip(MenuStack, false);
+			User.UpdateTooltip(MenuStack, bCanSpawnNewTooltip);
 		});
 
 		bool bSynthesizedCursorMoveThisFrame = false;
@@ -5898,10 +5899,6 @@ bool FSlateApplication::ProcessMouseMoveEvent( const FPointerEvent& MouseEvent,
 			return true;
 		}
 
-		QUICK_SCOPE_CYCLE_COUNTER(STAT_ProcessMouseMove_Tooltip);
-		
-		GetOrCreateUser(MouseEvent)->UpdateTooltip(MenuStack, /*bCanSpawnNewTooltip =*/true);
-		
 		// Guard against synthesized mouse moves and only track user interaction if the cursor pos changed
 		SetLastUserInteractionTime(this->GetCurrentTime());
 	}
diff --git a/Engine/Source/Runtime/Slate/Private/Framework/Application/SlateUser.cpp b/Engine/Source/Runtime/Slate/Private/Framework/Application/SlateUser.cpp
--- a/Engine/Source/Runtime/Slate/Private/Framework/Application/SlateUser.cpp
+++ b/Engine/Source/Runtime/Slate/Private/Framework/Application/SlateUser.cpp
@@ -56,6 +56,12 @@ FAutoConsoleVariableRef CVarTooltipIntroDuration(
 	TooltipIntroDuration,
 	TEXT("How long it takes for a tooltip to animate into view, in seconds."));
 
+static float CursorSignificantMoveDetectionThreshold = 0.0;
+FAutoConsoleVariableRef CVarCursorSignificantMoveDetectionThreshold(
+	TEXT("Slate.CursorSignificantMoveDetectionThreshold"),
+	CursorSignificantMoveDetectionThreshold,
+	TEXT("The distance from previous cursor position above which the move will be considered significant (used to trigger the display of the tooltips)."));
+
 //////////////////////////////////////////////////////////////////////////
 // FSlateVirtualUserHandle
 //////////////////////////////////////////////////////////////////////////
@@ -888,6 +894,12 @@ void FSlateUser::UpdateCursor()
 	{
 		QueryCursor();
 	}
+
+	const double MoveEpsilonSquared = CursorSignificantMoveDetectionThreshold * CursorSignificantMoveDetectionThreshold;
+	if (FVector2D::DistSquared(GetPreviousCursorPosition(), GetCursorPosition()) > MoveEpsilonSquared)
+	{
+		LastCursorSignificantMoveTime = FPlatformTime::Seconds();
+	}
 }
 
 void FSlateUser::ProcessCursorReply(const FCursorReply& CursorReply)
@@ -1178,7 +1190,10 @@ void FSlateUser::UpdateTooltip(const FMenuStack& MenuStack, bool bCanSpawnNewToo
 
 	SCOPE_CYCLE_COUNTER(STAT_SlateUpdateTooltip);
 
-	float DPIScaleFactor = 1.0f;
+	const double MotionLessDurationBeforeAllowingNewToolTip = 0.05;
+	bCanSpawnNewTooltip = bCanSpawnNewTooltip && (FPlatformTime::Seconds() - LastCursorSignificantMoveTime > MotionLessDurationBeforeAllowingNewToolTip);
+
+	float DPIScaleFactor = 1.0f; //todo: this value is never changed, we should investigate if it is necessary or not to handle it for the force field.
 	FWidgetPath WidgetsToQueryForTooltip;
 
 	const bool bCheckForTooltipChanges =
diff --git a/Engine/Source/Runtime/Slate/Public/Framework/Application/SlateUser.h b/Engine/Source/Runtime/Slate/Public/Framework/Application/SlateUser.h
--- a/Engine/Source/Runtime/Slate/Public/Framework/Application/SlateUser.h
+++ b/Engine/Source/Runtime/Slate/Public/Framework/Application/SlateUser.h
@@ -208,6 +208,9 @@ private:
 	/** The cursor this user is in control of. Guaranteed to be valid for all real users, absence implies this is a virtual user. */
 	TSharedPtr<ICursor> Cursor;
 
+	/** Store the last time the cursor position was changed */
+	double LastCursorSignificantMoveTime = 0.0;
+
 	/** SlateUsers can optionally be individually assigned a navigation config to use. This overrides the global nav config that lives on FSlateApplication when valid. */
 	TSharedPtr<FNavigationConfig> UserNavigationConfig;
 
