Build fix for not capturing a local variable (introduced in CL 23790605). Unclear why this isn't failing on the build server.

diff --git a/Engine/Source/Runtime/Slate/Private/Framework/Application/SlateApplication.cpp b/Engine/Source/Runtime/Slate/Private/Framework/Application/SlateApplication.cpp
--- a/Engine/Source/Runtime/Slate/Private/Framework/Application/SlateApplication.cpp
+++ b/Engine/Source/Runtime/Slate/Private/Framework/Application/SlateApplication.cpp
@@ -1554,7 +1554,7 @@ void FSlateApplication::TickPlatform(float DeltaTime)
 		SCOPE_CYCLE_COUNTER(STAT_SlateApplicationInput);
 
 		const bool bCanSpawnNewTooltip = true;
-		ForEachUser([this](FSlateUser& User) {
+		ForEachUser([this, bCanSpawnNewTooltip](FSlateUser& User) {
 			User.UpdateCursor();
 			User.UpdateTooltip(MenuStack, bCanSpawnNewTooltip);
 		});
