Added the ability to override the PlatformTextField in the SlateApplication.  This is especially useful when we need to support virtual keyboard on a platform that doesn't need to support it by default.
#rnx

diff --git a/Engine/Source/Runtime/Slate/Private/Framework/Application/SlateApplication.cpp b/Engine/Source/Runtime/Slate/Private/Framework/Application/SlateApplication.cpp
--- a/Engine/Source/Runtime/Slate/Private/Framework/Application/SlateApplication.cpp
+++ b/Engine/Source/Runtime/Slate/Private/Framework/Application/SlateApplication.cpp
@@ -773,7 +773,6 @@ FSlateApplication::FSlateApplication()
 	, bRequestLeaveDebugMode( false )
 	, bLeaveDebugForSingleStep( false )
 	, bIsExternalUIOpened( false )
-	, SlateTextField( nullptr )
 	, bIsFakingTouch(FParse::Param(FCommandLine::Get(), TEXT("simmobile")) || FParse::Param(FCommandLine::Get(), TEXT("faketouches")))
 	, bIsGameFakingTouch( false )
 	, bIsFakingTouched( false )
@@ -847,12 +846,6 @@ FSlateApplication::~FSlateApplication()
 {
 	FTabCommands::Unregister();
 	FGenericCommands::Unregister();
-	
-	if (SlateTextField != nullptr)
-	{
-		delete SlateTextField;
-		SlateTextField = nullptr;
-	}
 
 #if WITH_EDITOR
 	OnDebugSafeZoneChanged.RemoveAll(this);
@@ -977,6 +970,11 @@ void FSlateApplication::SetCursorPos(const FVector2D& MouseCoordinate)
 	GetCursorUser()->SetCursorPosition(UE::Slate::CastToVector2f(MouseCoordinate));
 }
 
+void FSlateApplication::OverridePlatformTextField(TUniquePtr<IPlatformTextField> PlatformTextField)
+{
+	SlateTextField = MoveTemp(PlatformTextField);
+}
+
 void FSlateApplication::UsePlatformCursorForCursorUser(bool bUsePlatformCursor)
 {
 	if (TSharedPtr<FSlateUser> SlateUser = GetUser(CursorUserIndex))
@@ -3873,9 +3871,9 @@ void FSlateApplication::ShowVirtualKeyboard( bool bShow, int32 UserIndex, TShare
 {
 	SCOPE_CYCLE_COUNTER(STAT_ShowVirtualKeyboard);
 
-	if(SlateTextField == nullptr)
+	if (!SlateTextField.IsValid())
 	{
-		SlateTextField = new FPlatformTextField();
+		SlateTextField = MakeUnique<FPlatformTextField>();
 	}
 
 	SlateTextField->ShowVirtualKeyboard(bShow, UserIndex, TextEntryWidget);
@@ -3883,9 +3881,9 @@ void FSlateApplication::ShowVirtualKeyboard( bool bShow, int32 UserIndex, TShare
 
 bool FSlateApplication::AllowMoveCursor()
 {
-	if (SlateTextField == nullptr)
+	if (!SlateTextField.IsValid())
 	{
-		SlateTextField = new FPlatformTextField();
+		SlateTextField = MakeUnique<FPlatformTextField>();
 	}
 
 	return SlateTextField->AllowMoveCursor();
diff --git a/Engine/Source/Runtime/Slate/Public/Framework/Application/SlateApplication.h b/Engine/Source/Runtime/Slate/Public/Framework/Application/SlateApplication.h
--- a/Engine/Source/Runtime/Slate/Public/Framework/Application/SlateApplication.h
+++ b/Engine/Source/Runtime/Slate/Public/Framework/Application/SlateApplication.h
@@ -301,6 +301,14 @@ public:
 	 */
 	void SetCursorPos( const FVector2D& MouseCoordinate ) override;
 
+	/** 
+	 * Replace the IPlatformTextField implementation with a custom one.  The implementation used by default is platform specific.
+	 * If you need to change the default behavior for virtual keyboard, this method is for you.
+	 *
+	 * @param PlatformTextField		The platform specific implementation of the IPlatformTextField.
+	 */
+	void OverridePlatformTextField(TUniquePtr<IPlatformTextField> PlatformTextField);
+	
 	/** 
 	 *	Updates the cursor user's cursor to either the platform cursor or fake cursor
 	 */
@@ -1884,7 +1892,7 @@ private:
 	/**
 	 * Virtual keyboard text field
 	 */
-	IPlatformTextField* SlateTextField;
+	TUniquePtr<IPlatformTextField> SlateTextField;
 
 	/** For desktop platforms that want to test touch style input, pass -faketouches or -simmobile on the commandline to set this */
 	bool bIsFakingTouch;
