Move the call to SetKeyboardFocus in SPinTypeSelector to a callback which happens as a result of the spawning the new dropdown window to fix a race condition when for async display servers (i.e. X11 on Linux).

diff --git a/Engine/Source/Editor/KismetWidgets/Private/SPinTypeSelector.cpp b/Engine/Source/Editor/KismetWidgets/Private/SPinTypeSelector.cpp
--- a/Engine/Source/Editor/KismetWidgets/Private/SPinTypeSelector.cpp
+++ b/Engine/Source/Editor/KismetWidgets/Private/SPinTypeSelector.cpp
@@ -1028,7 +1028,6 @@ void SPinTypeSelector::OnTypeSelectionChanged(FPinTypeTreeItem Selection, ESelec
 			if (SelectInfo != ESelectInfo::OnMouseClick && PinRow.IsValid() && PinRow->HasSubMenu() && !PinRow->IsSubMenuOpen())
 			{
 				PinRow->RequestSubMenuToggle(true);
-				FSlateApplication::Get().SetKeyboardFocus(WeakListView.Pin(), EFocusCause::SetDirectly);
 			}
 			else
 			{
diff --git a/Engine/Source/Runtime/Slate/Private/Framework/Application/MenuStack.cpp b/Engine/Source/Runtime/Slate/Private/Framework/Application/MenuStack.cpp
--- a/Engine/Source/Runtime/Slate/Private/Framework/Application/MenuStack.cpp
+++ b/Engine/Source/Runtime/Slate/Private/Framework/Application/MenuStack.cpp
@@ -783,6 +783,55 @@ void FMenuStack::OnWindowDestroyed(TSharedRef<SWindow> InWindow)
 	}
 }
 
+namespace
+{
+	void AddChildrenToStack(TSharedPtr<SWidget> Root, TArray<TSharedPtr<SWidget>>& Stack)
+	{
+		if (!Root.IsValid())
+		{
+			return;
+		}
+
+		FChildren* Children = Root->GetChildren();
+		if (!Children)
+		{
+			return;
+		}
+
+		Stack.Reserve(Stack.Num() + Children->Num());
+
+		for (int ChildIndex = 0; ChildIndex < Children->Num(); ChildIndex++)
+		{
+			Stack.Push(Children->GetChildAt(ChildIndex));
+		}
+	}
+
+	// Iterative DFS for the first instance of a keyboard-focusable widget in the descendants of Root
+	void FindFirstKeyboardFocusableChildIteratively(TSharedPtr<SWidget> Root, TSharedPtr<SWidget>& OutWidget)
+	{
+		if (!Root.IsValid())
+		{
+			return;
+		}
+
+		TArray<TSharedPtr<SWidget>> Stack;
+		AddChildrenToStack(Root, Stack);
+
+		while (!Stack.IsEmpty())
+		{
+			TSharedPtr CurrentWidget = Stack.Pop();
+
+			if (CurrentWidget->SupportsKeyboardFocus())
+			{
+				OutWidget = CurrentWidget;
+				return;
+			}
+
+			AddChildrenToStack(CurrentWidget, Stack);
+		}
+	}
+}
+
 void FMenuStack::OnWindowActivated( TSharedRef<SWindow> ActivatedWindow )
 {
 	if (ActivatedWindow != PendingNewWindow && HasMenus() && !FSlateApplication::Get().IsWindowHousingInteractiveTooltip(ActivatedWindow))
@@ -803,6 +852,13 @@ void FMenuStack::OnWindowActivated( TSharedRef<SWindow> ActivatedWindow )
 					break;
 				}
 			}
+
+			TSharedPtr<SWidget> KeyboardFocusableChild = nullptr;
+			FindFirstKeyboardFocusableChildIteratively(ActivatedMenu.Pin()->GetContent(), KeyboardFocusableChild);
+			if (KeyboardFocusableChild.IsValid())
+			{
+				FSlateApplication::Get().SetKeyboardFocus(KeyboardFocusableChild, EFocusCause::SetDirectly);
+			}
 		}
 		else
 		{
