Fix crash in FTabManager::OpenUnmanagedTab when the requested tab is not allowed

diff --git a/Engine/Source/Runtime/Slate/Private/Framework/Docking/TabManager.cpp b/Engine/Source/Runtime/Slate/Private/Framework/Docking/TabManager.cpp
--- a/Engine/Source/Runtime/Slate/Private/Framework/Docking/TabManager.cpp
+++ b/Engine/Source/Runtime/Slate/Private/Framework/Docking/TabManager.cpp
@@ -1552,8 +1552,7 @@ void FTabManager::OpenUnmanagedTab(FName PlaceholderId, const FSearchPreference&
 		{
 			UE_LOG(LogTabManager, Warning, TEXT("Unable to insert tab '%s'."), *(PlaceholderId.ToString()));
 			LiveTab = InvokeTab_Internal( FTabId( PlaceholderId ) );
-			check(LiveTab.IsValid());
-			if (LiveTab)
+			if (LiveTab.IsValid())
 			{
 				LiveTab->GetParent()->GetParentDockTabStack()->OpenTab( UnmanagedTab );
 			}
