Fix for copyright notices

diff --git a/Engine/Source/Developer/WidgetRegistration/Private/ToolkitBuilder.cpp b/Engine/Source/Developer/WidgetRegistration/Private/ToolkitBuilder.cpp
--- a/Engine/Source/Developer/WidgetRegistration/Private/ToolkitBuilder.cpp
+++ b/Engine/Source/Developer/WidgetRegistration/Private/ToolkitBuilder.cpp
@@ -1,4 +1,4 @@
-
+// Copyright Epic Games, Inc. All Rights Reserved.
 
 #include "ToolkitBuilder.h"
 
diff --git a/Engine/Source/Runtime/Slate/Private/Framework/MultiBox/SUniformToolbarButtonBlock.cpp b/Engine/Source/Runtime/Slate/Private/Framework/MultiBox/SUniformToolbarButtonBlock.cpp
--- a/Engine/Source/Runtime/Slate/Private/Framework/MultiBox/SUniformToolbarButtonBlock.cpp
+++ b/Engine/Source/Runtime/Slate/Private/Framework/MultiBox/SUniformToolbarButtonBlock.cpp
@@ -1,3 +1,5 @@
+// Copyright Epic Games, Inc. All Rights Reserved.
+
 #include "Framework/MultiBox/SUniformToolbarButtonBlock.h"
 
 FUniformToolbarButtonBlock::FUniformToolbarButtonBlock(FButtonArgs ButtonArgs) :
