Add cvar Slate.Input.MotionFiresUserInteractionEvents:
If this is false, LastUserInteractionTimeUpdateEvent events won't be fired based on motion input, and LastInteractionTime won't be updated
Some motion devices report small tiny changes constantly without filtering, so motion input is unhelpful for determining user activity

diff --git a/Engine/Source/Runtime/Slate/Private/Framework/Application/SlateApplication.cpp b/Engine/Source/Runtime/Slate/Private/Framework/Application/SlateApplication.cpp
--- a/Engine/Source/Runtime/Slate/Private/Framework/Application/SlateApplication.cpp
+++ b/Engine/Source/Runtime/Slate/Private/Framework/Application/SlateApplication.cpp
@@ -87,6 +87,13 @@ namespace UE::Slate::Private
 }
 #endif //WITH_SLATE_DEBUGGING
 
+bool GSlateInputMotionFiresUserInteractionEvents = true;
+static FAutoConsoleVariableRef CVarSlateInputMotionFiresUserInteractionEvents(
+	TEXT("Slate.Input.MotionFiresUserInteractionEvents"),
+	GSlateInputMotionFiresUserInteractionEvents,
+	TEXT("If this is false, LastUserInteractionTimeUpdateEvent events won't be fired based on motion input, and LastInteractionTime won't be updated\n")
+	TEXT("Some motion devices report small tiny changes constantly without filtering, so motion input is unhelpful for determining user activity"));
+
 //////////////////////////////////////////////////////////////////////////
 
 /** 
@@ -6312,7 +6319,10 @@ void FSlateApplication::ProcessMotionDetectedEvent( const FMotionEvent& MotionEv
 	FSlateDebugging::FScopeProcessInputEvent Scope(ESlateDebuggingInputEvent::MotionDetected, MotionEvent);
 #endif
 
-	SetLastUserInteractionTime(this->GetCurrentTime());
+	if (GSlateInputMotionFiresUserInteractionEvents)
+	{
+		SetLastUserInteractionTime(this->GetCurrentTime());
+	}
 	
 	if (!InputPreProcessors.HandleMotionDetectedEvent(*this, MotionEvent))
 	{
