Remove drifting in SHeaderRow for manually sized columns.

diff --git a/Engine/Source/Runtime/Slate/Private/Widgets/Views/SHeaderRow.cpp b/Engine/Source/Runtime/Slate/Private/Widgets/Views/SHeaderRow.cpp
--- a/Engine/Source/Runtime/Slate/Private/Widgets/Views/SHeaderRow.cpp
+++ b/Engine/Source/Runtime/Slate/Private/Widgets/Views/SHeaderRow.cpp
@@ -875,9 +875,10 @@ void SHeaderRow::RegenerateWidgets()
 					SizingGrip->SetOnMouseButtonUp(FPointerEventHandler::CreateLambda(SizingGrip_OnMouseButtonUp));
 					SizingGrip->SetOnMouseMove(FPointerEventHandler::CreateLambda(SizingGrip_OnMouseMove));
 
-					auto GetColumnWidthAsOptionalSize = [&SomeColumn]() -> FOptionalSize
+					auto GetColumnWidthAsOptionalSize = [&SomeColumn, SplitterHandleSizeCopy = SplitterHandleSize]() -> FOptionalSize
 					{
-						const float DesiredWidth = SomeColumn.GetWidth();
+						// Subtract SplitterHandleSize to compensate for SSplitter adding a handle between items.
+						const float DesiredWidth = SomeColumn.GetWidth() - SplitterHandleSizeCopy;
 						return FOptionalSize(DesiredWidth);
 					};
 
