Add option to ignore editor gamepad navigation actions in slate.

diff --git a/Engine/Source/Developer/SlateReflector/Private/Widgets/SSlateOptions.cpp b/Engine/Source/Developer/SlateReflector/Private/Widgets/SSlateOptions.cpp
--- a/Engine/Source/Developer/SlateReflector/Private/Widgets/SSlateOptions.cpp
+++ b/Engine/Source/Developer/SlateReflector/Private/Widgets/SSlateOptions.cpp
@@ -61,6 +61,7 @@ void SSlateOptions::Construct( const FArguments& InArgs )
 			AddMenuEntry(MenuBuilder, Icon, LOCTEXT("EnableFastWidgetPath", "Fast Widget Path"), TEXT("Slate.EnableFastWidgetPath"), false);
 			AddMenuEntry(MenuBuilder, Icon, LOCTEXT("EnableToolTips", "Enable Tooltips"), TEXT("Slate.EnableTooltips"));
 			AddMenuEntry(MenuBuilder, Icon, LOCTEXT("GlobalInvalidation", "Global Invalidation"), TEXT("Slate.EnableGlobalInvalidation"));
+			AddMenuEntry(MenuBuilder, Icon, LOCTEXT("DisabledEffect", "Transparent Disabled Effect"), TEXT("Slate.ApplyDisabledEffectOnWidgets"));
 
 			return MenuBuilder.MakeWidget();
 		}
diff --git a/Engine/Source/Runtime/Engine/Private/Slate/SGameLayerManager.cpp b/Engine/Source/Runtime/Engine/Private/Slate/SGameLayerManager.cpp
--- a/Engine/Source/Runtime/Engine/Private/Slate/SGameLayerManager.cpp
+++ b/Engine/Source/Runtime/Engine/Private/Slate/SGameLayerManager.cpp
@@ -31,14 +31,6 @@ static void HandlePerUserHitTestingToggled(IConsoleVariable* CVar)
 	FSlateApplication::Get().InvalidateAllWidgets(false);
 }
 
-static void HandleApplyDisabledEffectOnWidgetsToggled(IConsoleVariable* CVar)
-{
-	if (FSlateApplication::IsInitialized())
-	{
-		FSlateApplication::Get().InvalidateAllWidgets(false);
-	}
-}
-
 static int32 bEnablePerUserHitTesting = true;
 
 static FAutoConsoleVariableRef CVarEnablePerUserHitTesting(
@@ -48,15 +40,6 @@ static FAutoConsoleVariableRef CVarEnablePerUserHitTesting(
 	FConsoleVariableDelegate::CreateStatic(&HandlePerUserHitTestingToggled)
 );
 
-static int32 bApplyDisabledEffectOnWidgets = 1;
-
-static FAutoConsoleVariableRef CVarApplyDisabledEffectOnWidgets(
-	TEXT("ApplyDisabledEffectOnWidgets"),
-	bApplyDisabledEffectOnWidgets,
-	TEXT("If true, disabled game-layer widgets will have alpha multiplied by 0.45."),
-	FConsoleVariableDelegate::CreateStatic(&HandleApplyDisabledEffectOnWidgetsToggled)
-);
-
 #if UE_SLATE_WITH_GAMELAYER_CANVAS_VISIBILITY_COMMANDS
 namespace CanvasVisibility
 {
@@ -412,8 +395,7 @@ int32 SGameLayerManager::OnPaint(const FPaintArgs& Args, const FGeometry& Allott
 	}
 #endif
 
-	// By default game layer state is used to ignore disabled effects, don't set this flag if user wants disabled draw effects in-game
-	OutDrawElements.SetIsInGameLayer(!bApplyDisabledEffectOnWidgets);
+	OutDrawElements.SetIsInGameLayer(true);
 	const int32 ResultLayer = SCompoundWidget::OnPaint(Args, AllottedGeometry, MyCullingRect, OutDrawElements, LayerId, InWidgetStyle, bParentEnabled);
 	OutDrawElements.SetIsInGameLayer(false);
 	return ResultLayer;
diff --git a/Engine/Source/Runtime/Slate/Private/Framework/Application/SlateApplication.cpp b/Engine/Source/Runtime/Slate/Private/Framework/Application/SlateApplication.cpp
--- a/Engine/Source/Runtime/Slate/Private/Framework/Application/SlateApplication.cpp
+++ b/Engine/Source/Runtime/Slate/Private/Framework/Application/SlateApplication.cpp
@@ -79,7 +79,6 @@ static FAutoConsoleVariableRef CVarSlateVerifyWidgetLayerId(
 	GSlateVerifyWidgetLayerId,
 	TEXT("Every tick, verify that widgets have a LayerId range that fits with their siblings and their parent.")
 );
-
 namespace UE::Slate::Private
 {
 	void VerifyParentChildrenRelationship(const TSharedRef<SWindow>& WindowToDraw);
@@ -96,6 +95,15 @@ static FAutoConsoleVariableRef CVarSlateInputMotionFiresUserInteractionEvents(
 
 //////////////////////////////////////////////////////////////////////////
 
+bool GSlateEnableGamepadEditorNavigation = true;
+static FAutoConsoleVariableRef CVarSlateEnableGamepadEditorNavigation(
+	TEXT("Slate.EnableGamepadEditorNavigation"),
+	GSlateEnableGamepadEditorNavigation,
+	TEXT("True implies we allow gamepad navigation outside of the game viewport.")
+);
+
+//////////////////////////////////////////////////////////////////////////
+
 /** 
  * The cursor given to any additional input-providing SlateUsers beyond the first one (since the first one owns the platform cursor)
  * Or, if we're on a platform without a cursor, User 0 gets one of these too
@@ -3386,7 +3394,11 @@ void FSlateApplication::ProcessReply( const FWidgetPath& CurrentEventPath, const
 
 		if (NavigationSource.IsValid())
 		{
-			if (TheReply.GetNavigationDestination().IsValid())
+			if (!GSlateEnableGamepadEditorNavigation && TheReply.GetNavigationGenesis() == ENavigationGenesis::Controller && !NavigationSource.GetLastWidget()->GetPersistentState().bIsInGameLayer)
+			{
+				// Gamepad navigation while not in a game layer, do nothing as specified by GSlateEnableGamepadEditorNavigation
+			}
+			else if (TheReply.GetNavigationDestination().IsValid())
 			{
 				const bool bAlwaysHandleNavigationAttempt = false;
 				ExecuteNavigation(NavigationSource, TheReply.GetNavigationDestination(), UserIndex, bAlwaysHandleNavigationAttempt);
diff --git a/Engine/Source/Runtime/SlateCore/Private/Rendering/DrawElements.cpp b/Engine/Source/Runtime/SlateCore/Private/Rendering/DrawElements.cpp
--- a/Engine/Source/Runtime/SlateCore/Private/Rendering/DrawElements.cpp
+++ b/Engine/Source/Runtime/SlateCore/Private/Rendering/DrawElements.cpp
@@ -21,6 +21,23 @@ DECLARE_CYCLE_STAT(TEXT("FSlateDrawElement::Prebatch Time"), STAT_SlateDrawEleme
 DEFINE_STAT(STAT_SlateBufferPoolMemory);
 DEFINE_STAT(STAT_SlateCachedDrawElementMemory);
 
+static bool bApplyDisabledEffectOnWidgets = true;
+
+static void HandleApplyDisabledEffectOnWidgetsToggled(IConsoleVariable* CVar)
+{
+	if (FSlateApplicationBase::IsInitialized())
+	{
+		FSlateApplicationBase::Get().InvalidateAllWidgets(false);
+	}
+}
+
+static FAutoConsoleVariableRef CVarApplyDisabledEffectOnWidgets(
+	TEXT("Slate.ApplyDisabledEffectOnWidgets"),
+	bApplyDisabledEffectOnWidgets,
+	TEXT("If true, disabled game-layer widgets will have alpha multiplied by 0.45."),
+	FConsoleVariableDelegate::CreateStatic(&HandleApplyDisabledEffectOnWidgetsToggled)
+);
+
 static bool IsResourceObjectValid(UObject*& InObject)
 {
 	if (InObject != nullptr && (!IsValidChecked(InObject) || InObject->IsUnreachable() || InObject->HasAnyFlags(RF_BeginDestroyed)))
@@ -172,7 +189,7 @@ FSlateWindowElementList::~FSlateWindowElementList()
 
 void FSlateDrawElement::Init(FSlateWindowElementList& ElementList, EElementType InElementType, uint32 InLayer, const FPaintGeometry& PaintGeometry, ESlateDrawEffect InDrawEffects)
 {
-	if (ElementList.GetIsInGameLayer())
+	if (ElementList.GetIsInGameLayer() && !bApplyDisabledEffectOnWidgets)
 	{
 		InDrawEffects &= ~ESlateDrawEffect::DisabledEffect;
 	}
