Slate: Fix to prevent selecting an item in the list view when the list view is empty.

diff --git a/Engine/Source/Runtime/Slate/Public/Widgets/Views/SListView.h b/Engine/Source/Runtime/Slate/Public/Widgets/Views/SListView.h
--- a/Engine/Source/Runtime/Slate/Public/Widgets/Views/SListView.h
+++ b/Engine/Source/Runtime/Slate/Public/Widgets/Views/SListView.h
@@ -951,26 +951,30 @@ public:
 		}
 
 		const TArrayView<const ItemType> ItemsSourceRef = GetItems();
-
-		int32 RangeStartIndex = 0;
-		if( TListTypeTraits<ItemType>::IsPtrValid(RangeSelectionStart) )
+		// The InRangeSelectionEnd come from the WidgetGenerator (previous tick). Maybe it is not in the current ItemsSource list.
+		//RangeSelectionStart, maybe it is not in the current ItemsSource list.
+		if (ItemsSourceRef.Num() != 0)
 		{
-			RangeStartIndex = ItemsSourceRef.Find( TListTypeTraits<ItemType>::NullableItemTypeConvertToItemType( RangeSelectionStart ) );
-		}
+			int32 RangeStartIndex = 0;
+			if( TListTypeTraits<ItemType>::IsPtrValid(RangeSelectionStart) )
+			{
+				RangeStartIndex = ItemsSourceRef.Find( TListTypeTraits<ItemType>::NullableItemTypeConvertToItemType( RangeSelectionStart ) );
+			}
 
-		int32 RangeEndIndex = ItemsSourceRef.Find( InRangeSelectionEnd );
+			int32 RangeEndIndex = ItemsSourceRef.Find( InRangeSelectionEnd );
 
-		RangeStartIndex = FMath::Clamp(RangeStartIndex, 0, ItemsSourceRef.Num());
-		RangeEndIndex = FMath::Clamp(RangeEndIndex, 0, ItemsSourceRef.Num());
+			RangeStartIndex = FMath::Clamp(RangeStartIndex, 0, ItemsSourceRef.Num());
+			RangeEndIndex = FMath::Clamp(RangeEndIndex, 0, ItemsSourceRef.Num());
 
-		if (RangeEndIndex < RangeStartIndex)
-		{
-			Swap( RangeStartIndex, RangeEndIndex );
-		}
+			if (RangeEndIndex < RangeStartIndex)
+			{
+				Swap( RangeStartIndex, RangeEndIndex );
+			}
 
-		for( int32 ItemIndex = RangeStartIndex; ItemIndex <= RangeEndIndex; ++ItemIndex )
-		{
-			SelectedItems.Add( ItemsSourceRef[ItemIndex] );
+			for (int32 ItemIndex = RangeStartIndex; ItemIndex <= RangeEndIndex; ++ItemIndex)
+			{
+				SelectedItems.Add(ItemsSourceRef[ItemIndex]);
+			}
 		}
 
 		this->InertialScrollManager.ClearScrollVelocity();
