SColorBlock now honors parent enabled states when painting

diff --git a/Engine/Source/Runtime/Slate/Private/Widgets/Colors/SColorBlock.cpp b/Engine/Source/Runtime/Slate/Private/Widgets/Colors/SColorBlock.cpp
--- a/Engine/Source/Runtime/Slate/Private/Widgets/Colors/SColorBlock.cpp
+++ b/Engine/Source/Runtime/Slate/Private/Widgets/Colors/SColorBlock.cpp
@@ -51,7 +51,7 @@ void SColorBlock::Construct(const FArguments& InArgs)
 
 int32 SColorBlock::OnPaint(const FPaintArgs& Args, const FGeometry& AllottedGeometry, const FSlateRect& MyCullingRect, FSlateWindowElementList& OutDrawElements, int32 LayerId, const FWidgetStyle& InWidgetStyle, bool bParentEnabled) const
 {
-	const ESlateDrawEffect DrawEffects = ESlateDrawEffect::None;
+	const ESlateDrawEffect DrawEffects = bParentEnabled ? ESlateDrawEffect::None : ESlateDrawEffect::DisabledEffect;
 
 	EColorBlockAlphaDisplayMode DisplayMode = AlphaDisplayMode.Get();
 	bool bIgnoreAlpha = DisplayMode == EColorBlockAlphaDisplayMode::Ignore;
