Slate: Fix to prevent selecting an item in the list view when the list view is empty.

diff --git a/Engine/Source/Runtime/Slate/Public/Widgets/Views/SListView.h b/Engine/Source/Runtime/Slate/Public/Widgets/Views/SListView.h
--- a/Engine/Source/Runtime/Slate/Public/Widgets/Views/SListView.h
+++ b/Engine/Source/Runtime/Slate/Public/Widgets/Views/SListView.h
@@ -963,8 +963,8 @@ public:
 
 			int32 RangeEndIndex = ItemsSourceRef.Find( InRangeSelectionEnd );
 
-			RangeStartIndex = FMath::Clamp(RangeStartIndex, 0, ItemsSourceRef.Num());
-			RangeEndIndex = FMath::Clamp(RangeEndIndex, 0, ItemsSourceRef.Num());
+			RangeStartIndex = FMath::Clamp(RangeStartIndex, 0, ItemsSourceRef.Num()-1);
+			RangeEndIndex = FMath::Clamp(RangeEndIndex, 0, ItemsSourceRef.Num()-1);
 
 			if (RangeEndIndex < RangeStartIndex)
 			{
