Only allow explicit navigation if target widget is enabled
* Fixes selection getting lost and user not being able to get it back when trying to navigate to disabled widget

#test Local setup that used explicit navigation against a disabled widget

diff --git a/Engine/Source/Runtime/Slate/Private/Framework/Application/SlateApplication.cpp b/Engine/Source/Runtime/Slate/Private/Framework/Application/SlateApplication.cpp
--- a/Engine/Source/Runtime/Slate/Private/Framework/Application/SlateApplication.cpp
+++ b/Engine/Source/Runtime/Slate/Private/Framework/Application/SlateApplication.cpp
@@ -5985,11 +5985,36 @@ bool FSlateApplication::AttemptNavigation(const FWidgetPath& NavigationSource, c
 	EUINavigation NavigationType = NavigationEvent.GetNavigationType();
 	if ( NavigationReply.GetBoundaryRule() == EUINavigationRule::Explicit )
 	{
-		DestinationWidget = NavigationReply.GetFocusRecipient();
-		bAlwaysHandleNavigationAttempt = true;
+		const SWidget* FocusRecipient = NavigationReply.GetFocusRecipient().Get();
+		if ( FocusRecipient && FocusRecipient->IsEnabled() && FocusRecipient->SupportsKeyboardFocus() )
+		{
+			DestinationWidget = NavigationReply.GetFocusRecipient();
+			bAlwaysHandleNavigationAttempt = true;
 
 #if WITH_SLATE_DEBUGGING
-		NavigationMethod = ESlateDebuggingNavigationMethod::Explicit;
+			NavigationMethod = ESlateDebuggingNavigationMethod::Explicit;
+#endif
+		}
+#if WITH_SLATE_DEBUGGING
+		else
+		{
+			const TCHAR* Reason = TEXT("Unknown");
+			if (!FocusRecipient)
+			{
+				Reason = TEXT("Widget is a nullptr");
+			}
+			else if (!FocusRecipient->IsEnabled())
+			{
+				Reason = TEXT("Widget disabled");
+			}
+			else
+			{
+				ensure(!FocusRecipient->SupportsKeyboardFocus());
+				Reason = TEXT("Widget does not support keyboard focus");
+			}
+
+			UE_LOG(LogSlate, Warning, TEXT("Could not Explicitly navigate to widget '%s' because '%s'"), *FReflectionMetaData::GetWidgetDebugInfo(FocusRecipient), Reason);
+		}
 #endif
 	}
 	else if ( NavigationReply.GetBoundaryRule() == EUINavigationRule::Custom )
@@ -6173,7 +6198,7 @@ bool FSlateApplication::OnTouchStarted( const TSharedPtr< FGenericWindow >& Plat
 	// Don't process touches that overlap or surpass with the cursor pointer index.
 	if (TouchIndex >= (int32)ETouchIndex::CursorPointerIndex)
 	{
-#if !(UE_BUILD_SHIPPING || UE_BUILD_TEST)
+#if WITH_SLATE_DEBUGGING
 		// Only log when the touch starts, we don't want to spam the logs.
 		UE_LOG(LogSlate, Warning, TEXT("Maximum Touch Index Exceeded, %d, the maximum index allowed is %d"), TouchIndex, (((int32)ETouchIndex::CursorPointerIndex) - 1));
 #endif
