Fixes floating window with hidden tabs issue when changing modes.

diff --git a/Engine/Source/Runtime/Slate/Private/Framework/Docking/SDockingTabStack.h b/Engine/Source/Runtime/Slate/Private/Framework/Docking/SDockingTabStack.h
--- a/Engine/Source/Runtime/Slate/Private/Framework/Docking/SDockingTabStack.h
+++ b/Engine/Source/Runtime/Slate/Private/Framework/Docking/SDockingTabStack.h
@@ -119,6 +119,19 @@ public:
 
 	virtual SSplitter::ESizeRule GetSizeRule() const override;
 
+	/**
+	 * Returns a bool indicating whether or not a tab is in a state where it can be hidden. A tab typically shouldn't
+	 * be hidden in certain conditions like if there is more than one tab in the well, or if it's in a floating window.
+	 *
+	 * @return returns true if the tab can be hidden, else it returns false */
+	bool CanHideTabWell() const;
+
+	
+	/**
+	 * Sets the tab well to hidden/not hidden. 
+	 *
+	 * @param bShouldHideTabWell the bool that defines whether or not the tab well should be hidden
+	 */
 	void SetTabWellHidden( bool bShouldHideTabWell );
 	bool IsTabWellHidden() const;
 
@@ -230,9 +243,6 @@ private:
 	/** Unhides the tab well, revealing all tab headers */
 	FReply UnhideTabWell();
 
-	/** Only allow hiding the tab well when there is a single tab in it. */
-	bool CanHideTabWell() const;
-
 	/** Only allow closing the tab well when the tab allows it. */
 	bool CanCloseForegroundTab() const;
 
diff --git a/Engine/Source/Runtime/Slate/Private/Widgets/Docking/SDockTab.cpp b/Engine/Source/Runtime/Slate/Private/Widgets/Docking/SDockTab.cpp
--- a/Engine/Source/Runtime/Slate/Private/Widgets/Docking/SDockTab.cpp
+++ b/Engine/Source/Runtime/Slate/Private/Widgets/Docking/SDockTab.cpp
@@ -999,8 +999,18 @@ void SDockTab::UpdateActivationTime()
 
 void SDockTab::SetParentDockTabStackTabWellHidden(bool bIsTabWellHidden)
 {
-	if ( TSharedPtr<SDockingTabStack> ParentStack = GetParentDockTabStack() )
+	if ( const TSharedPtr<SDockingTabStack> ParentStack = GetParentDockTabStack() )
 	{
-		ParentStack->SetTabWellHidden(bIsTabWellHidden);
+		if ( bIsTabWellHidden )
+		{
+			if ( ParentStack->CanHideTabWell() )
+			{
+				ParentStack->SetTabWellHidden(true);				
+			}
+		}
+		else
+		{
+			ParentStack->SetTabWellHidden(false);
+		}
 	}
 }
\ No newline at end of file
