Slate: Expose the cvar Slate.EnableInvalidationPanels to shipping build.

diff --git a/Engine/Source/Runtime/Slate/Private/Widgets/SInvalidationPanel.cpp b/Engine/Source/Runtime/Slate/Private/Widgets/SInvalidationPanel.cpp
--- a/Engine/Source/Runtime/Slate/Private/Widgets/SInvalidationPanel.cpp
+++ b/Engine/Source/Runtime/Slate/Private/Widgets/SInvalidationPanel.cpp
@@ -14,7 +14,6 @@
 
 DECLARE_CYCLE_STAT(TEXT("SInvalidationPanel::Paint"), STAT_SlateInvalidationPaint, STATGROUP_Slate);
 
-#if WITH_SLATE_DEBUGGING
 
 void ConsoleVariableEnableInvalidationPanelsChanged(IConsoleVariable*);
 
@@ -26,6 +25,8 @@ FAutoConsoleVariableRef CVarEnableInvalidationPanels(
 	TEXT("Whether to attempt to cache any widgets through invalidation panels."),
 	FConsoleVariableDelegate::CreateStatic(ConsoleVariableEnableInvalidationPanelsChanged));
 
+#if WITH_SLATE_DEBUGGING
+
 static bool bAlwaysInvalidate = false;
 FAutoConsoleVariableRef CVarAlwaysInvalidate(
 	TEXT("Slate.AlwaysInvalidate"),
@@ -77,7 +78,6 @@ SInvalidationPanel::~SInvalidationPanel()
 	}
 }
 
-#if WITH_SLATE_DEBUGGING
 void ConsoleVariableEnableInvalidationPanelsChanged(IConsoleVariable*)
 {
 	// If the cache changed, the parent's InvalidationRoot need to rebuild its list
@@ -88,6 +88,7 @@ void ConsoleVariableEnableInvalidationPanelsChanged(IConsoleVariable*)
 	}
 }
 
+#if WITH_SLATE_DEBUGGING
 bool SInvalidationPanel::AreInvalidationPanelsEnabled()
 {
 	return bInvalidationPanelsEnabled;
@@ -111,14 +112,7 @@ void SInvalidationPanel::EnableInvalidationPanels(bool bEnable)
 
 bool SInvalidationPanel::GetCanCache() const
 {
-	bool bDebugFlags = true;
-
-#if WITH_SLATE_DEBUGGING
-	// Disable invalidation panels if global invalidation is turned on
-	bDebugFlags = bInvalidationPanelsEnabled;
-#endif
-
-	return bCanCache && !GSlateEnableGlobalInvalidation && bDebugFlags;
+	return bCanCache && !GSlateEnableGlobalInvalidation && bInvalidationPanelsEnabled;
 }
 
 void SInvalidationPanel::OnGlobalInvalidationToggled(bool bGlobalInvalidationEnabled)
