Fix searchable menu showing the first character in English when typing with IME keyboard.

diff --git a/Engine/Source/Runtime/Slate/Private/Framework/MultiBox/MultiBox.cpp b/Engine/Source/Runtime/Slate/Private/Framework/MultiBox/MultiBox.cpp
--- a/Engine/Source/Runtime/Slate/Private/Framework/MultiBox/MultiBox.cpp
+++ b/Engine/Source/Runtime/Slate/Private/Framework/MultiBox/MultiBox.cpp
@@ -896,6 +896,21 @@ void SMultiBoxWidget::CreateSearchTextWidget()
 /** Called when the SearchText changes */
 void SMultiBoxWidget::OnFilterTextChanged(const FText& InFilterText)
 {
+	// Activate the searchbox if it was empty and we are putting text in it for the first time.
+	// This is for IME keyboards only because they don't go through the OnKeyChar route.
+	if (bSearchable && SearchText.IsEmpty() 
+		&& !InFilterText.IsEmpty() 
+		&& !FSlateApplication::Get().HasUserFocusedDescendants(SearchTextWidget.ToSharedRef(), 0))
+	{
+		if (SearchTextWidget.IsValid() && SearchBlockWidget.IsValid())
+		{
+			// Make the search box visible and focused
+			SearchBlockWidget->SetVisibility(EVisibility::Visible);
+			FSlateApplication::Get().SetUserFocus(0, SearchTextWidget);
+
+		}
+	}
+
 	SearchText = InFilterText;
 
 	FilterMultiBoxEntries();
@@ -1539,6 +1554,11 @@ FReply SMultiBoxWidget::OnFocusReceived( const FGeometry& MyGeometry, const FFoc
 		ResetSearch();
 	}
 
+	if (SearchTextWidget.IsValid())
+	{
+		SearchTextWidget->EnableTextInputMethodContext();
+	}
+
 	if (InFocusEvent.GetCause() == EFocusCause::Navigation)
 	{
 		// forward focus to children
diff --git a/Engine/Source/Runtime/Slate/Private/Widgets/Input/SEditableText.cpp b/Engine/Source/Runtime/Slate/Private/Widgets/Input/SEditableText.cpp
--- a/Engine/Source/Runtime/Slate/Private/Widgets/Input/SEditableText.cpp
+++ b/Engine/Source/Runtime/Slate/Private/Widgets/Input/SEditableText.cpp
@@ -454,6 +454,11 @@ void SEditableText::AdvanceSearch(const bool InReverse)
 	EditableTextLayout->AdvanceSearch(InReverse);
 }
 
+void SEditableText::EnableTextInputMethodContext()
+{
+	EditableTextLayout->EnableTextInputMethodContext();
+}
+
 void SEditableText::SynchronizeTextStyle()
 {
 	// Has the style used for this editable text changed?
diff --git a/Engine/Source/Runtime/Slate/Private/Widgets/Input/SEditableTextBox.cpp b/Engine/Source/Runtime/Slate/Private/Widgets/Input/SEditableTextBox.cpp
--- a/Engine/Source/Runtime/Slate/Private/Widgets/Input/SEditableTextBox.cpp
+++ b/Engine/Source/Runtime/Slate/Private/Widgets/Input/SEditableTextBox.cpp
@@ -469,6 +469,10 @@ void SEditableTextBox::SetVirtualKeyboardDismissAction(TAttribute<EVirtualKeyboa
 	EditableText->SetVirtualKeyboardDismissAction(InVirtualKeyboardDismissAction);
 }
 
+void SEditableTextBox::EnableTextInputMethodContext()
+{
+	EditableText->EnableTextInputMethodContext();
+}
 #if WITH_ACCESSIBILITY
 TSharedRef<FSlateAccessibleWidget> SEditableTextBox::CreateAccessibleWidget()
 {
diff --git a/Engine/Source/Runtime/Slate/Private/Widgets/Text/SlateEditableTextLayout.cpp b/Engine/Source/Runtime/Slate/Private/Widgets/Text/SlateEditableTextLayout.cpp
--- a/Engine/Source/Runtime/Slate/Private/Widgets/Text/SlateEditableTextLayout.cpp
+++ b/Engine/Source/Runtime/Slate/Private/Widgets/Text/SlateEditableTextLayout.cpp
@@ -555,6 +555,34 @@ TSharedRef<ITextInputMethodContext> FSlateEditableTextLayout::GetTextInputMethod
 	return TextInputMethodContext.ToSharedRef();
 }
 
+void FSlateEditableTextLayout::EnableTextInputMethodContext()
+{
+	ITextInputMethodSystem* const TextInputMethodSystem = FSlateApplication::Get().GetTextInputMethodSystem();
+	if (TextInputMethodSystem)
+	{
+		if (!bHasRegisteredTextInputMethodContext)
+		{
+			bHasRegisteredTextInputMethodContext = true;
+
+			TextInputMethodChangeNotifier = TextInputMethodSystem->RegisterContext(TextInputMethodContext.ToSharedRef());
+			if (TextInputMethodChangeNotifier.IsValid())
+			{
+				TextInputMethodChangeNotifier->NotifyLayoutChanged(ITextInputMethodChangeNotifier::ELayoutChangeType::Created);
+			}
+		}
+
+		TextInputMethodContext->CacheWindow();
+
+		// Make sure to set Native OS window focus as well to ensure IME support
+		if (TSharedPtr<FGenericWindow> NativeWindow = TextInputMethodContext->GetWindow())
+		{
+			NativeWindow->SetWindowFocus();
+		}
+
+		TextInputMethodSystem->ActivateContext(TextInputMethodContext.ToSharedRef());
+	}
+}
+
 bool FSlateEditableTextLayout::Refresh()
 {
 	const FText& TextToSet = BoundText.Get(FText::GetEmpty());
@@ -811,30 +839,7 @@ bool FSlateEditableTextLayout::HandleFocusReceived(const FFocusEvent& InFocusEve
 	}
 	else
 	{
-		ITextInputMethodSystem* const TextInputMethodSystem = FSlateApplication::Get().GetTextInputMethodSystem();
-		if (TextInputMethodSystem)
-		{
-			if (!bHasRegisteredTextInputMethodContext)
-			{
-				bHasRegisteredTextInputMethodContext = true;
-
-				TextInputMethodChangeNotifier = TextInputMethodSystem->RegisterContext(TextInputMethodContext.ToSharedRef());
-				if (TextInputMethodChangeNotifier.IsValid())
-				{
-					TextInputMethodChangeNotifier->NotifyLayoutChanged(ITextInputMethodChangeNotifier::ELayoutChangeType::Created);
-				}
-			}
-
-			TextInputMethodContext->CacheWindow();
-
-			// Make sure to set Native OS window focus as well to ensure IME support
-			if (TSharedPtr<FGenericWindow> NativeWindow = TextInputMethodContext->GetWindow())
-			{
-				NativeWindow->SetWindowFocus();
-			}
-
-			TextInputMethodSystem->ActivateContext(TextInputMethodContext.ToSharedRef());
-		}
+		EnableTextInputMethodContext();
 	}
 
 	// Make sure we have the correct text (we might have been collapsed and have missed updates due to not being ticked)
diff --git a/Engine/Source/Runtime/Slate/Public/Widgets/Input/SEditableText.h b/Engine/Source/Runtime/Slate/Public/Widgets/Input/SEditableText.h
--- a/Engine/Source/Runtime/Slate/Public/Widgets/Input/SEditableText.h
+++ b/Engine/Source/Runtime/Slate/Public/Widgets/Input/SEditableText.h
@@ -366,6 +366,9 @@ public:
 	/** Advance the current search to the next match (does nothing if not currently searching) */
 	void AdvanceSearch(const bool InReverse = false);
 
+	/** Register and activate the IME context for this text layout */
+	void EnableTextInputMethodContext();
+
 protected:
 	//~ Begin SWidget Interface
 	virtual void Tick( const FGeometry& AllottedGeometry, const double InCurrentTime, const float InDeltaTime ) override;
diff --git a/Engine/Source/Runtime/Slate/Public/Widgets/Input/SEditableTextBox.h b/Engine/Source/Runtime/Slate/Public/Widgets/Input/SEditableTextBox.h
--- a/Engine/Source/Runtime/Slate/Public/Widgets/Input/SEditableTextBox.h
+++ b/Engine/Source/Runtime/Slate/Public/Widgets/Input/SEditableTextBox.h
@@ -393,6 +393,9 @@ public:
 	/** Advance the current search to the next match (does nothing if not currently searching) */
 	void AdvanceSearch(const bool InReverse = false);
 
+	/** Register and activate the IME context for the text layout of this textbox */
+	void EnableTextInputMethodContext();
+
 	bool HasError() const;
 
 	// SWidget overrides
diff --git a/Engine/Source/Runtime/Slate/Public/Widgets/Text/SlateEditableTextLayout.h b/Engine/Source/Runtime/Slate/Public/Widgets/Text/SlateEditableTextLayout.h
--- a/Engine/Source/Runtime/Slate/Public/Widgets/Text/SlateEditableTextLayout.h
+++ b/Engine/Source/Runtime/Slate/Public/Widgets/Text/SlateEditableTextLayout.h
@@ -129,6 +129,9 @@ public:
 	/** Get the IME context for this text layout */
 	TSharedRef<ITextInputMethodContext> GetTextInputMethodContext() const;
 
+	/** Register and activate the IME context for this text layout */
+	void EnableTextInputMethodContext();
+
 	/** Refresh this editable text immediately, rather than wait for the usual caching mechanisms to take affect on the text Tick */
 	bool Refresh();
 
