Fix to Numeric Entry Box to:
- Show the same strings in Tool Tip as the Spin Box
- Fix to properly use Locals settings and show proper

diff --git a/Engine/Source/Runtime/Slate/Public/Widgets/Input/SNumericEntryBox.h b/Engine/Source/Runtime/Slate/Public/Widgets/Input/SNumericEntryBox.h
--- a/Engine/Source/Runtime/Slate/Public/Widgets/Input/SNumericEntryBox.h
+++ b/Engine/Source/Runtime/Slate/Public/Widgets/Input/SNumericEntryBox.h
@@ -262,12 +262,14 @@ public:
 				.OnBeginSliderMovement(InArgs._OnBeginSliderMovement)
 				.OnEndSliderMovement(InArgs._OnEndSliderMovement)
 				.MinDesiredWidth(InArgs._MinDesiredValueWidth)
-				.TypeInterface(Interface);
+				.TypeInterface(Interface)
+				.ToolTipText(this, &SNumericEntryBox<NumericType>::GetValueAsText);
 		}
 
 		// Always create an editable text box.  In the case of an undetermined value being passed in, we cant use the spinbox.
 		SAssignNew(EditableText, SEditableText)
 			.Text(this, &SNumericEntryBox<NumericType>::OnGetValueForTextBox)
+			.ToolTipText(this, &SNumericEntryBox<NumericType>::GetValueAsText)
 			.ColorAndOpacity(InArgs._EditableTextBoxStyle->ForegroundColor)
 			.Visibility(bAllowSpin ? EVisibility::Collapsed : EVisibility::Visible)
 			.Font(InArgs._Font.IsSet() ? InArgs._Font : InArgs._EditableTextBoxStyle->TextStyle.Font)
@@ -553,14 +555,15 @@ private:
 		return CachedValueString;
 	}
 
+	/** @return the value being observed by the Numeric Entry Box as a FText */
 	FText GetValueAsText() const
 	{
 		const TOptional<NumericType>& Value = ValueAttribute.Get();
 		if (Value.IsSet() == true)
 		{
-			return FText::FromString(GetCachedString(Value.GetValue()));
+			NumericType CurrentValue = Value.GetValue();
+			return FText::FromString(GetCachedString(CurrentValue));
 		}
-		
 		return FText::GetEmpty();
 	}
 
@@ -569,25 +572,21 @@ private:
 	 */
 	FText OnGetValueForTextBox() const
 	{
-		FText NewText = FText::GetEmpty();
-
 		if( EditableText->GetVisibility() == EVisibility::Visible )
 		{
-			const auto& Value = ValueAttribute.Get();
-
-			// If the value was set convert it to a string, otherwise the value cannot be determined
-			if( Value.IsSet() == true )
+			const TOptional<NumericType>& Value = ValueAttribute.Get();
+			if (Value.IsSet() == true)
 			{
-				NewText = FText::FromString(GetCachedString(Value.GetValue()));
+				return GetValueAsText();
 			}
 			else
 			{
-				NewText = UndeterminedString;
+				return UndeterminedString;
 			}
 		}
 
-		// The box isnt visible, just return an empty string
-		return NewText;
+		// The box isnt visible, just return an empty Text
+		return  FText::GetEmpty();
 	}
 
 
