Optimized code to reserve the sizes of arrays.

diff --git a/Engine/Source/Runtime/Slate/Private/Framework/Text/TextLayout.cpp b/Engine/Source/Runtime/Slate/Private/Framework/Text/TextLayout.cpp
--- a/Engine/Source/Runtime/Slate/Private/Framework/Text/TextLayout.cpp
+++ b/Engine/Source/Runtime/Slate/Private/Framework/Text/TextLayout.cpp
@@ -2164,11 +2164,13 @@ void FTextLayout::AddLine( const FNewLineData& NewLine )
 
 void FTextLayout::AddLines( const TArray<FNewLineData>& NewLines )
 {
+	LineModels.Reserve(NewLines.Num());
 	for (const auto& NewLine : NewLines)
 	{
 		FLineModel LineModel( NewLine.Text );
 		TransformLineText(LineModel);
 
+		LineModel.Runs.Reserve(NewLine.Runs.Num());
 		for (const auto& Run : NewLine.Runs)
 		{
 			LineModel.Runs.Add( FRunModel( Run ) );
