Slate: Fix lookup of cached values for numeric entry box. Fixed caching of the numeric value in Tick() if the widget didn't have a spinner.

#rnx

diff --git a/Engine/Source/Runtime/Slate/Public/Widgets/Input/SNumericEntryBox.h b/Engine/Source/Runtime/Slate/Public/Widgets/Input/SNumericEntryBox.h
--- a/Engine/Source/Runtime/Slate/Public/Widgets/Input/SNumericEntryBox.h
+++ b/Engine/Source/Runtime/Slate/Public/Widgets/Input/SNumericEntryBox.h
@@ -546,13 +546,8 @@ private:
 
 	FString GetCachedString(const NumericType CurrentValue) const
 	{
-		bool bUseCachedString = CachedExternalValue.IsSet() && CurrentValue == CachedExternalValue.GetValue() && !bCachedValueStringDirty;
-		if (!bUseCachedString)
-		{
-			CachedValueString = Interface->ToString(CurrentValue);
-			bCachedValueStringDirty = false;
-		}
-		return CachedValueString;
+		const bool bUseCachedString = CachedExternalValue.IsSet() && CurrentValue == CachedExternalValue.GetValue() && !bCachedValueStringDirty;
+		return bUseCachedString ? CachedValueString : Interface->ToString(CurrentValue);  
 	}
 
 	/** @return the value being observed by the Numeric Entry Box as a FText */
@@ -577,7 +572,7 @@ private:
 			const TOptional<NumericType>& Value = ValueAttribute.Get();
 			if (Value.IsSet() == true)
 			{
-				return GetValueAsText();
+				return FText::FromString(GetCachedString(Value.GetValue()));
 			}
 			else
 			{
@@ -689,31 +684,32 @@ private:
 	 */
 	virtual void Tick( const FGeometry& AllottedGeometry, const double InCurrentTime, const float InDeltaTime ) override
 	{
-		// Visibility toggle only matters if the spinbox is used
-		if (!SpinBox.IsValid())
-		{
-			return;
-		}
-
-		const auto& Value = ValueAttribute.Get();
-
+		// Update the cached value, if needed.
+		const TOptional<NumericType>& Value = ValueAttribute.Get();
 		if (Value.IsSet() == true)
 		{
 			SetCachedString(Value.GetValue());
-
-			if (SpinBox->GetVisibility() != EVisibility::Visible)
-			{
-				// Set the visibility of the spinbox to visible if we have a valid value
-				SpinBox->SetVisibility( EVisibility::Visible );
-				// The text box should be invisible
-				EditableText->SetVisibility( EVisibility::Collapsed );
-			}
 		}
-		else
+		
+		// Visibility toggle only matters if the spinbox is used
+		if (SpinBox.IsValid())
 		{
-			// The value isn't set so the spinbox should be hidden and the text box shown
-			SpinBox->SetVisibility(EVisibility::Collapsed);
-			EditableText->SetVisibility(EVisibility::Visible);
+			if (Value.IsSet() == true)
+			{
+				if (SpinBox->GetVisibility() != EVisibility::Visible)
+				{
+					// Set the visibility of the spinbox to visible if we have a valid value
+					SpinBox->SetVisibility( EVisibility::Visible );
+					// The text box should be invisible
+					EditableText->SetVisibility( EVisibility::Collapsed );
+				}
+			}
+			else
+			{
+				// The value isn't set so the spinbox should be hidden and the text box shown
+				SpinBox->SetVisibility(EVisibility::Collapsed);
+				EditableText->SetVisibility(EVisibility::Visible);
+			}
 		}
 	}
 
@@ -783,9 +779,9 @@ private:
 	/** Cached value of entry box, updated on set & per tick */
 	TOptional<NumericType> CachedExternalValue;
 	/** Used to prevent per-frame re-conversion of the cached numeric value to a string. */
-	mutable FString CachedValueString;
+	FString CachedValueString;
 	/** Whetever the interfaced setting changed and the CachedValueString needs to be recomputed. */
-	mutable bool bCachedValueStringDirty;
+	bool bCachedValueStringDirty;
 	TAttribute< TOptional<int32> > MinFractionalDigits;
 	TAttribute< TOptional<int32> > MaxFractionalDigits;
 };
