Fixed blank slate notification window when Interchange import is failing right away.
Root cause: Race condition between FSlateAsyncTaskNotificationImpl::TickNotification and FCoreAsyncTaskNotificationImpl::SetNotificationState.
Solution: Replace use of FCoreAsyncTaskNotificationImpl::State with the local variable CompletionStateToApply

#rnx

diff --git a/Engine/Source/Runtime/Slate/Private/Framework/Notifications/SlateAsyncTaskNotificationImpl.cpp b/Engine/Source/Runtime/Slate/Private/Framework/Notifications/SlateAsyncTaskNotificationImpl.cpp
--- a/Engine/Source/Runtime/Slate/Private/Framework/Notifications/SlateAsyncTaskNotificationImpl.cpp
+++ b/Engine/Source/Runtime/Slate/Private/Framework/Notifications/SlateAsyncTaskNotificationImpl.cpp
@@ -431,7 +431,7 @@ bool FSlateAsyncTaskNotificationImpl::TickNotification(float InDeltaTime)
 		if(OwningNotification)
 		{
 			SNotificationItem::ECompletionState OwningCompletionState = SNotificationItem::CS_None;
-			switch (State)
+			switch (CompletionStateToApply)
 			{
 			case EAsyncTaskNotificationState::Pending:
 				OwningCompletionState = SNotificationItem::CS_Pending;
