Add KeyDownHandlerSupport to STreeView

diff --git a/Engine/Source/Runtime/Slate/Public/Widgets/Views/SListView.h b/Engine/Source/Runtime/Slate/Public/Widgets/Views/SListView.h
--- a/Engine/Source/Runtime/Slate/Public/Widgets/Views/SListView.h
+++ b/Engine/Source/Runtime/Slate/Public/Widgets/Views/SListView.h
@@ -335,7 +335,12 @@ public:
 				return Reply;
 			}
 		}
+		return OnKeyDown_Internal(MyGeometry, InKeyEvent);
+	}
 
+protected:
+	FReply OnKeyDown_Internal(const FGeometry& MyGeometry, const FKeyEvent& InKeyEvent)
+	{
 		const TArrayView<const ItemType> ItemsSourceRef = GetItems();
 
 		// Don't respond to key-presses containing "Alt" as a modifier
@@ -474,11 +479,6 @@ public:
 
 		return STableViewBase::OnKeyDown(MyGeometry, InKeyEvent);
 	}
-
-private:
-
-	FOnKeyDown OnKeyDownHandler;
-
 public:
 
 	virtual FNavigationReply OnNavigation(const FGeometry& MyGeometry, const FNavigationEvent& InNavigationEvent) override
@@ -2545,6 +2545,9 @@ protected:
 
 	/** Called when the user double-clicks on an element in the list view with the left mouse button */
 	FOnMouseButtonDoubleClick OnDoubleClick;
+	
+	/** Called when the user presses a keyboard key */
+	FOnKeyDown OnKeyDownHandler;
 
 	/** True when the list view supports keyboard focus */
 	TAttribute<bool> IsFocusable;
diff --git a/Engine/Source/Runtime/Slate/Public/Widgets/Views/STreeView.h b/Engine/Source/Runtime/Slate/Public/Widgets/Views/STreeView.h
--- a/Engine/Source/Runtime/Slate/Public/Widgets/Views/STreeView.h
+++ b/Engine/Source/Runtime/Slate/Public/Widgets/Views/STreeView.h
@@ -225,6 +225,9 @@ public:
 
 		/** If true, Show the current hierarchy of items pinned at the top of the Tree View */
 		SLATE_ATTRIBUTE(bool, ShouldStackHierarchyHeaders)
+	
+		/** Callback delegate to have first chance handling of the OnKeyDown event */
+		SLATE_EVENT(FOnKeyDown, OnKeyDownHandler)
 
 	SLATE_END_ARGS()
 
@@ -246,6 +249,7 @@ public:
 		this->OnSetExpansionRecursive = InArgs._OnSetExpansionRecursive;
 		this->TreeItemsSource = InArgs._TreeItemsSource;
 
+		this->OnKeyDownHandler = InArgs._OnKeyDownHandler;
 		this->OnContextMenuOpening = InArgs._OnContextMenuOpening;
 		this->OnClick = InArgs._OnMouseButtonClick;
 		this->OnDoubleClick = InArgs._OnMouseButtonDoubleClick;
@@ -339,6 +343,15 @@ public:
 
 	virtual FReply OnKeyDown( const FGeometry& MyGeometry, const FKeyEvent& InKeyEvent ) override
 	{
+		if (this->OnKeyDownHandler.IsBound())
+		{
+			FReply Reply = this->OnKeyDownHandler.Execute(MyGeometry, InKeyEvent);
+			if (Reply.IsEventHandled())
+			{
+				return Reply;
+			}
+		}
+		
 		// Check for selection/expansion toggling keys (Left, Right)
 		// SelectorItem represents the keyboard selection. If it isn't valid then we don't know what to expand.
 		// Don't respond to key-presses containing "Alt" as a modifier
@@ -405,7 +418,7 @@ public:
 			}
 		}
 
-		return SListView<ItemType>::OnKeyDown(MyGeometry, InKeyEvent);
+		return SListView<ItemType>::OnKeyDown_Internal(MyGeometry, InKeyEvent);
 	}
 	
 private:
