Text in palette buttons in Modeling mode can overrun their borders

diff --git a/Engine/Source/Editor/EditorStyle/Private/StarshipStyle.cpp b/Engine/Source/Editor/EditorStyle/Private/StarshipStyle.cpp
--- a/Engine/Source/Editor/EditorStyle/Private/StarshipStyle.cpp
+++ b/Engine/Source/Editor/EditorStyle/Private/StarshipStyle.cpp
@@ -1688,7 +1688,7 @@ void FStarshipEditorStyle::FStyle::SetupGeneralStyles()
 		SlimPaletteToolBarStyle.SetLabelStyle(SlimPaletteToolbarLabelStyle);
 		SlimPaletteToolBarStyle.SetBackgroundPadding(FMargin(2.0f, 4.0f, 6.0f, 8.0f));
 		SlimPaletteToolBarStyle.SetBackground(FSlateColorBrush(FStyleColors::Panel));
-		SlimPaletteToolBarStyle.SetLabelPadding(FMargin(0.0f, 0.0f, 0.0f, 0.0f));
+		SlimPaletteToolBarStyle.SetLabelPadding(FMargin(0.0f, 0.0f, 4.0f, 0.0f));
 		SlimPaletteToolBarStyle.SetIconSize(Icon20x20);
 
 		SlimPaletteToolBarStyle.SetButtonPadding(       FMargin(4.0f, 4.0f, 0.0f, 0.0f));
diff --git a/Engine/Source/Runtime/Slate/Private/Framework/MultiBox/SToolBarButtonBlock.cpp b/Engine/Source/Runtime/Slate/Private/Framework/MultiBox/SToolBarButtonBlock.cpp
--- a/Engine/Source/Runtime/Slate/Private/Framework/MultiBox/SToolBarButtonBlock.cpp
+++ b/Engine/Source/Runtime/Slate/Private/Framework/MultiBox/SToolBarButtonBlock.cpp
@@ -220,13 +220,23 @@ void SToolBarButtonBlock::BuildMultiBlockWidget(const ISlateStyle* StyleSet, con
 		.Image(this, &SToolBarButtonBlock::GetIconBrush);
 
 	IconWidget->AddLayer(TAttribute<const FSlateBrush*>(this, &SToolBarButtonBlock::GetOverlayIconBrush));
+	const bool bIsSlimHorizontalUniformToolBar = MultiBox->GetType() == EMultiBoxType::SlimHorizontalUniformToolBar;
 
 	// Create the content for our button
 	TSharedRef<SWidget> ButtonContent = SNullWidget::NullWidget;
 	if (MultiBox->GetType() == EMultiBoxType::SlimHorizontalToolBar 
-		|| MultiBox->GetType() == EMultiBoxType::SlimHorizontalUniformToolBar)
+		|| bIsSlimHorizontalUniformToolBar)
 	{
 		const FVector2f IconSize = ToolBarStyle.IconSize;
+		const TSharedRef<STextBlock> TextBlock = SNew(STextBlock)
+				.Visibility(LabelVisibility)
+				.Text(ActualLabel)
+				.TextStyle(&ToolBarStyle.LabelStyle); // Smaller font for tool tip labels
+
+		if (bIsSlimHorizontalUniformToolBar)
+		{
+			TextBlock->SetOverflowPolicy(ETextOverflowPolicy::Ellipsis);
+		}
 
 		IconWidget->SetDesiredSizeOverride(FVector2D(IconSize));
 		ButtonContent =
@@ -241,16 +251,15 @@ void SToolBarButtonBlock::BuildMultiBlockWidget(const ISlateStyle* StyleSet, con
 				IconWidget
 			]
 			// Label text
-			+ SHorizontalBox::Slot()
-			.AutoWidth()
-			.Padding(ToolBarStyle.LabelPadding)
-			.VAlign(VAlign_Center)
-			[
-				SNew(STextBlock)
-				.Visibility(LabelVisibility)
-				.Text(ActualLabel)
-				.TextStyle(&ToolBarStyle.LabelStyle)	// Smaller font for tool tip labels
-			];
+		+ (bIsSlimHorizontalUniformToolBar ?
+		SHorizontalBox::Slot()
+		.Padding(ToolBarStyle.LabelPadding)
+		.VAlign(VAlign_Center) :
+		SHorizontalBox::Slot()
+		.AutoWidth()
+		.Padding(ToolBarStyle.LabelPadding)
+		.VAlign(VAlign_Center))
+		[ TextBlock ];
 	}
 	else
 	{
