Fixes crash upon pinning modeling mode tab to sidebar, navigating away and then back to modeling mode.

diff --git a/Engine/Source/Runtime/Slate/Private/Framework/Docking/SDockingTabStack.cpp b/Engine/Source/Runtime/Slate/Private/Framework/Docking/SDockingTabStack.cpp
--- a/Engine/Source/Runtime/Slate/Private/Framework/Docking/SDockingTabStack.cpp
+++ b/Engine/Source/Runtime/Slate/Private/Framework/Docking/SDockingTabStack.cpp
@@ -1028,13 +1028,21 @@ FReply SDockingTabStack::UnhideTabWell()
 
 bool SDockingTabStack::CanHideTabWell() const
 {
-	TSharedPtr<SDockingSplitter> ParentNode = ParentNodePtr.Pin();
-	// Is target tab located at the upper and left most among tabs in the parent window(as first child). Unreal icon will overlap golden triangle(unhide button) when the tab is the first child of the window.
-	const bool bIsUpperLeftmostTab = (FGlobalTabmanager::Get()->GetActiveTab() == ParentNode->GetAllChildTabs()[0]);
-	// Is target tab in the Floating Window. Unreal icon will overlap when the tab is in the floating window(without menu)
-	const bool bIsInFloatingWindow = ParentNode->GetDockArea()->GetParentWindow().IsValid();
+	const TSharedPtr<SDockingSplitter> ParentNode = ParentNodePtr.Pin();
+
+	if (ParentNode && !ParentNode->GetAllChildTabs().IsEmpty())
+	{
+		// Is target tab located at the upper and left most among tabs in the parent window(as first child). Unreal icon will overlap golden triangle(unhide button) when the tab is the first child of the window.
+		const bool bIsUpperLeftmostTab = (FGlobalTabmanager::Get()->GetActiveTab() == ParentNode->GetAllChildTabs()[0]);
+		// Is target tab in the Floating Window. Unreal icon will overlap when the tab is in the floating window(without menu)
+		const bool bIsInFloatingWindow = ParentNode->GetDockArea()->GetParentWindow().IsValid();
 
-	return GetNumTabs() == 1 && FGlobalTabmanager::Get()->CanSetAsActiveTab(GetTabs()[0]) && !(bIsUpperLeftmostTab && bIsInFloatingWindow);
+		return GetNumTabs() == 1 && FGlobalTabmanager::Get()->CanSetAsActiveTab(GetTabs()[0]) && !(bIsUpperLeftmostTab && bIsInFloatingWindow);
+	}
+
+	/* in the case where there are no parent splitter or child tabs, it is invalid to hide the tab well.
+	 * The likely case for this would be in a sidebar flyout  */
+	return false;
 }
 
 bool SDockingTabStack::CanCloseForegroundTab() const
