Fixes the issue where if UE editor preference "use small toolbar icons" is set to true, Modeling mode palette buttons have no text

diff --git a/Engine/Source/Runtime/Slate/Private/Framework/MultiBox/SToolBarButtonBlock.cpp b/Engine/Source/Runtime/Slate/Private/Framework/MultiBox/SToolBarButtonBlock.cpp
--- a/Engine/Source/Runtime/Slate/Private/Framework/MultiBox/SToolBarButtonBlock.cpp
+++ b/Engine/Source/Runtime/Slate/Private/Framework/MultiBox/SToolBarButtonBlock.cpp
@@ -236,6 +236,7 @@ void SToolBarButtonBlock::BuildMultiBlockWidget(const ISlateStyle* StyleSet, con
 		if (bIsSlimHorizontalUniformToolBar)
 		{
 			TextBlock->SetOverflowPolicy(ETextOverflowPolicy::Ellipsis);
+			TextBlock->SetVisibility(EVisibility::Visible);
 		}
 
 		IconWidget->SetDesiredSizeOverride(FVector2D(IconSize));
