Menus - Allow suppress extenders using runtime customization
#rnx

diff --git a/Engine/Source/Developer/ToolMenus/Private/ToolMenus.cpp b/Engine/Source/Developer/ToolMenus/Private/ToolMenus.cpp
--- a/Engine/Source/Developer/ToolMenus/Private/ToolMenus.cpp
+++ b/Engine/Source/Developer/ToolMenus/Private/ToolMenus.cpp
@@ -898,6 +898,11 @@ void UToolMenus::ApplyCustomization(UToolMenu* GeneratedMenu)
 	}
 	FCustomizedToolMenu CustomizedMenu = CustomizationHierarchy.GenerateFlattened();
 
+	if (CustomizedMenu.IsSuppressExtenders())
+	{
+		GeneratedMenu->SetExtendersEnabled(false);
+	}
+
 	TArray<FToolMenuSection> NewSections;
 	NewSections.Reserve(GeneratedMenu->Sections.Num());
 
@@ -2409,6 +2414,11 @@ void UToolMenus::UnregisterRuntimeMenuCustomizationOwner(FName InOwnerName)
 		{
 			bNeedsRefresh = true;
 		}
+
+		if (CustomizedToolMenu.SuppressExtenders.Remove(InOwnerName) > 0)
+		{
+			bNeedsRefresh = true;
+		}
 	}
 
 	// Refresh any widgets that are currently displayed to the user
diff --git a/Engine/Source/Runtime/Slate/Private/Framework/MultiBox/ToolMenuCustomization.cpp b/Engine/Source/Runtime/Slate/Private/Framework/MultiBox/ToolMenuCustomization.cpp
--- a/Engine/Source/Runtime/Slate/Private/Framework/MultiBox/ToolMenuCustomization.cpp
+++ b/Engine/Source/Runtime/Slate/Private/Framework/MultiBox/ToolMenuCustomization.cpp
@@ -79,6 +79,23 @@ bool FCustomizedToolMenu::IsSectionHidden(const FName InEntryName) const
 	return GetSectionVisiblity(InEntryName) == ECustomizedToolMenuVisibility::Hidden;
 }
 
+void FCustomizedToolMenu::SetSuppressExtenders(const FName InOwnerName, const bool bInSuppress)
+{
+	if (bInSuppress)
+	{
+		SuppressExtenders.AddUnique(InOwnerName);
+	}
+	else
+	{
+		SuppressExtenders.Remove(InOwnerName);
+	}
+}
+
+bool FCustomizedToolMenu::IsSuppressExtenders() const
+{
+	return !SuppressExtenders.IsEmpty();
+}
+
 FName FCustomizedToolMenuHierarchy::GetEntrySectionName(const FName InEntryName) const
 {
 	for (int32 i = Hierarchy.Num() - 1; i >= 0; --i)
@@ -142,6 +159,22 @@ bool FCustomizedToolMenuHierarchy::IsSectionHidden(const FName InSectionName) co
 	return false;
 }
 
+bool FCustomizedToolMenuHierarchy::IsSuppressExtenders() const
+{
+	for (int32 i = Hierarchy.Num() - 1; i >= 0; --i)
+	{
+		if (Hierarchy[i])
+		{
+			if (Hierarchy[i]->IsSuppressExtenders())
+			{
+				return true;
+			}
+		}
+	}
+
+	return false;
+}
+
 FCustomizedToolMenu FCustomizedToolMenuHierarchy::GenerateFlattened() const
 {
 	static auto HandleCustomizedToolMenu = [](FCustomizedToolMenu& Result, const FCustomizedToolMenu* Current)
@@ -151,6 +184,11 @@ FCustomizedToolMenu FCustomizedToolMenuHierarchy::GenerateFlattened() const
 			return;
 		}
 
+		if (Current->IsSuppressExtenders() && !Result.IsSuppressExtenders())
+		{
+			Result.SuppressExtenders = Current->SuppressExtenders;
+		}
+
 		if (Current->SectionOrder.Num() > 0)
 		{
 			Result.SectionOrder = Current->SectionOrder;
diff --git a/Engine/Source/Runtime/Slate/Public/Framework/MultiBox/ToolMenuBase.h b/Engine/Source/Runtime/Slate/Public/Framework/MultiBox/ToolMenuBase.h
--- a/Engine/Source/Runtime/Slate/Public/Framework/MultiBox/ToolMenuBase.h
+++ b/Engine/Source/Runtime/Slate/Public/Framework/MultiBox/ToolMenuBase.h
@@ -74,6 +74,9 @@ struct SLATE_API FCustomizedToolMenu
 	UPROPERTY()
 	TArray<FName> SectionOrder;
 
+	UPROPERTY()
+	TArray<FName> SuppressExtenders;
+
 	FNamePermissionList MenuPermissions;
 
 	FCustomizedToolMenuEntry* FindEntry(const FName InEntryName);
@@ -88,6 +91,9 @@ struct SLATE_API FCustomizedToolMenu
 	FCustomizedToolMenuSection* AddSection(const FName InSectionName);
 	ECustomizedToolMenuVisibility GetSectionVisiblity(const FName InSectionName) const;
 	bool IsSectionHidden(const FName InSectionName) const;
+
+	void SetSuppressExtenders(const FName InOwnerName, const bool bInSuppress);
+	bool IsSuppressExtenders() const;
 };
 
 struct SLATE_API FCustomizedToolMenuHierarchy
@@ -97,6 +103,8 @@ struct SLATE_API FCustomizedToolMenuHierarchy
 
 	bool IsSectionHidden(const FName InSectionName) const;
 
+	bool IsSuppressExtenders() const;
+
 	FCustomizedToolMenu GenerateFlattened() const;
 
 	TArray<const FCustomizedToolMenu*> Hierarchy;
