Avoid losing tabs by not collapsing the sidebar when docking a tab to it.

diff --git a/Engine/Source/Runtime/Slate/Private/Framework/Docking/SDockingArea.cpp b/Engine/Source/Runtime/Slate/Private/Framework/Docking/SDockingArea.cpp
--- a/Engine/Source/Runtime/Slate/Private/Framework/Docking/SDockingArea.cpp
+++ b/Engine/Source/Runtime/Slate/Private/Framework/Docking/SDockingArea.cpp
@@ -624,10 +624,16 @@ void SDockingArea::UpdateWindowChromeAndSidebar()
 				}
 			}
 
-			bCanHaveSidebar = !bContainsMajorTabs;
+			if (LeftSidebar->GetAllTabIds().IsEmpty())
+			{
+				LeftSidebar->SetVisibility(EVisibility::Collapsed);
+			}
+			if (RightSidebar->GetAllTabIds().IsEmpty())
+			{
+				RightSidebar->SetVisibility(EVisibility::Collapsed);
+			}
 
-			LeftSidebar->SetVisibility(EVisibility::Collapsed);
-			RightSidebar->SetVisibility(EVisibility::Collapsed);
+			bCanHaveSidebar = !bContainsMajorTabs;
 
 			if (bCanHaveSidebar)
 			{
