Add 1 pixel of padding to custom widgets without labels in menus so that the menu border shows up

diff --git a/Engine/Source/Runtime/Slate/Private/Framework/MultiBox/SWidgetBlock.cpp b/Engine/Source/Runtime/Slate/Private/Framework/MultiBox/SWidgetBlock.cpp
--- a/Engine/Source/Runtime/Slate/Private/Framework/MultiBox/SWidgetBlock.cpp
+++ b/Engine/Source/Runtime/Slate/Private/Framework/MultiBox/SWidgetBlock.cpp
@@ -111,7 +111,8 @@ void SWidgetBlock::BuildMultiBlockWidget(const ISlateStyle* StyleSet, const FNam
 		// If there is no label, allow the custom menu widget to consume the entire space
 		if (!bHasLabel)
 		{
-			Padding = FMargin(0);
+			// Wee use 1 pixel of padding to ensure the menu border shows up
+			Padding = FMargin(1);
 		}
 	}
 
