Fix Tab Manager crashing when Loading a Layout right after resizing a tab

diff --git a/Engine/Source/Runtime/Slate/Private/Framework/Docking/TabManager.cpp b/Engine/Source/Runtime/Slate/Private/Framework/Docking/TabManager.cpp
--- a/Engine/Source/Runtime/Slate/Private/Framework/Docking/TabManager.cpp
+++ b/Engine/Source/Runtime/Slate/Private/Framework/Docking/TabManager.cpp
@@ -951,11 +951,7 @@ TSharedRef<FTabManager::FLayout> FTabManager::PersistLayout() const
 
 void FTabManager::SavePersistentLayout()
 {
-	if (PendingLayoutSaveHandle.IsValid())
-	{
-		FTSTicker::GetCoreTicker().RemoveTicker(PendingLayoutSaveHandle);
-		PendingLayoutSaveHandle.Reset();
-	}
+	ClearPendingLayoutSave();
 
 	const TSharedRef<FLayout> LayoutState = this->PersistLayout();
 	OnPersistLayout_Handler.ExecuteIfBound(LayoutState);
@@ -965,22 +961,30 @@ void FTabManager::RequestSavePersistentLayout()
 {
 	// if we already have a request pending, remove it and schedule a new one
 	// this is to avoid hitches when eg. resizing a docked tab
-	if (PendingLayoutSaveHandle.IsValid())
-	{
-		FTSTicker::GetCoreTicker().RemoveTicker(PendingLayoutSaveHandle);
-		PendingLayoutSaveHandle.Reset();
-	}
+	ClearPendingLayoutSave();
 
-	auto OnTick = [this](float FrameTime)
+	auto OnTick = [ThisWeak = AsWeak()](float FrameTime)
 	{
-		this->PendingLayoutSaveHandle.Reset();
-		this->SavePersistentLayout();
+		if (TSharedPtr<FTabManager> This = ThisWeak.Pin())
+		{
+			This->PendingLayoutSaveHandle.Reset();
+			This->SavePersistentLayout();
+		}
 		return false;
 	};
 
 	PendingLayoutSaveHandle = FTSTicker::GetCoreTicker().AddTicker(FTickerDelegate::CreateLambda(OnTick), 5.0f);
 }
 
+void FTabManager::ClearPendingLayoutSave()
+{
+	if (PendingLayoutSaveHandle.IsValid())
+	{
+		FTSTicker::GetCoreTicker().RemoveTicker(PendingLayoutSaveHandle);
+		PendingLayoutSaveHandle.Reset();
+	}
+}
+
 FTabSpawnerEntry& FTabManager::RegisterTabSpawner(const FName TabId, const FOnSpawnTab& OnSpawnTab, const FCanSpawnTab& CanSpawnTab)
 {
 	ensure(!TabSpawner.Contains(TabId));
@@ -2072,6 +2076,11 @@ TSharedPtr<SDockTab> FTabManager::FindExistingLiveTab( const FTabId& TabId ) con
 	return TSharedPtr<SDockTab>();
 }
 
+FTabManager::~FTabManager()
+{
+	ClearPendingLayoutSave();
+}
+
 TSharedPtr<SDockTab> FTabManager::FindLastTabInWindow(TSharedPtr<SWindow> Window) const
 {
 	if ( Window.IsValid() )
diff --git a/Engine/Source/Runtime/Slate/Public/Framework/Docking/TabManager.h b/Engine/Source/Runtime/Slate/Public/Framework/Docking/TabManager.h
--- a/Engine/Source/Runtime/Slate/Public/Framework/Docking/TabManager.h
+++ b/Engine/Source/Runtime/Slate/Public/Framework/Docking/TabManager.h
@@ -835,6 +835,9 @@ class SLATE_API FTabManager : public TSharedFromThis<FTabManager>
 		/** Request a deferred save of the layout. */
 		void RequestSavePersistentLayout();
 
+		/** Clears the request for a deferred save of the layout */
+		void ClearPendingLayoutSave();
+
 		/**
 		 * Register a new tab spawner with the tab manager.  The spawner will be called when anyone calls
 		 * InvokeTab().
@@ -911,9 +914,7 @@ class SLATE_API FTabManager : public TSharedFromThis<FTabManager>
 		 */
 		TSharedPtr<SDockTab> FindExistingLiveTab(const FTabId& TabId) const;
 
-		virtual ~FTabManager()
-		{
-		}
+		virtual ~FTabManager();
 
 		/** Sets whether or not this tab manager supports a custom menu bar for the active major tab that will be shown on top of the major tab area in the window this tab manager resides in. */
 		void SetAllowWindowMenuBar(bool bInAllowWindowMenuBar);
