Sorted includes and updated to include CoreUObjectSharedPCH.h to make it easier to manage.

diff --git a/Engine/Source/Runtime/Slate/Public/SlateSharedPCH.h b/Engine/Source/Runtime/Slate/Public/SlateSharedPCH.h
--- a/Engine/Source/Runtime/Slate/Public/SlateSharedPCH.h
+++ b/Engine/Source/Runtime/Slate/Public/SlateSharedPCH.h
@@ -2,381 +2,174 @@
 
 #pragma once
 
+#include "CoreUObjectSharedPCH.h"
+
 // From Core:
-#include "Containers/ContainersFwd.h"
-#include "UObject/UObjectHierarchyFwd.h"
-#include "CoreTypes.h"
-#include "CoreFwd.h"
-#include "Misc/Exec.h"
-#include "Misc/AssertionMacros.h"
-#include "HAL/PlatformMisc.h"
-#include "GenericPlatform/GenericPlatformMisc.h"
-#include "HAL/PlatformCrt.h"
-#include "Containers/Array.h"
-#include "HAL/UnrealMemory.h"
-#include "Templates/IsPointer.h"
-#include "HAL/PlatformMemory.h"
-#include "GenericPlatform/GenericPlatformMemory.h"
-#include "HAL/MemoryBase.h"
-#include "Misc/OutputDevice.h"
-#include "Logging/LogVerbosity.h"
-#include "Misc/VarArgs.h"
-#include "HAL/PlatformAtomics.h"
-#include "GenericPlatform/GenericPlatformAtomics.h"
-#include "Templates/UnrealTypeTraits.h"
-#include "Templates/AndOrNot.h"
-#include "Templates/IsArithmetic.h"
-#include "Templates/IsPODType.h"
-#include "Templates/IsTriviallyCopyConstructible.h"
-#include "Templates/UnrealTemplate.h"
-#include "Templates/EnableIf.h"
-#include "Templates/RemoveReference.h"
-#include "Templates/TypeCompatibleBytes.h"
-#include "Templates/ChooseClass.h"
-#include "Templates/IntegralConstant.h"
-#include "Templates/IsClass.h"
-#include "Traits/IsContiguousContainer.h"
-#include "Containers/ContainerAllocationPolicies.h"
-#include "HAL/PlatformMath.h"
-#include "GenericPlatform/GenericPlatformMath.h"
-#include "Templates/MemoryOps.h"
-#include "Templates/IsTriviallyCopyAssignable.h"
-#include "Templates/IsTriviallyDestructible.h"
-#include "Math/NumericLimits.h"
-#include "Serialization/Archive.h"
-#include "Templates/IsEnumClass.h"
-#include "HAL/PlatformProperties.h"
-#include "GenericPlatform/GenericPlatformProperties.h"
-#include "Misc/Compression.h"
-#include "Misc/EngineVersionBase.h"
-#include "Internationalization/TextNamespaceFwd.h"
-#include "Templates/Less.h"
-#include "Templates/Sorting.h"
-#include "Containers/UnrealString.h"
-#include "Misc/CString.h"
-#include "Misc/Char.h"
-#include "HAL/PlatformString.h"
-#include "GenericPlatform/GenericPlatformStricmp.h"
-#include "GenericPlatform/GenericPlatformString.h"
-#include "Misc/Crc.h"
-#include "Math/UnrealMathUtility.h"
-#include "Containers/Map.h"
-#include "Misc/StructBuilder.h"
-#include "Templates/AlignmentTemplates.h"
-#include "Templates/Function.h"
-#include "Templates/Decay.h"
-#include "Templates/Invoke.h"
-#include "Templates/PointerIsConvertibleFromTo.h"
-#include "Containers/Set.h"
-#include "Templates/TypeHash.h"
-#include "Containers/SparseArray.h"
-#include "Containers/ScriptArray.h"
-#include "Containers/BitArray.h"
-#include "Algo/Reverse.h"
-#include "Math/IntPoint.h"
-#include "Logging/LogMacros.h"
-#include "Logging/LogCategory.h"
-#include "UObject/NameTypes.h"
-#include "HAL/CriticalSection.h"
-#include "Misc/Timespan.h"
-#include "Containers/StringConv.h"
-#include "UObject/UnrealNames.h"
-#include "Templates/SharedPointer.h"
-#include "CoreGlobals.h"
-#include "HAL/PlatformTLS.h"
-#include "GenericPlatform/GenericPlatformTLS.h"
-#include "Delegates/Delegate.h"
-#include "UObject/WeakObjectPtrTemplates.h"
-#include "Delegates/MulticastDelegateBase.h"
-#include "Delegates/IDelegateInstance.h"
-#include "Delegates/DelegateSettings.h"
-#include "Delegates/DelegateBase.h"
-#include "Delegates/IntegerSequence.h"
-#include "Templates/Tuple.h"
-#include "UObject/ScriptDelegates.h"
-#include "Misc/Optional.h"
-#include "CoreMinimal.h"
-#include "Math/UnrealMath.h"
-#include "Templates/IsFloatingPoint.h"
-#include "Templates/IsIntegral.h"
-#include "Templates/IsSigned.h"
-#include "Templates/Greater.h"
-#include "Math/ColorList.h"
-#include "Math/Color.h"
-#include "Misc/Parse.h"
-#include "Math/IntRect.h"
-#include "Math/Vector2D.h"
-#include "Math/Edge.h"
-#include "Math/Vector.h"
-#include "Misc/ByteSwap.h"
-#include "Internationalization/Text.h"
-#include "Containers/EnumAsByte.h"
-#include "Internationalization/CulturePointer.h"
-#include "Internationalization/TextLocalizationManager.h"
-#include "Templates/UniquePtr.h"
-#include "Templates/IsArray.h"
-#include "Templates/RemoveExtent.h"
-#include "Internationalization/Internationalization.h"
-#include "Templates/UniqueObj.h"
-#include "Math/IntVector.h"
-#include "Math/CapsuleShape.h"
-#include "Math/RangeSet.h"
-#include "Math/Range.h"
-#include "Misc/DateTime.h"
-#include "Math/RangeBound.h"
-#include "Math/Box2D.h"
-#include "Math/BoxSphereBounds.h"
-#include "Math/Sphere.h"
-#include "Math/Box.h"
-#include "Math/OrientedBox.h"
-#include "Math/Interval.h"
-#include "Math/RotationAboutPointMatrix.h"
-#include "Math/Matrix.h"
-#include "Math/Vector4.h"
-#include "Math/Plane.h"
-#include "UObject/ObjectVersion.h"
-#include "Math/Rotator.h"
-#include "Math/VectorRegister.h"
-#include "Math/Axis.h"
-#include "Math/RotationTranslationMatrix.h"
-#include "Math/ScaleRotationTranslationMatrix.h"
-#include "Math/RotationMatrix.h"
-#include "Math/PerspectiveMatrix.h"
-#include "Math/OrthoMatrix.h"
-#include "Math/TranslationMatrix.h"
-#include "Math/QuatRotationTranslationMatrix.h"
-#include "Math/Quat.h"
-#include "Math/InverseRotationMatrix.h"
-#include "Math/ScaleMatrix.h"
-#include "Math/MirrorMatrix.h"
-#include "Math/ClipProjectionMatrix.h"
-#include "Math/InterpCurve.h"
-#include "Math/TwoVectors.h"
-#include "Math/InterpCurvePoint.h"
-#include "Math/CurveEdInterface.h"
-#include "Math/Float16Color.h"
-#include "Math/Float16.h"
-#include "Math/Float32.h"
-#include "Math/Vector2DHalf.h"
-#include "Math/Transform.h"
-#include "Math/ConvexHull2d.h"
-#include "HAL/ThreadSingleton.h"
-#include "HAL/TlsAutoCleanup.h"
-#include "HAL/ThreadSafeCounter.h"
-#include "Misc/EnumClassFlags.h"
-#include "HAL/PlatformTime.h"
-#include "GenericPlatform/GenericPlatformTime.h"
-#include "Stats/Stats.h"
-#include "Containers/LockFreeList.h"
-#include "Misc/NoopCounter.h"
-#include "Containers/ChunkedArray.h"
-#include "Containers/IndirectArray.h"
-#include "Misc/Guid.h"
-#include "Misc/OutputDeviceRedirector.h"
-#include "ProfilingDebugging/ResourceSize.h"
-#include "Misc/Paths.h"
-#include "Math/RandomStream.h"
+#include "Features/IModularFeature.h"
 #include "GenericPlatform/GenericApplication.h"
 #include "GenericPlatform/GenericApplicationMessageHandler.h"
-#include "GenericPlatform/GenericWindowDefinition.h"
 #include "GenericPlatform/GenericWindow.h"
-#include "Misc/Attribute.h"
-#include "Math/TransformCalculus2D.h"
-#include "Math/TransformCalculus.h"
+#include "GenericPlatform/GenericWindowDefinition.h"
 #include "GenericPlatform/ICursor.h"
-#include "Misc/MemStack.h"
-#include "Containers/LockFreeFixedSizeAllocator.h"
-#include "Misc/CoreMisc.h"
-#include "Containers/StaticArray.h"
-#include "Modules/ModuleInterface.h"
-#include "Modules/ModuleManager.h"
-#include "Modules/Boilerplate/ModuleBoilerplate.h"
-#include "HAL/IConsoleManager.h"
-#include "HAL/PlatformProcess.h"
-#include "GenericPlatform/GenericPlatformProcess.h"
-#include "GenericPlatform/GenericPlatformFile.h"
-#include "HAL/FileManager.h"
-#include "Misc/CommandLine.h"
-#include "Serialization/MemoryWriter.h"
-#include "Serialization/MemoryArchive.h"
-#include "Misc/ConfigCacheIni.h"
-#include "Misc/App.h"
-#include "HAL/Event.h"
-#include "Containers/List.h"
-#include "Templates/RefCounting.h"
-#include "Features/IModularFeature.h"
-#include "Async/TaskGraphInterfaces.h"
-#include "Misc/CoreDelegates.h"
-#include "Misc/IQueuedWork.h"
-#include "Misc/OutputDeviceError.h"
-#include "Misc/ObjectThumbnail.h"
-#include "Serialization/CustomVersion.h"
 #include "Logging/TokenizedMessage.h"
+#include "Math/RandomStream.h"
+#include "Math/TransformCalculus.h"
+#include "Math/TransformCalculus2D.h"
+#include "Misc/Attribute.h"
 #include "Misc/EngineVersion.h"
 #include "Misc/MessageDialog.h"
+#include "Misc/ObjectThumbnail.h"
+#include "Misc/OutputDeviceError.h"
+#include "ProfilingDebugging/ResourceSize.h"
+#include "Serialization/CustomVersion.h"
 
 // From Json:
+#include "Policies/JsonPrintPolicy.h"
+#include "Policies/PrettyJsonPrintPolicy.h"
 #include "Serialization/JsonTypes.h"
 #include "Serialization/JsonWriter.h"
-#include "Policies/PrettyJsonPrintPolicy.h"
-#include "Policies/JsonPrintPolicy.h"
 
 // From CoreUObject:
-#include "UObject/ObjectMacros.h"
-#include "UObject/Script.h"
-#include "UObject/Object.h"
-#include "UObject/UObjectBaseUtility.h"
-#include "UObject/UObjectBase.h"
-#include "UObject/UObjectGlobals.h"
-#include "UObject/UObjectArray.h"
-#include "UObject/UObjectMarks.h"
-#include "UObject/Class.h"
-#include "UObject/GarbageCollection.h"
-#include "Serialization/ArchiveUObject.h"
-#include "UObject/CoreNative.h"
-#include "UObject/WeakObjectPtr.h"
-#include "UObject/GCObject.h"
-#include "UObject/Interface.h"
 #include "Misc/PackageName.h"
+#include "Misc/WorldCompositionUtility.h"
+#include "Serialization/SerializedPropertyScope.h"
 #include "Templates/Casts.h"
+#include "UObject/CoreNetTypes.h"
 #include "UObject/LazyObjectPtr.h"
-#include "UObject/SoftObjectPtr.h"
 #include "UObject/Package.h"
-#include "Misc/WorldCompositionUtility.h"
-#include "UObject/UnrealType.h"
-#include "UObject/CoreNetTypes.h"
-#include "UObject/ScriptInterface.h"
 #include "UObject/PropertyTag.h"
-#include "Serialization/SerializedPropertyScope.h"
+#include "UObject/ScriptInterface.h"
+#include "UObject/SoftObjectPtr.h"
+#include "UObject/UnrealType.h"
 
 // From InputCore:
 #include "InputCoreTypes.h"
 
 // From SlateCore:
-#include "Types/SlateEnums.h"
-#include "Styling/SlateColor.h"
-#include "Styling/WidgetStyle.h"
-#include "Layout/Margin.h"
+#include "Animation/CurveHandle.h"
+#include "Animation/CurveSequence.h"
+#include "Application/SlateApplicationBase.h"
+#include "Application/SlateWindowHelper.h"
+#include "Application/ThrottleManager.h"
+#include "Brushes/SlateBorderBrush.h"
+#include "Brushes/SlateBoxBrush.h"
+#include "Brushes/SlateColorBrush.h"
+#include "Brushes/SlateDynamicImageBrush.h"
+#include "Brushes/SlateImageBrush.h"
+#include "Brushes/SlateNoResource.h"
 #include "Fonts/CompositeFont.h"
+#include "Fonts/ShapedTextFwd.h"
 #include "Fonts/SlateFontInfo.h"
-#include "Styling/SlateBrush.h"
-#include "Sound/SlateSound.h"
-#include "Styling/ISlateStyle.h"
-#include "Styling/StyleDefaults.h"
-#include "Brushes/SlateNoResource.h"
+#include "Input/CursorReply.h"
+#include "Input/DragAndDrop.h"
+#include "Input/Events.h"
+#include "Input/NavigationReply.h"
+#include "Input/PopupMethodReply.h"
 #include "Input/Reply.h"
 #include "Input/ReplyBase.h"
-#include "Input/Events.h"
+#include "Layout/ArrangedChildren.h"
+#include "Layout/ArrangedWidget.h"
+#include "Layout/Children.h"
 #include "Layout/Geometry.h"
+#include "Layout/LayoutGeometry.h"
+#include "Layout/Margin.h"
+#include "Layout/PaintGeometry.h"
 #include "Layout/SlateRect.h"
+#include "Layout/Visibility.h"
+#include "Layout/WidgetPath.h"
+#include "Rendering/DrawElements.h"
+#include "Rendering/RenderingCommon.h"
+#include "Rendering/ShaderResourceManager.h"
 #include "Rendering/SlateLayoutTransform.h"
-#include "Layout/PaintGeometry.h"
+#include "Rendering/SlateRenderer.h"
 #include "Rendering/SlateRenderTransform.h"
-#include "Input/DragAndDrop.h"
-#include "Input/CursorReply.h"
-#include "Layout/Visibility.h"
-#include "Input/NavigationReply.h"
-#include "Input/PopupMethodReply.h"
-#include "Widgets/DeclarativeSyntaxSupport.h"
-#include "Types/ISlateMetaData.h"
-#include "Widgets/SNullWidget.h"
-#include "Widgets/SWidget.h"
-#include "Layout/ArrangedWidget.h"
-#include "Types/WidgetActiveTimerDelegate.h"
-#include "Layout/LayoutGeometry.h"
-#include "Widgets/SCompoundWidget.h"
-#include "Layout/Children.h"
+#include "SlateGlobals.h"
 #include "SlotBase.h"
-#include "Styling/SlateWidgetStyle.h"
-#include "Styling/SlateTypes.h"
+#include "Sound/SlateSound.h"
 #include "Styling/CoreStyle.h"
+#include "Styling/ISlateStyle.h"
+#include "Styling/SlateBrush.h"
+#include "Styling/SlateColor.h"
+#include "Styling/SlateStyle.h"
+#include "Styling/SlateTypes.h"
+#include "Styling/SlateWidgetStyle.h"
 #include "Styling/SlateWidgetStyleAsset.h"
 #include "Styling/SlateWidgetStyleContainerBase.h"
 #include "Styling/SlateWidgetStyleContainerInterface.h"
-#include "Widgets/SPanel.h"
+#include "Styling/StyleDefaults.h"
+#include "Styling/WidgetStyle.h"
+#include "Textures/SlateIcon.h"
+#include "Textures/SlateShaderResource.h"
+#include "Textures/SlateTextureData.h"
+#include "Types/ISlateMetaData.h"
+#include "Types/PaintArgs.h"
+#include "Types/SlateConstants.h"
+#include "Types/SlateEnums.h"
 #include "Types/SlateStructs.h"
+#include "Types/WidgetActiveTimerDelegate.h"
+#include "Widgets/DeclarativeSyntaxSupport.h"
+#include "Widgets/IToolTip.h"
 #include "Widgets/SBoxPanel.h"
-#include "Rendering/RenderingCommon.h"
+#include "Widgets/SCompoundWidget.h"
 #include "Widgets/SLeafWidget.h"
-#include "SlateGlobals.h"
-#include "Textures/SlateShaderResource.h"
-#include "Rendering/DrawElements.h"
-#include "Fonts/ShapedTextFwd.h"
-#include "Types/PaintArgs.h"
-#include "Rendering/ShaderResourceManager.h"
-#include "Textures/SlateTextureData.h"
+#include "Widgets/SNullWidget.h"
 #include "Widgets/SOverlay.h"
-#include "Animation/CurveSequence.h"
-#include "Animation/CurveHandle.h"
+#include "Widgets/SPanel.h"
+#include "Widgets/SWidget.h"
 #include "Widgets/SWindow.h"
-#include "Brushes/SlateDynamicImageBrush.h"
-#include "Types/SlateConstants.h"
-#include "Application/SlateWindowHelper.h"
-#include "Layout/ArrangedChildren.h"
-#include "Widgets/IToolTip.h"
-#include "Layout/WidgetPath.h"
-#include "Rendering/SlateRenderer.h"
-#include "Textures/SlateIcon.h"
-#include "Application/SlateApplicationBase.h"
-#include "Application/ThrottleManager.h"
-#include "Brushes/SlateImageBrush.h"
-#include "Styling/SlateStyle.h"
-#include "Brushes/SlateBorderBrush.h"
-#include "Brushes/SlateBoxBrush.h"
-#include "Brushes/SlateColorBrush.h"
 
 // From Slate:
+#include "Framework/Application/IMenu.h"
+#include "Framework/Application/MenuStack.h"
+#include "Framework/Application/SlateApplication.h"
+#include "Framework/Commands/Commands.h"
+#include "Framework/Commands/InputBindingManager.h"
+#include "Framework/Commands/InputChord.h"
+#include "Framework/Commands/UIAction.h"
+#include "Framework/Commands/UICommandInfo.h"
+#include "Framework/Commands/UICommandList.h"
+#include "Framework/Docking/TabManager.h"
+#include "Framework/Docking/WorkspaceItem.h"
+#include "Framework/Layout/InertialScrollManager.h"
+#include "Framework/Layout/IScrollableWidget.h"
+#include "Framework/Layout/Overscroll.h"
+#include "Framework/MultiBox/MultiBoxBuilder.h"
+#include "Framework/MultiBox/MultiBoxDefs.h"
+#include "Framework/MultiBox/MultiBoxExtender.h"
 #include "Framework/SlateDelegates.h"
-#include "Widgets/Layout/SBorder.h"
-#include "Framework/Text/TextLayout.h"
-#include "Framework/Text/TextRunRenderer.h"
-#include "Framework/Text/TextLineHighlight.h"
 #include "Framework/Text/IRun.h"
 #include "Framework/Text/ShapedTextCacheFwd.h"
-#include "Widgets/Text/STextBlock.h"
-#include "Framework/Layout/InertialScrollManager.h"
-#include "Widgets/Layout/SBox.h"
+#include "Framework/Text/TextLayout.h"
+#include "Framework/Text/TextLineHighlight.h"
+#include "Framework/Text/TextRunRenderer.h"
 #include "Framework/Views/ITypedTableView.h"
-#include "Framework/Application/IMenu.h"
-#include "Widgets/Views/STableViewBase.h"
-#include "Framework/Layout/IScrollableWidget.h"
-#include "Widgets/Views/SHeaderRow.h"
-#include "Widgets/Views/STableRow.h"
-#include "Widgets/Views/SExpanderArrow.h"
 #include "Framework/Views/TableViewTypeTraits.h"
-#include "Widgets/Layout/SScrollBar.h"
-#include "Framework/Layout/Overscroll.h"
-#include "Widgets/Views/SListView.h"
-#include "Framework/MultiBox/MultiBoxExtender.h"
-#include "Widgets/Images/SImage.h"
-#include "Framework/Application/SlateApplication.h"
-#include "Framework/Application/MenuStack.h"
+#include "SlateFwd.h"
 #include "SlateOptMacros.h"
+#include "Widgets/Docking/SDockTab.h"
+#include "Widgets/Images/SImage.h"
+#include "Widgets/Input/IVirtualKeyboardEntry.h"
 #include "Widgets/Input/SButton.h"
-#include "SlateFwd.h"
-#include "Widgets/Input/SMenuAnchor.h"
+#include "Widgets/Input/SCheckBox.h"
+#include "Widgets/Input/SComboBox.h"
 #include "Widgets/Input/SComboButton.h"
-#include "Framework/Commands/InputChord.h"
-#include "Framework/Commands/UICommandInfo.h"
-#include "Widgets/Text/ISlateEditableTextWidget.h"
-#include "Widgets/Input/IVirtualKeyboardEntry.h"
 #include "Widgets/Input/SEditableText.h"
-#include "Framework/Commands/UIAction.h"
-#include "Framework/Commands/UICommandList.h"
 #include "Widgets/Input/SEditableTextBox.h"
-#include "Framework/MultiBox/MultiBoxBuilder.h"
-#include "Framework/MultiBox/MultiBoxDefs.h"
-#include "Widgets/Input/SCheckBox.h"
-#include "Framework/Docking/WorkspaceItem.h"
-#include "Framework/Docking/TabManager.h"
-#include "Widgets/Input/SComboBox.h"
-#include "Widgets/Layout/SSplitter.h"
-#include "Framework/Commands/InputBindingManager.h"
-#include "Widgets/Docking/SDockTab.h"
-#include "Widgets/Layout/SGridPanel.h"
-#include "Framework/Commands/Commands.h"
-#include "Widgets/Views/STreeView.h"
+#include "Widgets/Input/SMenuAnchor.h"
+#include "Widgets/Layout/SBorder.h"
+#include "Widgets/Layout/SBox.h"
 #include "Widgets/Layout/SExpandableArea.h"
+#include "Widgets/Layout/SGridPanel.h"
+#include "Widgets/Layout/SScrollBar.h"
 #include "Widgets/Layout/SScrollBox.h"
 #include "Widgets/Layout/SSeparator.h"
 #include "Widgets/Layout/SSpacer.h"
+#include "Widgets/Layout/SSplitter.h"
+#include "Widgets/Text/ISlateEditableTextWidget.h"
+#include "Widgets/Text/STextBlock.h"
+#include "Widgets/Views/SExpanderArrow.h"
+#include "Widgets/Views/SHeaderRow.h"
+#include "Widgets/Views/SListView.h"
+#include "Widgets/Views/STableRow.h"
+#include "Widgets/Views/STableViewBase.h"
+#include "Widgets/Views/STreeView.h"
