Change logging verbosity when navigating to a disabled widget

diff --git a/Engine/Source/Runtime/Slate/Private/Framework/Application/SlateApplication.cpp b/Engine/Source/Runtime/Slate/Private/Framework/Application/SlateApplication.cpp
--- a/Engine/Source/Runtime/Slate/Private/Framework/Application/SlateApplication.cpp
+++ b/Engine/Source/Runtime/Slate/Private/Framework/Application/SlateApplication.cpp
@@ -6016,7 +6016,7 @@ bool FSlateApplication::AttemptNavigation(const FWidgetPath& NavigationSource, c
 				Reason = TEXT("Widget does not support keyboard focus");
 			}
 
-			UE_LOG(LogSlate, Warning, TEXT("Could not Explicitly navigate to widget '%s' because '%s'"), *FReflectionMetaData::GetWidgetDebugInfo(FocusRecipient), Reason);
+			UE_LOG(LogSlate, VeryVerbose, TEXT("Could not Explicitly navigate to widget '%s' because '%s'"), *FReflectionMetaData::GetWidgetDebugInfo(FocusRecipient), Reason);
 		}
 #endif
 	}
