Fix a deadlock on Mac.
The deadlock was triggered because the game thread was locking GCachedMenuStateCS, but then, in order to create a window, the game thread was making a call to the main thread.
Meanwhile, the main thread was processing a call, then tried to lock GCachedMenuStateCS, making the game thread waiting forever to place its call to main thread, and never releasing GCachedMenuStateCS.

The fix split the job done on the game thread, in order to ensure that if we have to call the main thread, GCachedMenuStateCS won't be locked on the game thread at this time.

#tests: 100% repro steps from the jira no longer trigger the deadlock. Played with the various menu, submenu, menu from different subwindows, no regression noticed.

diff --git a/Engine/Source/Runtime/Slate/Private/Framework/MultiBox/Mac/MacMenu.cpp b/Engine/Source/Runtime/Slate/Private/Framework/MultiBox/Mac/MacMenu.cpp
--- a/Engine/Source/Runtime/Slate/Private/Framework/MultiBox/Mac/MacMenu.cpp
+++ b/Engine/Source/Runtime/Slate/Private/Framework/MultiBox/Mac/MacMenu.cpp
@@ -762,76 +762,113 @@ void FSlateMacMenu::UpdateCachedState()
 
 	if (bShouldUpdate)
 	{
-		FScopeLock CachedMenuStateLock(&GCachedMenuStateCS);
+		TMap<TSharedPtr<const FMenuEntryBlock>, TSharedRef<SWidget>> InvalidMenuEntryBlocks;
 
-		for (TMap<FMacMenu*, TSharedPtr<TArray<FMacMenuItemState>>>::TIterator It(GCachedMenuState); It; ++It)
+		//Step 1: Gather list of null widgets, with GCachedMenuStateCS locked.
 		{
-			FMacMenu* Menu = It.Key();
-			TSharedPtr<TArray<FMacMenuItemState>> MenuState = It.Value();
+			FScopeLock CachedMenuStateLock(&GCachedMenuStateCS);
 
-			TSharedRef<SWidget> Widget = SNullWidget::NullWidget;
-			if (!Menu.MultiBox.IsValid())
+			//Ensure we will alloc only a single time, even if most of the time it will be bigger than necessary (but always quite small anyway).
+			InvalidMenuEntryBlocks.Reserve(GCachedMenuState.Num());
+
+			for (TMap<FMacMenu*, TSharedPtr<TArray<FMacMenuItemState>>>::TIterator It(GCachedMenuState); It; ++It)
 			{
-				TSharedPtr<const FMenuEntryBlock> MenuEntryBlock = Menu.MenuEntryBlock.Pin();
-				if (MenuEntryBlock.IsValid())
+				FMacMenu* Menu = It.Key();
+
+				if (!Menu.MultiBox.IsValid())
 				{
-					if (MenuEntryBlock->MenuBuilder.IsBound())
-					{
-						Widget = MenuEntryBlock->MenuBuilder.Execute();
-					}
-					else
+					TSharedPtr<const FMenuEntryBlock> MenuEntryBlock = Menu.MenuEntryBlock.Pin();
+					if (MenuEntryBlock.IsValid())
 					{
-						const bool bShouldCloseWindowAfterMenuSelection = true;
-						FMenuBuilder MenuBuilder(bShouldCloseWindowAfterMenuSelection, MenuEntryBlock->GetActionList(), MenuEntryBlock->Extender);
-						{
-							// Have the menu fill its contents
-							MenuEntryBlock->EntryBuilder.ExecuteIfBound(MenuBuilder);
-						}
- 
-                        // Use INT_MAX for MaxHeight to so that SMultiBoxWidget is returned by MakeWidget instead of SVerticalBox
-                        const int32 MaxHeight = INT_MAX;
-						Widget = MenuBuilder.MakeWidget(/*InMakeMultiBoxBuilderOverride=*/ nullptr, MaxHeight);
+						InvalidMenuEntryBlocks.Add(MenuEntryBlock, SNullWidget::NullWidget);
 					}
+				}
+			}
+		}
+		
+		//Step 2: Create the widgets, which can be a complex process (like when creating a new window that will try to perform a call on the main thread, that may already have locked GCachedMenuStateCS).
+		//So we ensure that GCachedMenuStateCS is not locked for this step.
+		for(TPair<TSharedPtr<const FMenuEntryBlock>, TSharedRef<SWidget>>& It : InvalidMenuEntryBlocks)
+		{
+			TSharedPtr<const FMenuEntryBlock> MenuEntryBlock = It.Key;
+			TSharedRef<SWidget> Widget = SNullWidget::NullWidget;
 
-					if (Widget->GetType() == FName(TEXT("SMultiBoxWidget")))
-					{
-						Menu.MultiBox = TSharedPtr<const FMultiBox>(StaticCastSharedRef<SMultiBoxWidget>(Widget)->GetMultiBox());
-					}
-					else
-					{
-                        const FName ActionName = MenuEntryBlock->GetAction().IsValid() ? MenuEntryBlock->GetAction()->GetCommandName() : NAME_None;
-						UE_LOG(LogMac, Warning, TEXT("Unsupported type of menu widget in FSlateMacMenu::UpdateCachedState(): %s, %s, %s"),
-                               *Widget->GetType().ToString(), *MenuEntryBlock->GetExtensionHook().ToString(), *ActionName.ToString());
-					}
+			if (MenuEntryBlock->MenuBuilder.IsBound())
+			{
+				Widget = MenuEntryBlock->MenuBuilder.Execute();
+			}
+			else
+			{
+				const bool bShouldCloseWindowAfterMenuSelection = true;
+				FMenuBuilder MenuBuilder(bShouldCloseWindowAfterMenuSelection, MenuEntryBlock->GetActionList(), MenuEntryBlock->Extender);
+				{
+					// Have the menu fill its contents
+					MenuEntryBlock->EntryBuilder.ExecuteIfBound(MenuBuilder);
 				}
+ 
+				// Use INT_MAX for MaxHeight to so that SMultiBoxWidget is returned by MakeWidget instead of SVerticalBox
+				const int32 MaxHeight = INT_MAX;
+				Widget = MenuBuilder.MakeWidget(/*InMakeMultiBoxBuilderOverride=*/ nullptr, MaxHeight);
 			}
+			
+			It.Value = Widget;
+		}
+
+		//Step 3: assign the widgets and setup the menu state.
+		{
+			FScopeLock CachedMenuStateLock(&GCachedMenuStateCS);
 
-			if (Menu.MultiBox.IsValid())
+			for (TMap<FMacMenu*, TSharedPtr<TArray<FMacMenuItemState>>>::TIterator It(GCachedMenuState); It; ++It)
 			{
-				const TArray<TSharedRef<const FMultiBlock>>& MenuBlocks = Menu.MultiBox.Pin()->GetBlocks();
-				for (int32 Index = MenuState->Num(); MenuBlocks.Num() > MenuState->Num(); Index++)
+				FMacMenu* Menu = It.Key();
+				TSharedRef<SWidget> Widget = SNullWidget::NullWidget;
+				
+				if (!Menu.MultiBox.IsValid())
 				{
-					MenuState->Add(FMacMenuItemState());
+					TSharedPtr<const FMenuEntryBlock> MenuEntryBlock = Menu.MenuEntryBlock.Pin();
+					TSharedRef<SWidget>* Value = InvalidMenuEntryBlocks.Find(MenuEntryBlock);
+					if (Value!= nullptr)
+					{
+						Widget = *Value;
+						if (Widget->GetType() == FName(TEXT("SMultiBoxWidget")))
+						{
+							Menu.MultiBox = TSharedPtr<const FMultiBox>(StaticCastSharedRef<SMultiBoxWidget>(Widget)->GetMultiBox());
+						}
+						else
+						{
+							const FName ActionName = MenuEntryBlock->GetAction().IsValid() ? MenuEntryBlock->GetAction()->GetCommandName() : NAME_None;
+							UE_LOG(LogMac, Warning, TEXT("Unsupported type of menu widget in FSlateMacMenu::UpdateCachedState(): %s, %s, %s"),
+								   *Widget->GetType().ToString(), *MenuEntryBlock->GetExtensionHook().ToString(), *ActionName.ToString());
+						}
+					}
 				}
-				for (int32 Index = 0; Index < MenuBlocks.Num(); Index++)
+
+				if (Menu.MultiBox.IsValid())
 				{
-					FMacMenuItemState& ItemState = (*MenuState)[Index];
-					ItemState.Type = MenuBlocks[Index]->GetType();
+					TSharedPtr<TArray<FMacMenuItemState>> MenuState = It.Value();
+					const TArray<TSharedRef<const FMultiBlock>>& MenuBlocks = Menu.MultiBox.Pin()->GetBlocks();
+					MenuState->AddDefaulted(FMath::Max(0, MenuBlocks.Num() - MenuState->Num()));
 
-					if (ItemState.Type == EMultiBlockType::MenuEntry)
+					for (int32 Index = 0; Index < MenuBlocks.Num(); Index++)
 					{
-						TSharedRef<const FMenuEntryBlock> Block = StaticCastSharedRef<const FMenuEntryBlock>(MenuBlocks[Index]);
-						ItemState.Block = Block;
-						ItemState.Title = [FSlateMacMenu::GetMenuItemTitle(Block) retain];
-						ItemState.KeyEquivalent = [MacMenuHelper::GetMenuItemKeyEquivalent(Block, &ItemState.KeyModifiers) retain];
-						if (!ItemState.Icon)
+						FMacMenuItemState& ItemState = (*MenuState)[Index];
+						ItemState.Type = MenuBlocks[Index]->GetType();
+
+						if (ItemState.Type == EMultiBlockType::MenuEntry)
 						{
-							SCOPED_AUTORELEASE_POOL;
-							ItemState.Icon = [FSlateMacMenu::GetMenuItemIcon(Block) retain];
+							TSharedRef<const FMenuEntryBlock> Block = StaticCastSharedRef<const FMenuEntryBlock>(MenuBlocks[Index]);
+							ItemState.Block = Block;
+							ItemState.Title = [FSlateMacMenu::GetMenuItemTitle(Block) retain];
+							ItemState.KeyEquivalent = [MacMenuHelper::GetMenuItemKeyEquivalent(Block, &ItemState.KeyModifiers) retain];
+							if (!ItemState.Icon)
+							{
+								SCOPED_AUTORELEASE_POOL;
+								ItemState.Icon = [FSlateMacMenu::GetMenuItemIcon(Block) retain];
+							}
+							ItemState.IsSubMenu = Block->bIsSubMenu;
+							ItemState.IsEnabled = FSlateMacMenu::IsMenuItemEnabled(Block);
+							ItemState.State = ItemState.IsSubMenu ? 0 : FSlateMacMenu::GetMenuItemState(Block);
 						}
-						ItemState.IsSubMenu = Block->bIsSubMenu;
-						ItemState.IsEnabled = FSlateMacMenu::IsMenuItemEnabled(Block);
-						ItemState.State = ItemState.IsSubMenu ? 0 : FSlateMacMenu::GetMenuItemState(Block);
 					}
 				}
 			}
