Fix: Null window wouldn't be shown when rendering off screen, leading to widget prepass never occuring. Due to the window not being shown, the bIsWindowVisible bool in PrepassWindowAndChildren (SlateApplication.cpp L1215) would be false, leading to PrepassWindowAndChildren never being recursively called.

diff --git a/Engine/Source/Runtime/Engine/Private/GameEngine.cpp b/Engine/Source/Runtime/Engine/Private/GameEngine.cpp
--- a/Engine/Source/Runtime/Engine/Private/GameEngine.cpp
+++ b/Engine/Source/Runtime/Engine/Private/GameEngine.cpp
@@ -33,6 +33,7 @@
 #include "EngineModule.h"
 #include "Misc/PackageName.h"
 #include "HAL/PlatformApplicationMisc.h"
+#include "Null/NullPlatformApplicationMisc.h"
 #include "ShaderPipelineCache.h"
 
 #include "Misc/ConfigCacheIni.h"
@@ -639,8 +640,8 @@ TSharedRef<SWindow> UGameEngine::CreateGameWindow()
 		Window->SetWindowMode(WindowMode);
 	}
 
-	// No need to show window in off-screen rendering mode as it does not render to screen
-	if (FSlateApplication::Get().IsRenderingOffScreen())
+	// No need to show window if rendering off-screen without the null platform as it does not render to screen
+	if (FSlateApplication::Get().IsRenderingOffScreen() && !FNullPlatformApplicationMisc::IsUsingNullApplication())
 	{
 		FSlateApplicationBase::Get().GetRenderer()->CreateViewport(Window);
 	}
diff --git a/Engine/Source/Runtime/Slate/Private/Framework/Application/SlateApplication.cpp b/Engine/Source/Runtime/Slate/Private/Framework/Application/SlateApplication.cpp
--- a/Engine/Source/Runtime/Slate/Private/Framework/Application/SlateApplication.cpp
+++ b/Engine/Source/Runtime/Slate/Private/Framework/Application/SlateApplication.cpp
@@ -1817,17 +1817,14 @@ TSharedRef<SWindow> FSlateApplication::AddWindow( TSharedRef<SWindow> InSlateWin
 
 TSharedRef< FGenericWindow > FSlateApplication::MakeWindow( TSharedRef<SWindow> InSlateWindow, const bool bShowImmediately )
 {
-	// Windows and Linux support a NullPlatform which is handled in the Windows/LinuxApplication. Other platforms do not.
-#if !(PLATFORM_WINDOWS || PLATFORM_LINUX)
-	// When rendering off-screen don't render to screen, create a dummy generic window
-	if (bRenderOffScreen)
+	// When rendering off-screen without the null platform, don't render to screen. Create a dummy generic window instead
+	if (bRenderOffScreen && !FNullPlatformApplicationMisc::IsUsingNullApplication())
 	{
 		TSharedRef< FGenericWindow > NewWindow = MakeShareable(new FGenericWindow());
 		InSlateWindow->SetNativeWindow(NewWindow);
 		FSlateApplicationBase::Get().GetRenderer()->CreateViewport(InSlateWindow);
 		return NewWindow;
 	}
-#endif
 
 	TSharedPtr<FGenericWindow> NativeParent = nullptr;
 	TSharedPtr<SWindow> ParentWindow = InSlateWindow->GetParentWindow();
