Fixes issue where the Ctrl +A Command wasn't displaying the correct highlights in all EditableText while global invalidation is turn on.

diff --git a/Engine/Source/Runtime/Slate/Private/Widgets/Text/SlateEditableTextLayout.cpp b/Engine/Source/Runtime/Slate/Private/Widgets/Text/SlateEditableTextLayout.cpp
--- a/Engine/Source/Runtime/Slate/Private/Widgets/Text/SlateEditableTextLayout.cpp
+++ b/Engine/Source/Runtime/Slate/Private/Widgets/Text/SlateEditableTextLayout.cpp
@@ -1822,6 +1822,7 @@ void FSlateEditableTextLayout::SelectAllText()
 	const FTextLocation NewCursorPosition = FTextLocation(NumberOfLines - 1, Lines[NumberOfLines - 1].Text->Len());
 	CursorInfo.SetCursorLocationAndCalculateAlignment(*TextLayout, NewCursorPosition);
 	UpdateCursorHighlight();
+	OwnerWidget->GetSlateWidget()->Invalidate(EInvalidateWidget::LayoutAndVolatility);
 }
 
 void FSlateEditableTextLayout::SelectWordAt(const FGeometry& MyGeometry, const UE::Slate::FDeprecateVector2DParameter& ScreenSpacePosition)
