Fixed Tooltips appearing over foreground programs.
It was due to the tick having the ability to open tooltips, even if the mouse wasn't directly on the slate window.
It's a regression that was introduced in an optimization (CL23790605) of Tooltips.

Tests: I was able to repro it quite consistently, by moving another program halfway on top of UE editor and a widget supporting tooltip. Tooltips are no longer displayed on foreground programs, and are still displayed in normal case (tested on a variety of widget, like content browser and its element, outliner, place actors panel, landscape panel, viewport element), both when UE has the focus and when it doesn't.

diff --git a/Engine/Source/Runtime/Slate/Private/Framework/Application/SlateApplication.cpp b/Engine/Source/Runtime/Slate/Private/Framework/Application/SlateApplication.cpp
--- a/Engine/Source/Runtime/Slate/Private/Framework/Application/SlateApplication.cpp
+++ b/Engine/Source/Runtime/Slate/Private/Framework/Application/SlateApplication.cpp
@@ -1567,7 +1567,7 @@ void FSlateApplication::TickPlatform(float DeltaTime)
 	{
 		SCOPE_CYCLE_COUNTER(STAT_SlateApplicationInput);
 
-		const bool bCanSpawnNewTooltip = true;
+		const bool bCanSpawnNewTooltip = PlatformApplication->IsCursorDirectlyOverSlateWindow();
 		ForEachUser([this, bCanSpawnNewTooltip](FSlateUser& User) {
 			User.UpdateCursor();
 			User.UpdateTooltip(MenuStack, bCanSpawnNewTooltip);
