Fixed issue where widget in scroll box could get stuck in it's hover state
Fixed issue where mouse position would not be calculated correctly when dragging a scroll area
See UDN for context: https://epicgames.lightning.force.com/lightning/r/Case/5004z00001jItzyAAC/view

#rn Fixed some issues with mouse drag detection leaving widgets stuck in their hover state

diff --git a/Engine/Source/Runtime/Slate/Private/Framework/Application/SlateApplication.cpp b/Engine/Source/Runtime/Slate/Private/Framework/Application/SlateApplication.cpp
--- a/Engine/Source/Runtime/Slate/Private/Framework/Application/SlateApplication.cpp
+++ b/Engine/Source/Runtime/Slate/Private/Framework/Application/SlateApplication.cpp
@@ -3306,6 +3306,42 @@ void FSlateApplication::ProcessReply( const FWidgetPath& CurrentEventPath, const
 							}
 						}
 					}
+					// Need to handle the case where the mouse has moved onto a new widget before the drag was detected by also calling MouseLeave on the newly hovered widgets
+					if (LastWidgetsUnderCursor.Widgets.Last().Pin() != WidgetsUnderMouse->Widgets.Last().Widget)
+					{
+						for (int32 WidgetIndex = WidgetsUnderMouse->Widgets.Num() - 1; WidgetIndex >= 0; --WidgetIndex)
+						{
+							TSharedPtr<SWidget> WidgetNowUnderCursor = WidgetsUnderMouse->Widgets[WidgetIndex].Widget;
+
+							if (WidgetNowUnderCursor.IsValid())
+							{
+								if (WidgetNowUnderCursor != RequestedMouseCaptor && !LastWidgetsUnderCursor.ContainsWidget(WidgetNowUnderCursor.Get()))
+								{
+									// It's possible for mouse event to be null if we end up here from a keyboard event. If so, we should synthesize an event.
+									if (InMouseEvent)
+									{
+										FPointerEvent TransformedPointerEvent = TransformPointerEvent(*InMouseEvent, WidgetsUnderMouse->GetWindow());
+
+										// Note that the event's pointer position is not translated.
+										WidgetNowUnderCursor->OnMouseLeave(TransformedPointerEvent);
+									}
+									else
+									{
+										const FPointerEvent& SimulatedPointer = FPointerEvent();
+										WidgetNowUnderCursor->OnMouseLeave(SimulatedPointer);
+									}
+#if WITH_SLATE_DEBUGGING
+									FSlateDebugging::BroadcastInputEvent(ESlateDebuggingInputEvent::MouseLeave, InMouseEvent, WidgetNowUnderCursor);
+#endif
+								}
+								else
+								{
+									// Done routing mouse leave
+									break;
+								}
+							}
+						}
+					}
 				}
 			}
 			else
diff --git a/Engine/Source/Runtime/Slate/Private/Widgets/Layout/SScrollBox.cpp b/Engine/Source/Runtime/Slate/Private/Widgets/Layout/SScrollBox.cpp
--- a/Engine/Source/Runtime/Slate/Private/Widgets/Layout/SScrollBox.cpp
+++ b/Engine/Source/Runtime/Slate/Private/Widgets/Layout/SScrollBox.cpp
@@ -843,7 +843,7 @@ FReply SScrollBox::OnMouseMove( const FGeometry& MyGeometry, const FPointerEvent
 				// Check if the mouse has moved.
 				if ( bDidScroll )
 				{
-					SetScrollComponentOnVector(SoftwareCursorPosition, GetScrollComponentFromVector(SoftwareCursorPosition) + GetScrollComponentFromVector(MouseEvent.GetCursorDelta()));
+					SetScrollComponentOnVector(SoftwareCursorPosition, GetScrollComponentFromVector(SoftwareCursorPosition) + ScrollByAmountLocal);
 				}
 
 				return Reply;
