Fixed mismatched copy constructor/assignment operator in FSlateWidgetRun::FWidgetRunInfo.

diff --git a/Engine/Source/Runtime/Slate/Public/Framework/Text/SlateWidgetRun.h b/Engine/Source/Runtime/Slate/Public/Framework/Text/SlateWidgetRun.h
--- a/Engine/Source/Runtime/Slate/Public/Framework/Text/SlateWidgetRun.h
+++ b/Engine/Source/Runtime/Slate/Public/Framework/Text/SlateWidgetRun.h
@@ -34,14 +34,6 @@ public:
 		{
 
 		}
-
-		FWidgetRunInfo( const FWidgetRunInfo& Other )
-			: Widget( Other.Widget )
-			, Baseline( Other.Baseline )
-			, Size( Other.Size )
-		{
-
-		}
 	};
 
 	static TSharedRef< FSlateWidgetRun > Create(const TSharedRef<class FTextLayout>& TextLayout, const FRunInfo& InRunInfo, const TSharedRef< const FString >& InText, const FWidgetRunInfo& InWidgetInfo);
