Add null check for WidgetsUnderMouse

diff --git a/Engine/Source/Runtime/Slate/Private/Framework/Application/SlateApplication.cpp b/Engine/Source/Runtime/Slate/Private/Framework/Application/SlateApplication.cpp
--- a/Engine/Source/Runtime/Slate/Private/Framework/Application/SlateApplication.cpp
+++ b/Engine/Source/Runtime/Slate/Private/Framework/Application/SlateApplication.cpp
@@ -3307,7 +3307,7 @@ void FSlateApplication::ProcessReply( const FWidgetPath& CurrentEventPath, const
 						}
 					}
 					// Need to handle the case where the mouse has moved onto a new widget before the drag was detected by also calling MouseLeave on the newly hovered widgets
-					if (LastWidgetsUnderCursor.Widgets.Last().Pin() != WidgetsUnderMouse->Widgets.Last().Widget)
+					if (WidgetsUnderMouse && WidgetsUnderMouse->IsValid() && LastWidgetsUnderCursor.Widgets.Last().Pin() != WidgetsUnderMouse->Widgets.Last().Widget)
 					{
 						for (int32 WidgetIndex = WidgetsUnderMouse->Widgets.Num() - 1; WidgetIndex >= 0; --WidgetIndex)
 						{
