Fixes hiding category title and modeling mode selection toolbar missing icons.

diff --git a/Engine/Source/Developer/WidgetRegistration/Public/ToolkitBuilder.h b/Engine/Source/Developer/WidgetRegistration/Public/ToolkitBuilder.h
--- a/Engine/Source/Developer/WidgetRegistration/Public/ToolkitBuilder.h
+++ b/Engine/Source/Developer/WidgetRegistration/Public/ToolkitBuilder.h
@@ -359,10 +359,10 @@ private:
 	FToolkitWidgetStyle Style;
 
 	/** If SelectedCategoryTitleVisibility == EVisibility::Visible, the selected
-	 * category title is visible, else it is not displayed  */	
-	EVisibility SelectedCategoryTitleVisibility;
+	 * category title is visible, else it is not displayed. By default the selected category title is Collapsed  */	
+	EVisibility SelectedCategoryTitleVisibility = EVisibility::Collapsed;
 
 	/** If CategoryButtonLabelVisibility == EVisibility::Visible, the category button
-	 * labels are visible, else they are not displayed  */	
+	 * labels are visible, else they are not displayed. By default the selected category button labels are Visible    */	
 	EVisibility CategoryButtonLabelVisibility;
 };
\ No newline at end of file
diff --git a/Engine/Source/Editor/EditorStyle/Private/StarshipStyle.cpp b/Engine/Source/Editor/EditorStyle/Private/StarshipStyle.cpp
--- a/Engine/Source/Editor/EditorStyle/Private/StarshipStyle.cpp
+++ b/Engine/Source/Editor/EditorStyle/Private/StarshipStyle.cpp
@@ -1751,7 +1751,7 @@ void FStarshipEditorStyle::FStyle::SetupGeneralStyles()
 		FVerticalToolBarStyle.ButtonStyle.SetNormalPadding(FMargin(12.0f, 6.0f));
 		FVerticalToolBarStyle.ButtonStyle.SetPressedPadding(FMargin(12.0f, 6.0f));
 		FVerticalToolBarStyle.SetExpandBrush(IMAGE_BRUSH("Icons/toolbar_expand_16x", Icon8x8));
-		FVerticalToolBarStyle.SetIconPaddingWithCollapsedLabel(FMargin(8.0f, 8.0f));
+		FVerticalToolBarStyle.SetIconPadding(FMargin(8.0f, 8.0f));
 		FVerticalToolBarStyle.SetIconPaddingWithVisibleLabel(FMargin(8.0f, 8.0f, 8.0f, 2.0f));
  
 		FCheckBoxStyle CheckBoxStyle = FCheckBoxStyle(FVerticalToolBarStyle.ToggleButton)
diff --git a/Engine/Source/Runtime/Slate/Private/Framework/MultiBox/SToolBarButtonBlock.cpp b/Engine/Source/Runtime/Slate/Private/Framework/MultiBox/SToolBarButtonBlock.cpp
--- a/Engine/Source/Runtime/Slate/Private/Framework/MultiBox/SToolBarButtonBlock.cpp
+++ b/Engine/Source/Runtime/Slate/Private/Framework/MultiBox/SToolBarButtonBlock.cpp
@@ -264,8 +264,8 @@ void SToolBarButtonBlock::BuildMultiBlockWidget(const ISlateStyle* StyleSet, con
 	}
 	else
 	{
-		const FMargin IconPadding = LabelVisibility.Get() == EVisibility::Collapsed ? ToolBarStyle.IconPaddingWithCollapsedLabel :
-						ToolBarStyle.IconPaddingWithVisibleLabel;
+		const FMargin IconPadding = !LabelVisibility.IsSet() ? ToolBarStyle.IconPadding :
+		(LabelVisibility.Get() == EVisibility::Collapsed ? ToolBarStyle.IconPadding : ToolBarStyle.IconPaddingWithVisibleLabel);
 		const TSharedRef ContentVBox = 
 			SNew(SVerticalBox)
 					// Icon image
@@ -282,7 +282,7 @@ void SToolBarButtonBlock::BuildMultiBlockWidget(const ISlateStyle* StyleSet, con
 					.HAlign(HAlign_Center)	// Center the label text horizontally
 					[
 						SNew(STextBlock)
-					.Visibility(LabelVisibility)
+						.Visibility(LabelVisibility)
 						.Text(ActualLabel)
 						.TextStyle(&ToolBarStyle.LabelStyle)	// Smaller font for tool tip labels
 					];
diff --git a/Engine/Source/Runtime/SlateCore/Public/Styling/ToolBarStyle.h b/Engine/Source/Runtime/SlateCore/Public/Styling/ToolBarStyle.h
--- a/Engine/Source/Runtime/SlateCore/Public/Styling/ToolBarStyle.h
+++ b/Engine/Source/Runtime/SlateCore/Public/Styling/ToolBarStyle.h
@@ -130,11 +130,11 @@ struct SLATECORE_API FToolBarStyle : public FSlateWidgetStyle
 	FToolBarStyle& SetShowLabels(bool bInShowLabels) { bShowLabels = bInShowLabels; return *this; }
 	
 	UPROPERTY(EditAnywhere, BlueprintReadWrite, Category=Appearance)
-	float ButtonContentMaxWidth;
+	float ButtonContentMaxWidth = 64.0f;
 	FToolBarStyle& SetButtonContentMaxWidth(float InButtonContentMaxWidth) { ButtonContentMaxWidth = InButtonContentMaxWidth; return *this; }
 
 	UPROPERTY(EditAnywhere, BlueprintReadWrite, Category=Appearance)
-	float ButtonContentFillWidth;
+	float ButtonContentFillWidth = 1.0f;
 	FToolBarStyle& SetButtonContentFillWidth(float InButtonContentFillWidth) { ButtonContentFillWidth = InButtonContentFillWidth; return *this; }
 	
 	UPROPERTY(EditAnywhere, BlueprintReadWrite, Category = Appearance)
