Fixes when some modes (such as brush editing) do not have left padding in the palette buttons.

diff --git a/Engine/Source/Runtime/Slate/Private/Framework/MultiBox/SToolBarButtonBlock.cpp b/Engine/Source/Runtime/Slate/Private/Framework/MultiBox/SToolBarButtonBlock.cpp
--- a/Engine/Source/Runtime/Slate/Private/Framework/MultiBox/SToolBarButtonBlock.cpp
+++ b/Engine/Source/Runtime/Slate/Private/Framework/MultiBox/SToolBarButtonBlock.cpp
@@ -297,7 +297,6 @@ void SToolBarButtonBlock::BuildMultiBlockWidget(const ISlateStyle* StyleSet, con
 			.VAlign(VAlign_Center) :
 			SHorizontalBox::Slot()
 			.FillWidth(ToolBarStyle.ButtonContentFillWidth)
-			.AutoWidth()
 			.VAlign(VAlign_Center))
 			[
 				ContentVBox
