Slate: MacMenu moving calls to MainThread. This bit of code trips up the Main Thread Checker and will be fatal in future versions.

diff --git a/Engine/Source/Runtime/Slate/Private/Framework/MultiBox/Mac/MacMenu.cpp b/Engine/Source/Runtime/Slate/Private/Framework/MultiBox/Mac/MacMenu.cpp
--- a/Engine/Source/Runtime/Slate/Private/Framework/MultiBox/Mac/MacMenu.cpp
+++ b/Engine/Source/Runtime/Slate/Private/Framework/MultiBox/Mac/MacMenu.cpp
@@ -346,83 +346,85 @@ void FSlateMacMenu::PostInitStartup()
 		FMacMenuCommands::Register();
 	
 		// Build Default menu's
-        FCocoaMenu* MenuBar = [[FCocoaMenu new] autorelease];
-        FCocoaMenu* AppMenu = [[FCocoaMenu new] autorelease];
-        NSMenuItem* AppMenuItem = [[NSMenuItem new] autorelease];
-        [AppMenuItem setTitle:@"AppMenuItem"];
-        [MenuBar addItem:AppMenuItem];
-        [AppMenuItem setSubmenu:AppMenu];
-        [NSApp setMainMenu:MenuBar];
-
-        NSString* AppName = MacMenuHelper::ComputeAppName();
+        MainThreadCall(^{
+            FCocoaMenu* MenuBar = [[FCocoaMenu new] autorelease];
+            FCocoaMenu* AppMenu = [[FCocoaMenu new] autorelease];
+            NSMenuItem* AppMenuItem = [[NSMenuItem new] autorelease];
+            [AppMenuItem setTitle:@"AppMenuItem"];
+            [MenuBar addItem:AppMenuItem];
+            [AppMenuItem setSubmenu:AppMenu];
+            [NSApp setMainMenu:MenuBar];
+
+            NSString* AppName = MacMenuHelper::ComputeAppName();
+            
+            NSMenu* MainMenu = [NSApp mainMenu];
         
-        NSMenu* MainMenu = [NSApp mainMenu];
-	
-        NSMenuItem* PreferencesItem = GIsEditor ? MacMenuHelper::CreateNSMenuItemForCommand(FMacMenuCommands::Get().Preferences, @selector(showPreferencesWindow:), MacMenuHelper::CmdID_Preferences) : nil;
-        NSMenuItem* HideItem = MacMenuHelper::CreateNSMenuItemForCommand(FMacMenuCommands::Get().Hide, @selector(hide:), MacMenuHelper::CmdID_Hide, AppName);
-        NSMenuItem* HideOthersItem = MacMenuHelper::CreateNSMenuItemForCommand(FMacMenuCommands::Get().HideOthers, @selector(hideOtherApplications:), MacMenuHelper::CmdID_HideOthers);
-		NSMenuItem* ShowAllItem = MacMenuHelper::CreateNSMenuItemForCommand(FMacMenuCommands::Get().ShowAll, @selector(unhideAllApplications:), MacMenuHelper::CmdID_ShowAll);
-
-		SEL ShowAboutSelector = [[NSApp delegate] respondsToSelector:@selector(showAboutWindow:)] ? @selector(showAboutWindow:) : @selector(orderFrontStandardAboutPanel:);
-		NSMenuItem* AboutItem = MacMenuHelper::CreateNSMenuItemForCommand(FMacMenuCommands::Get().About, ShowAboutSelector, MacMenuHelper::CmdID_About, AppName);
-		
-        SEL RequestQuitSelector = [[NSApp delegate] respondsToSelector:@selector(requestQuit:)] ? @selector(requestQuit:) : @selector(terminate:);
-        NSMenuItem* QuitItem = MacMenuHelper::CreateNSMenuItemForCommand(FMacMenuCommands::Get().Quit, RequestQuitSelector, MacMenuHelper::CmdID_Quit, AppName);
-        
-        NSMenuItem* ServicesItem = [[NSMenuItem new] autorelease];
-        FCocoaMenu* ServicesMenu = [[FCocoaMenu new] autorelease];
-        [ServicesItem setTitle:NSLOCTEXT("MainMenu","ServicesMenu","Services").ToString().GetNSString()];
-        [ServicesItem setSubmenu:ServicesMenu];
-        [ServicesItem setTag:MacMenuHelper::CmdID_ServicesMenu];
-        [NSApp setServicesMenu:ServicesMenu];
-        [AppMenu addItem:AboutItem];
-        [AppMenu addItem:[NSMenuItem separatorItem]];
-        if (PreferencesItem)
-        {
-            [AppMenu addItem:PreferencesItem];
+            NSMenuItem* PreferencesItem = GIsEditor ? MacMenuHelper::CreateNSMenuItemForCommand(FMacMenuCommands::Get().Preferences, @selector(showPreferencesWindow:), MacMenuHelper::CmdID_Preferences) : nil;
+            NSMenuItem* HideItem = MacMenuHelper::CreateNSMenuItemForCommand(FMacMenuCommands::Get().Hide, @selector(hide:), MacMenuHelper::CmdID_Hide, AppName);
+            NSMenuItem* HideOthersItem = MacMenuHelper::CreateNSMenuItemForCommand(FMacMenuCommands::Get().HideOthers, @selector(hideOtherApplications:), MacMenuHelper::CmdID_HideOthers);
+            NSMenuItem* ShowAllItem = MacMenuHelper::CreateNSMenuItemForCommand(FMacMenuCommands::Get().ShowAll, @selector(unhideAllApplications:), MacMenuHelper::CmdID_ShowAll);
+
+            SEL ShowAboutSelector = [[NSApp delegate] respondsToSelector:@selector(showAboutWindow:)] ? @selector(showAboutWindow:) : @selector(orderFrontStandardAboutPanel:);
+            NSMenuItem* AboutItem = MacMenuHelper::CreateNSMenuItemForCommand(FMacMenuCommands::Get().About, ShowAboutSelector, MacMenuHelper::CmdID_About, AppName);
+            
+            SEL RequestQuitSelector = [[NSApp delegate] respondsToSelector:@selector(requestQuit:)] ? @selector(requestQuit:) : @selector(terminate:);
+            NSMenuItem* QuitItem = MacMenuHelper::CreateNSMenuItemForCommand(FMacMenuCommands::Get().Quit, RequestQuitSelector, MacMenuHelper::CmdID_Quit, AppName);
+            
+            NSMenuItem* ServicesItem = [[NSMenuItem new] autorelease];
+            FCocoaMenu* ServicesMenu = [[FCocoaMenu new] autorelease];
+            [ServicesItem setTitle:NSLOCTEXT("MainMenu","ServicesMenu","Services").ToString().GetNSString()];
+            [ServicesItem setSubmenu:ServicesMenu];
+            [ServicesItem setTag:MacMenuHelper::CmdID_ServicesMenu];
+            [NSApp setServicesMenu:ServicesMenu];
+            [AppMenu addItem:AboutItem];
             [AppMenu addItem:[NSMenuItem separatorItem]];
-        }
-        [AppMenu addItem:ServicesItem];
-        [AppMenu addItem:[NSMenuItem separatorItem]];
-        [AppMenu addItem:HideItem];
-        [AppMenu addItem:HideOthersItem];
-        [AppMenu addItem:ShowAllItem];
-        [AppMenu addItem:[NSMenuItem separatorItem]];
-        [AppMenu addItem:QuitItem];
-
-		if (FApp::IsGame())
-		{
-			NSMenu* ViewMenu = [[FCocoaMenu new] autorelease];
-			[ViewMenu setTitle:NSLOCTEXT("MainMenu","ViewMenu","View").ToString().GetNSString()];
-			NSMenuItem* ViewMenuItem = [[NSMenuItem new] autorelease];
-			[ViewMenuItem setSubmenu:ViewMenu];
-			[[NSApp mainMenu] addItem:ViewMenuItem];
-
-			NSMenuItem* ToggleFullscreenItem = MacMenuHelper::CreateNSMenuItemForCommand(FMacMenuCommands::Get().ToggleFullScreen, @selector(toggleFullScreen:), MacMenuHelper::CmdID_ToggleFullScreen);
-			[ViewMenu addItem:ToggleFullscreenItem];
-		}
-		
-        NSMenu* WindowMenu = [NSApp windowsMenu];
-        if (!WindowMenu)
-        {
-            WindowMenu = [[FCocoaMenu new] autorelease];
-            [WindowMenu setTitle:NSLOCTEXT("MainMenu","WindowMenu","Window").ToString().GetNSString()];
-            NSMenuItem* WindowMenuItem = [[NSMenuItem new] autorelease];
-            [WindowMenuItem setSubmenu:WindowMenu];
-            [[NSApp mainMenu] addItem:WindowMenuItem];
-            [NSApp setWindowsMenu:WindowMenu];
-        }
-        
-        NSMenuItem* MinimizeItem = MacMenuHelper::CreateNSMenuItemForCommand(FMacMenuCommands::Get().Minimize, @selector(miniaturize:), MacMenuHelper::CmdID_Minimize);
-        NSMenuItem* ZoomItem = MacMenuHelper::CreateNSMenuItemForCommand(FMacMenuCommands::Get().Zoom, @selector(zoom:), MacMenuHelper::CmdID_Zoom);
-        NSMenuItem* CloseItem = MacMenuHelper::CreateNSMenuItemForCommand(FMacMenuCommands::Get().Close, @selector(performClose:), MacMenuHelper::CmdID_Close);
-        NSMenuItem* BringAllToFrontItem = MacMenuHelper::CreateNSMenuItemForCommand(FMacMenuCommands::Get().BringAllToFront, @selector(arrangeInFront:), MacMenuHelper::CmdID_BringAllToFront);
-        [WindowMenu addItem:MinimizeItem];
-        [WindowMenu addItem:ZoomItem];
-        [WindowMenu addItem:CloseItem];
-        [WindowMenu addItem:[NSMenuItem separatorItem]];
-        [WindowMenu addItem:BringAllToFrontItem];
-        [WindowMenu addItem:[NSMenuItem separatorItem]];
+            if (PreferencesItem)
+            {
+                [AppMenu addItem:PreferencesItem];
+                [AppMenu addItem:[NSMenuItem separatorItem]];
+            }
+            [AppMenu addItem:ServicesItem];
+            [AppMenu addItem:[NSMenuItem separatorItem]];
+            [AppMenu addItem:HideItem];
+            [AppMenu addItem:HideOthersItem];
+            [AppMenu addItem:ShowAllItem];
+            [AppMenu addItem:[NSMenuItem separatorItem]];
+            [AppMenu addItem:QuitItem];
+
+            if (FApp::IsGame())
+            {
+                NSMenu* ViewMenu = [[FCocoaMenu new] autorelease];
+                [ViewMenu setTitle:NSLOCTEXT("MainMenu","ViewMenu","View").ToString().GetNSString()];
+                NSMenuItem* ViewMenuItem = [[NSMenuItem new] autorelease];
+                [ViewMenuItem setSubmenu:ViewMenu];
+                [[NSApp mainMenu] addItem:ViewMenuItem];
+
+                NSMenuItem* ToggleFullscreenItem = MacMenuHelper::CreateNSMenuItemForCommand(FMacMenuCommands::Get().ToggleFullScreen, @selector(toggleFullScreen:), MacMenuHelper::CmdID_ToggleFullScreen);
+                [ViewMenu addItem:ToggleFullscreenItem];
+            }
+            
+            NSMenu* WindowMenu = [NSApp windowsMenu];
+            if (!WindowMenu)
+            {
+                WindowMenu = [[FCocoaMenu new] autorelease];
+                [WindowMenu setTitle:NSLOCTEXT("MainMenu","WindowMenu","Window").ToString().GetNSString()];
+                NSMenuItem* WindowMenuItem = [[NSMenuItem new] autorelease];
+                [WindowMenuItem setSubmenu:WindowMenu];
+                [[NSApp mainMenu] addItem:WindowMenuItem];
+                [NSApp setWindowsMenu:WindowMenu];
+            }
+            
+            NSMenuItem* MinimizeItem = MacMenuHelper::CreateNSMenuItemForCommand(FMacMenuCommands::Get().Minimize, @selector(miniaturize:), MacMenuHelper::CmdID_Minimize);
+            NSMenuItem* ZoomItem = MacMenuHelper::CreateNSMenuItemForCommand(FMacMenuCommands::Get().Zoom, @selector(zoom:), MacMenuHelper::CmdID_Zoom);
+            NSMenuItem* CloseItem = MacMenuHelper::CreateNSMenuItemForCommand(FMacMenuCommands::Get().Close, @selector(performClose:), MacMenuHelper::CmdID_Close);
+            NSMenuItem* BringAllToFrontItem = MacMenuHelper::CreateNSMenuItemForCommand(FMacMenuCommands::Get().BringAllToFront, @selector(arrangeInFront:), MacMenuHelper::CmdID_BringAllToFront);
+            [WindowMenu addItem:MinimizeItem];
+            [WindowMenu addItem:ZoomItem];
+            [WindowMenu addItem:CloseItem];
+            [WindowMenu addItem:[NSMenuItem separatorItem]];
+            [WindowMenu addItem:BringAllToFrontItem];
+            [WindowMenu addItem:[NSMenuItem separatorItem]];
+        }, NSDefaultRunLoopMode, false);
 	}
 }
 
