Add a way to specify a "main" non closable tab for Asset Editors
Make the first tab to be registered the main tab by default

diff --git a/Engine/Source/Editor/UnrealEd/Private/Toolkits/AssetEditorToolkit.cpp b/Engine/Source/Editor/UnrealEd/Private/Toolkits/AssetEditorToolkit.cpp
--- a/Engine/Source/Editor/UnrealEd/Private/Toolkits/AssetEditorToolkit.cpp
+++ b/Engine/Source/Editor/UnrealEd/Private/Toolkits/AssetEditorToolkit.cpp
@@ -190,6 +190,14 @@ void FAssetEditorToolkit::InitAssetEditor( const EToolkitMode::Type Mode, const
 		const TSharedRef<FTabManager> NewTabManager = FGlobalTabmanager::Get()->NewTabManager( NewMajorTab.ToSharedRef() );		
 		NewTabManager->SetOnPersistLayout(FTabManager::FOnPersistLayout::CreateRaw(this, &FAssetEditorToolkit::HandleTabManagerPersistLayout));
 		NewTabManager->SetAllowWindowMenuBar(true);
+		NewTabManager->SetEnforceMainTab(true);
+
+		FName MainTabName = GetMainTabName();
+		if(MainTabName.IsValid())
+		{
+			NewTabManager->SetMainTab(MainTabName);
+		}
+
 		this->TabManager = NewTabManager;
 
 		TArray<TWeakObjectPtr<UObject>> ObjectsToEditWeak;
diff --git a/Engine/Source/Editor/UnrealEd/Public/Toolkits/AssetEditorToolkit.h b/Engine/Source/Editor/UnrealEd/Public/Toolkits/AssetEditorToolkit.h
--- a/Engine/Source/Editor/UnrealEd/Public/Toolkits/AssetEditorToolkit.h
+++ b/Engine/Source/Editor/UnrealEd/Public/Toolkits/AssetEditorToolkit.h
@@ -243,6 +243,11 @@ public:
 	static TSharedPtr<FExtensibilityManager> GetSharedMenuExtensibilityManager();
 	static TSharedPtr<FExtensibilityManager> GetSharedToolBarExtensibilityManager();
 
+	/** Override this to manually specify the tab that cannot be closed for your asset editor, it is the first tab
+	 *  registered otherwise
+	 */
+	virtual FName GetMainTabName() const { return NAME_None; }
+
 	/**
 	 * Allows the caller to set a menu overlay, displayed to the far right of the editor's menu bar
 	 *
diff --git a/Engine/Source/Runtime/Slate/Private/Framework/Docking/TabManager.cpp b/Engine/Source/Runtime/Slate/Private/Framework/Docking/TabManager.cpp
--- a/Engine/Source/Runtime/Slate/Private/Framework/Docking/TabManager.cpp
+++ b/Engine/Source/Runtime/Slate/Private/Framework/Docking/TabManager.cpp
@@ -707,7 +707,7 @@ void FTabManager::FPrivateApi::OnTabManagerClosing()
 
 bool FTabManager::FPrivateApi::CanTabLeaveTabWell(const TSharedRef<const SDockTab>& TabToTest) const
 {
-	return TabManager.bCanDoDragOperation && TabToTest != TabManager.MainNonCloseableTab.Pin();
+	return TabManager.bCanDoDragOperation && !(TabToTest->GetLayoutIdentifier() == TabManager.MainNonCloseableTabID);
 }
 
 const TArray< TWeakPtr<SDockingArea> >& FTabManager::FPrivateApi::GetLiveDockAreas() const
@@ -784,7 +784,7 @@ void FTabManager::UpdateMainMenu(TSharedPtr<SDockTab> ForTab, const bool bForce)
 	{
 		ParentWindowOfOwningTab = OwnerTabPinned->GetParentWindow();
 	}
-	else if (auto MainNonCloseableTabPinned = MainNonCloseableTab.Pin())
+	else if (auto MainNonCloseableTabPinned = FindExistingLiveTab(MainNonCloseableTabID))
 	{
 		ParentWindowOfOwningTab = MainNonCloseableTabPinned->GetParentWindow();
 	}
@@ -833,14 +833,39 @@ void FTabManager::UpdateMainMenu(TSharedPtr<SDockTab> ForTab, const bool bForce)
 #endif
 }
 
+void FTabManager::SetMainTab(const FTabId& InMainTabID)
+{
+	MainNonCloseableTabID = InMainTabID;
+}
+
 void FTabManager::SetMainTab(const TSharedRef<const SDockTab>& InTab)
 {
-	MainNonCloseableTab = InTab;
+	if(!InTab->GetLayoutIdentifier().TabType.IsNone())
+	{
+		SetMainTab(InTab->GetLayoutIdentifier());
+	}
+	else
+	{
+		PendingMainNonClosableTab = InTab;
+	}
+	
+}
+
+void FTabManager::SetEnforceMainTab(bool bInEnforceMainTab)
+{
+	bEnforceMainTab = bInEnforceMainTab;
+
+	// Set the main tab to whichever first tab we can find if we don't have one
+	if(bEnforceMainTab && !MainNonCloseableTabID.TabType.IsNone())
+	{
+		FTabSpawner::TIterator SpawnerIterator(TabSpawner);
+		MainNonCloseableTabID = SpawnerIterator.Key();
+	}
 }
 
 bool FTabManager::IsTabCloseable(const TSharedRef<const SDockTab>& InTab) const
 {
-	return !(MainNonCloseableTab.Pin() == InTab);
+	return !(MainNonCloseableTabID == InTab->GetLayoutIdentifier());
 }
 
 const TSharedRef<FWorkspaceItem> FTabManager::GetLocalWorkspaceMenuRoot() const
@@ -994,6 +1019,13 @@ FTabSpawnerEntry& FTabManager::RegisterTabSpawner(const FName TabId, const FOnSp
 
 	TSharedRef<FTabSpawnerEntry> NewSpawnerEntry = MakeShareable(new FTabSpawnerEntry(TabId, OnSpawnTab, CanSpawnTab));
 	TabSpawner.Add(TabId, NewSpawnerEntry);
+
+	// If this tab manager always wants a non closeable tab, and we don't already have one set this tab to be the main tab
+	if(bEnforceMainTab && MainNonCloseableTabID.TabType.IsNone() && !PendingMainNonClosableTab)
+	{
+		MainNonCloseableTabID = TabId;
+	}
+	
 	return NewSpawnerEntry.Get();
 }
 
@@ -1988,6 +2020,13 @@ TSharedPtr<SDockTab> FTabManager::SpawnTab(const FTabId& TabId, const TSharedPtr
 		if (bSpawningAllowedBySpawner && (!Spawner->SpawnedTabPtr.IsValid() || Spawner->OnFindTabToReuse.IsBound()))
 		{
 			NewTabWidget = Spawner->OnSpawnTab.Execute(FSpawnTabArgs(ParentWindow, TabId));
+
+			if(PendingMainNonClosableTab && NewTabWidget == PendingMainNonClosableTab)
+			{
+				PendingMainNonClosableTab = nullptr;
+				MainNonCloseableTabID = TabId;
+			}
+			
 			NewTabWidget->SetLayoutIdentifier(TabId);
 			NewTabWidget->ProvideDefaultLabel(Spawner->GetDisplayName().IsEmpty() ? FText::FromName(Spawner->TabType) : Spawner->GetDisplayName());
 			NewTabWidget->ProvideDefaultIcon(Spawner->GetIcon().GetIcon());
@@ -2889,7 +2928,7 @@ void FProxyTabmanager::OpenUnmanagedTab(FName PlaceholderId, const FSearchPrefer
 				NewlyOpenedTab->GetParent()->GetParentDockTabStack()->OpenTab(UnmanagedTab);
 				NewlyOpenedTab->RequestCloseTab();
 
-				MainNonCloseableTab = UnmanagedTab;
+				MainNonCloseableTabID = UnmanagedTab->GetLayoutIdentifier();
 
 				OnTabOpened.Broadcast(UnmanagedTab);
 			}
diff --git a/Engine/Source/Runtime/Slate/Public/Framework/Docking/TabManager.h b/Engine/Source/Runtime/Slate/Public/Framework/Docking/TabManager.h
--- a/Engine/Source/Runtime/Slate/Public/Framework/Docking/TabManager.h
+++ b/Engine/Source/Runtime/Slate/Public/Framework/Docking/TabManager.h
@@ -935,9 +935,20 @@ class SLATE_API FTabManager : public TSharedFromThis<FTabManager>
 		 */
 		void UpdateMainMenu(TSharedPtr<SDockTab> ForTab, const bool bForce);
 
-		/** Provide a tab that will be the main tab and cannot be closed. */
+		/** Provide a tab that will be the main tab and cannot be closed.
+		 *  Prefer using the version that takes in a TabID so you can set the MainTab before it is actually created
+		 */
 		void SetMainTab(const TSharedRef<const SDockTab>& InTab);
 
+		/** Provide a tab that will be the main tab and cannot be closed. */
+		void SetMainTab(const FTabId& InMainTabID);
+
+		/** If true, this tab manager will always have one main tab that cannot be closed
+		 *  Automatically set to the first tab registered if there are no tabs registered already (set on creation)
+		 *  or the first tab found otherwise
+		 */
+		void SetEnforceMainTab(bool bInEnforceMainTab);
+
 		/* Prevent or allow all tabs to be drag */
 		void SetCanDoDragOperation(bool CanDoDragOperation) { bCanDoDragOperation = CanDoDragOperation; }
 		
@@ -1146,9 +1157,14 @@ class SLATE_API FTabManager : public TSharedFromThis<FTabManager>
 		/** Returns the default window size for the TabId, or the fallback window size if it wasn't registered */
 		static FVector2D GetDefaultTabWindowSize(const FTabId& TabId);
 
+		/** The ID of the main tab, which cannot be closed*/
+		FTabId MainNonCloseableTabID;
+
+		/** To keep track of any pending tabs that want to be the main tab, but don't have their ID set yet */
+		TSharedPtr<const SDockTab> PendingMainNonClosableTab;
 
-		/** The main tab, this tab cannot be closed. */
-		TWeakPtr<const SDockTab> MainNonCloseableTab;
+		/** If true, this tab manager will always have one non-closable tab */
+		bool bEnforceMainTab = false;
 
 		/** The last window we docked a nomad or major tab into */
 		TWeakPtr<SWindow> LastMajorDockWindow;
