Fix an issue causing main tabs to not be set properly for asset editors
Change the main tab to be the viewport for the Skeletal Mesh Editor

diff --git a/Engine/Source/Editor/SkeletalMeshEditor/Private/SkeletalMeshEditor.cpp b/Engine/Source/Editor/SkeletalMeshEditor/Private/SkeletalMeshEditor.cpp
--- a/Engine/Source/Editor/SkeletalMeshEditor/Private/SkeletalMeshEditor.cpp
+++ b/Engine/Source/Editor/SkeletalMeshEditor/Private/SkeletalMeshEditor.cpp
@@ -212,6 +212,11 @@ bool FSkeletalMeshEditor::OnRequestClose()
 	return bAllowClose;
 }
 
+FName FSkeletalMeshEditor::GetMainTabName() const
+{
+	return SkeletalMeshEditorTabs::ViewportTab;
+}
+
 void FSkeletalMeshEditor::RegisterTabSpawners(const TSharedRef<class FTabManager>& InTabManager)
 {
 	WorkspaceMenuCategory = InTabManager->AddLocalWorkspaceMenuCategory(LOCTEXT("WorkspaceMenu_SkeletalMeshEditor", "Skeletal Mesh Editor"));
diff --git a/Engine/Source/Editor/SkeletalMeshEditor/Private/SkeletalMeshEditor.h b/Engine/Source/Editor/SkeletalMeshEditor/Private/SkeletalMeshEditor.h
--- a/Engine/Source/Editor/SkeletalMeshEditor/Private/SkeletalMeshEditor.h
+++ b/Engine/Source/Editor/SkeletalMeshEditor/Private/SkeletalMeshEditor.h
@@ -83,6 +83,7 @@ public:
 	
 	//~ Begin FAssetEditorToolkit Interface.
 	virtual bool OnRequestClose() override;
+	virtual FName GetMainTabName() const override;
 	//~ End FAssetEditorToolkit Interface.
 
 	/** FEditorUndoClient interface */
diff --git a/Engine/Source/Editor/UnrealEd/Private/Toolkits/AssetEditorToolkit.cpp b/Engine/Source/Editor/UnrealEd/Private/Toolkits/AssetEditorToolkit.cpp
--- a/Engine/Source/Editor/UnrealEd/Private/Toolkits/AssetEditorToolkit.cpp
+++ b/Engine/Source/Editor/UnrealEd/Private/Toolkits/AssetEditorToolkit.cpp
@@ -193,7 +193,7 @@ void FAssetEditorToolkit::InitAssetEditor( const EToolkitMode::Type Mode, const
 		NewTabManager->SetEnforceMainTab(true);
 
 		FName MainTabName = GetMainTabName();
-		if(MainTabName.IsValid())
+		if(!MainTabName.IsNone())
 		{
 			NewTabManager->SetMainTab(MainTabName);
 		}
diff --git a/Engine/Source/Runtime/Slate/Private/Framework/Docking/TabManager.cpp b/Engine/Source/Runtime/Slate/Private/Framework/Docking/TabManager.cpp
--- a/Engine/Source/Runtime/Slate/Private/Framework/Docking/TabManager.cpp
+++ b/Engine/Source/Runtime/Slate/Private/Framework/Docking/TabManager.cpp
@@ -855,8 +855,9 @@ void FTabManager::SetEnforceMainTab(bool bInEnforceMainTab)
 {
 	bEnforceMainTab = bInEnforceMainTab;
 
-	// Set the main tab to whichever first tab we can find if we don't have one
-	if(bEnforceMainTab && !MainNonCloseableTabID.TabType.IsNone())
+	static const FTabId InvalidTab;
+	// If we already have tabs registered, set the main tab to whichever first tab we can find if we don't have one
+	if(bEnforceMainTab && MainNonCloseableTabID == InvalidTab && !TabSpawner.IsEmpty())
 	{
 		FTabSpawner::TIterator SpawnerIterator(TabSpawner);
 		MainNonCloseableTabID = SpawnerIterator.Key();
@@ -1020,8 +1021,9 @@ FTabSpawnerEntry& FTabManager::RegisterTabSpawner(const FName TabId, const FOnSp
 	TSharedRef<FTabSpawnerEntry> NewSpawnerEntry = MakeShareable(new FTabSpawnerEntry(TabId, OnSpawnTab, CanSpawnTab));
 	TabSpawner.Add(TabId, NewSpawnerEntry);
 
+	static const FTabId InvalidTab;
 	// If this tab manager always wants a non closeable tab, and we don't already have one set this tab to be the main tab
-	if(bEnforceMainTab && MainNonCloseableTabID.TabType.IsNone() && !PendingMainNonClosableTab)
+	if(bEnforceMainTab && MainNonCloseableTabID == InvalidTab && !PendingMainNonClosableTab)
 	{
 		MainNonCloseableTabID = TabId;
 	}
