A fix for UI screenshot test failure when there is animation: set a constant tick time for animation sequencer and the UI screenshot test.
- tested with slow and fast moving animation, returns 0 diff locally.

diff --git a/Engine/Source/Developer/FunctionalTesting/Classes/FunctionalUIScreenshotTest.h b/Engine/Source/Developer/FunctionalTesting/Classes/FunctionalUIScreenshotTest.h
--- a/Engine/Source/Developer/FunctionalTesting/Classes/FunctionalUIScreenshotTest.h
+++ b/Engine/Source/Developer/FunctionalTesting/Classes/FunctionalUIScreenshotTest.h
@@ -39,7 +39,11 @@ protected:
 	virtual void OnScreenshotTakenAndCompared() override;
 
 	virtual void RequestScreenshot() override;
-	
+
+	virtual bool IsReady_Implementation() override;
+
+	virtual void Tick(float DeltaSeconds) override;
+
 protected:
 	UPROPERTY(EditAnywhere, Category = "UI")
 	TSubclassOf<UUserWidget> WidgetClass;
@@ -58,4 +62,8 @@ protected:
 
 private:
 	TOptional<bool> PreviousDebugCanvasVisible;
+	int32 NumTickPassed;
+	bool bWasPreviouslyUsingFixedDeltaTime;
+	double PreviousFixedDeltaTime;
+	const double TestFixedDeltaTime = 1 / 60.0;
 };
diff --git a/Engine/Source/Developer/FunctionalTesting/Private/FunctionalUIScreenshotTest.cpp b/Engine/Source/Developer/FunctionalTesting/Private/FunctionalUIScreenshotTest.cpp
--- a/Engine/Source/Developer/FunctionalTesting/Private/FunctionalUIScreenshotTest.cpp
+++ b/Engine/Source/Developer/FunctionalTesting/Private/FunctionalUIScreenshotTest.cpp
@@ -109,6 +109,14 @@ void AFunctionalUIScreenshotTest::PrepareTest()
 
 		SpawnedWidget->SetVisibility(ESlateVisibility::HitTestInvisible);
 	}
+	NumTickPassed = 0;
+	if (IConsoleVariable* CVarFixedDeltaTime = IConsoleManager::Get().FindConsoleVariable(TEXT("Slate.UseFixedDeltaTime")))
+	{
+		bWasPreviouslyUsingFixedDeltaTime = CVarFixedDeltaTime->GetBool();
+		PreviousFixedDeltaTime = FSlateApplication::GetFixedDeltaTime();
+		FSlateApplication::SetFixedDeltaTime(TestFixedDeltaTime);
+		CVarFixedDeltaTime->SetWithCurrentPriority(true);
+	}
 
 	UAutomationBlueprintFunctionLibrary::FinishLoadingBeforeScreenshot();
 }
@@ -147,6 +155,13 @@ void AFunctionalUIScreenshotTest::EndPlay(const EEndPlayReason::Type EndPlayReas
 			PreviousDebugCanvasVisible.Reset();
 		}
 	}
+
+	if (IConsoleVariable* CVarFixedDeltaTime = IConsoleManager::Get().FindConsoleVariable(TEXT("Slate.UseFixedDeltaTime")))
+	{
+		CVarFixedDeltaTime->SetWithCurrentPriority(bWasPreviouslyUsingFixedDeltaTime);
+		FSlateApplication::SetFixedDeltaTime(PreviousFixedDeltaTime);
+	}
+	NumTickPassed = 0;
 }
 
 /**
@@ -189,6 +204,21 @@ void ReadPixelsFromRT(UTextureRenderTarget2D* InRT, TArray<FColor>* OutPixels)
 	FlushRenderingCommands();
 }
 
+bool  AFunctionalUIScreenshotTest::IsReady_Implementation()
+{
+	if (NumTickPassed * FSlateApplication::GetFixedDeltaTime() >= ScreenshotOptions.Delay)
+	{
+		return NumTickPassed > ScreenshotOptions.FrameDelay;
+	}
+	return false;
+}
+
+void  AFunctionalUIScreenshotTest::Tick(float DeltaSeconds)
+{
+	NumTickPassed += 1;
+	Super::Tick(DeltaSeconds);
+}
+
 void AFunctionalUIScreenshotTest::RequestScreenshot()
 {
 	// Register a handler to UGameViewportClient::OnScreenshotCaptured
diff --git a/Engine/Source/Runtime/Slate/Private/Framework/Application/SlateApplication.cpp b/Engine/Source/Runtime/Slate/Private/Framework/Application/SlateApplication.cpp
--- a/Engine/Source/Runtime/Slate/Private/Framework/Application/SlateApplication.cpp
+++ b/Engine/Source/Runtime/Slate/Private/Framework/Application/SlateApplication.cpp
@@ -111,6 +111,12 @@ static FAutoConsoleVariableRef CVarSlateEnableGamepadEditorNavigation(
 	TEXT("True implies we allow gamepad navigation outside of the game viewport.")
 );
 
+static bool GSlateUseFixedDeltaTime = false;
+static FAutoConsoleVariableRef CVarSlateUseFixedDeltaTime(
+	TEXT("Slate.UseFixedDeltaTime"),
+	GSlateUseFixedDeltaTime,
+	TEXT("True means we use a constant delta time on every widget tick.")
+);
 //////////////////////////////////////////////////////////////////////////
 
 /** 
@@ -777,6 +783,7 @@ void FSlateApplication::Shutdown(bool bShutdownPlatform)
 }
 
 TSharedPtr<FSlateApplication> FSlateApplication::CurrentApplication = nullptr;
+double FSlateApplication::FixedDeltaTime = 1 / 60.0;
 
 FSlateApplication::FSlateApplication()
 	: bAppIsActive(true)
@@ -1504,7 +1511,7 @@ void FSlateApplication::Tick(ESlateTickType TickType)
 	{
 		SCOPE_CYCLE_COUNTER(STAT_SlateTickTime);
 
-		const float DeltaTime = GetDeltaTime();
+		float DeltaTime = GetDeltaTime();
 
 		// IMPORTANT
 		// Do not add code to these different if-statements, if you need to add additional logic to
@@ -1522,6 +1529,12 @@ void FSlateApplication::Tick(ESlateTickType TickType)
 			TickTime();
 		}
 
+		if (GSlateUseFixedDeltaTime)
+		{
+			DeltaTime = GetFixedDeltaTime();
+
+		}
+
 		if (EnumHasAnyFlags(TickType, ESlateTickType::Widgets))
 		{
 			TickAndDrawWidgets(DeltaTime);
@@ -4300,6 +4313,10 @@ void FSlateApplication::ForEachUser(TFunctionRef<void(FSlateUser*)> InPredicate,
 		}, bIncludeVirtualUsers);
 }
 
+void FSlateApplication::SetFixedDeltaTime(double InSeconds)
+{
+	FixedDeltaTime = InSeconds;
+}
 /* FSlateApplicationBase interface
  *****************************************************************************/
 
diff --git a/Engine/Source/Runtime/Slate/Public/Framework/Application/SlateApplication.h b/Engine/Source/Runtime/Slate/Public/Framework/Application/SlateApplication.h
--- a/Engine/Source/Runtime/Slate/Public/Framework/Application/SlateApplication.h
+++ b/Engine/Source/Runtime/Slate/Public/Framework/Application/SlateApplication.h
@@ -1038,6 +1038,21 @@ public:
 	UE_DEPRECATED(4.23, "ForEachUser now provides an FSlateUser& parameter to the lambda instead of an FSlateUser*")
 	void ForEachUser(TFunctionRef<void(FSlateUser*)> InPredicate, bool bIncludeVirtualUsers = false);
 
+	/**
+	 * Gets time step in seconds if a fixed delta time is wanted.
+	 * @return Time step in seconds for fixed delta time.
+	 */
+	FORCEINLINE static double GetFixedDeltaTime()
+	{
+		return FixedDeltaTime;
+	}
+
+	/**
+	 * Sets time step in seconds if a fixed delta time is wanted.
+	 * @param seconds Time step in seconds for fixed delta time.
+	 */
+	static void SetFixedDeltaTime(double InSeconds);
+
 protected:
 	/**
 	 * Register a new user with Slate.  Normally this is unnecessary as Slate automatically adds
@@ -1993,7 +2008,9 @@ private:
 
 	/** Platform mouse movement event count. */
 	uint64 PlatformMouseMovementEvents = 0;
-		
+
+	/** Constant delta time used on every Slate widget Tick if CVarSlateUseFixedDeltaTime is enabled. */
+	static double FixedDeltaTime;
 	/**
 	 * A helper class to wrap the list of input pre-processors. 
 	 */
