-Adding support for text overflow in SHeaderRow.
-Use it in PCG graph list view.

#tests: in the editor, there are no change overall in the headers, but in the ones where the new overflow option is set to ellipsis.

diff --git a/Engine/Plugins/Experimental/PCG/Source/PCGEditor/Private/SPCGEditorGraphAttributeListView.cpp b/Engine/Plugins/Experimental/PCG/Source/PCGEditor/Private/SPCGEditorGraphAttributeListView.cpp
--- a/Engine/Plugins/Experimental/PCG/Source/PCGEditor/Private/SPCGEditorGraphAttributeListView.cpp
+++ b/Engine/Plugins/Experimental/PCG/Source/PCGEditor/Private/SPCGEditorGraphAttributeListView.cpp
@@ -934,6 +934,7 @@ void SPCGEditorGraphAttributeListView::AddColumn(const UPCGPointData* InPCGPoint
 	Arguments.HAlignCell(CellHAlign);
 	Arguments.SortMode(this, &SPCGEditorGraphAttributeListView::GetColumnSortMode, InColumnId);
 	Arguments.OnSort(this, &SPCGEditorGraphAttributeListView::OnColumnSortModeChanged);
+	Arguments.OverflowPolicy(ETextOverflowPolicy::Ellipsis);
 
 	SHeaderRow::FColumn* NewColumn = new SHeaderRow::FColumn(Arguments);
 	NewColumn->bIsVisible = !HiddenAttributes.Contains(InColumnId); 
diff --git a/Engine/Source/Runtime/Slate/Private/Widgets/Views/SHeaderRow.cpp b/Engine/Source/Runtime/Slate/Private/Widgets/Views/SHeaderRow.cpp
--- a/Engine/Source/Runtime/Slate/Private/Widgets/Views/SHeaderRow.cpp
+++ b/Engine/Source/Runtime/Slate/Private/Widgets/Views/SHeaderRow.cpp
@@ -92,6 +92,7 @@ public:
 					SNew( STextBlock )
 					.TextStyle( FAppStyle::Get(), "NormalText" )
 					.Text( LabelText )
+					.OverflowPolicy(Column.OverflowPolicy)
 				];
 		}
 
diff --git a/Engine/Source/Runtime/Slate/Public/Widgets/Views/SHeaderRow.h b/Engine/Source/Runtime/Slate/Public/Widgets/Views/SHeaderRow.h
--- a/Engine/Source/Runtime/Slate/Public/Widgets/Views/SHeaderRow.h
+++ b/Engine/Source/Runtime/Slate/Public/Widgets/Views/SHeaderRow.h
@@ -121,6 +121,7 @@ public:
 			, _InitialSortMode( EColumnSortMode::Ascending )
 			, _SortMode( EColumnSortMode::None )
 			, _OnSort()
+			, _OverflowPolicy( ETextOverflowPolicy::Clip)
 			{}
 			SLATE_ARGUMENT( FName, ColumnId )
 			SLATE_ATTRIBUTE( FText, DefaultLabel )
@@ -155,6 +156,7 @@ public:
 			SLATE_EVENT( FOnSortModeChanged, OnSort )
 
 			SLATE_ATTRIBUTE(bool, ShouldGenerateWidget)
+			SLATE_ARGUMENT(ETextOverflowPolicy, OverflowPolicy)
 		SLATE_END_ARGS()
 
 		FColumn( const FArguments& InArgs )
@@ -175,13 +177,14 @@ public:
 			, HeaderComboVisibility (InArgs._HeaderComboVisibility )
 			, CellHAlignment( InArgs._HAlignCell )
 			, CellVAlignment( InArgs._VAlignCell )
+			, OverflowPolicy(InArgs._OverflowPolicy)
 			, InitialSortMode( InArgs._InitialSortMode )
 			, SortMode( InArgs._SortMode )
 			, SortPriority( InArgs._SortPriority )
 			, OnSortModeChanged( InArgs._OnSort )
 			, ShouldGenerateWidget(InArgs._ShouldGenerateWidget)
 			, bIsVisible(true)
-		{
+			{
 			if ( InArgs._FixedWidth.IsSet() )
 			{
 				Width = InArgs._FixedWidth.GetValue();
@@ -260,6 +263,8 @@ public:
 		EHorizontalAlignment CellHAlignment;
 		EVerticalAlignment CellVAlignment;
 
+		ETextOverflowPolicy OverflowPolicy;
+
 		TAttribute< EColumnSortMode::Type > InitialSortMode;
 		TAttribute< EColumnSortMode::Type > SortMode;
 		TAttribute< EColumnSortPriority::Type > SortPriority;
