Fix to prevent FSlateUser to spawn tooltip when we are not using the mouse.

Tests:
- played with it in editor, both UE Main and other game branch, no change noticed: editor tool tip are still displayed.
- Associating a tooltip to a button then navigating with the gamepad no longer display tooltips, while using the mouse still do.
- Tested on project who find the bug, issue is fixed.
- tested with and without Common UI activated.

diff --git a/Engine/Source/Runtime/Slate/Private/Framework/Application/SlateUser.cpp b/Engine/Source/Runtime/Slate/Private/Framework/Application/SlateUser.cpp
--- a/Engine/Source/Runtime/Slate/Private/Framework/Application/SlateUser.cpp
+++ b/Engine/Source/Runtime/Slate/Private/Framework/Application/SlateUser.cpp
@@ -1195,7 +1195,7 @@ void FSlateUser::UpdateTooltip(const FMenuStack& MenuStack, bool bCanSpawnNewToo
 	SCOPE_CYCLE_COUNTER(STAT_SlateUpdateTooltip);
 
 	const double MotionLessDurationBeforeAllowingNewToolTip = 0.05;
-	bCanSpawnNewTooltip = bCanSpawnNewTooltip && (FPlatformTime::Seconds() - LastCursorSignificantMoveTime > MotionLessDurationBeforeAllowingNewToolTip);
+	bCanSpawnNewTooltip = bCanSpawnNewTooltip && bCanDrawCursor && (FPlatformTime::Seconds() - LastCursorSignificantMoveTime > MotionLessDurationBeforeAllowingNewToolTip);
 
 	float DPIScaleFactor = 1.0f; //todo: this value is never changed, we should investigate if it is necessary or not to handle it for the force field.
 	FWidgetPath WidgetsToQueryForTooltip;
