Fallback to creating a Window instead of a modal Window if we can't find the Slate Virtual Window.

diff --git a/Engine/Source/Runtime/Slate/Private/Framework/Application/SlateApplication.cpp b/Engine/Source/Runtime/Slate/Private/Framework/Application/SlateApplication.cpp
--- a/Engine/Source/Runtime/Slate/Private/Framework/Application/SlateApplication.cpp
+++ b/Engine/Source/Runtime/Slate/Private/Framework/Application/SlateApplication.cpp
@@ -2030,8 +2030,15 @@ void FSlateApplication::AddModalWindow( TSharedRef<SWindow> InSlateWindow, const
 	{
 		// Find the window of the parent widget
 		FWidgetPath WidgetPath;
-		GeneratePathToWidgetChecked( InParentWidget.ToSharedRef(), WidgetPath );
-		AddWindowAsNativeChild( InSlateWindow, WidgetPath.GetWindow(), bShowWindow );
+		if (GeneratePathToWidgetUnchecked( InParentWidget.ToSharedRef(), WidgetPath ))
+		{
+			AddWindowAsNativeChild( InSlateWindow, WidgetPath.GetWindow(), bShowWindow );
+		}
+		else
+		{
+			UE_LOG(LogSlate, Warning, TEXT("Modal Window fail to open as a native child. The path to the parent widget (%s) could not be found"), *InParentWidget->ToString());
+			AddWindow( InSlateWindow, bShowWindow );
+		}
 	}
 	else
 	{
