Remove ensure from TransformLineText, as it was falsly tripping when changing the case of eszett characters

See UDN: https://epicgames.lightning.force.com/lightning/r/Case/5004z00001o0hoLAAQ/view
A bit of context via the ICU userguide (problem #1 under Case Mappings) - [REDACTED_PROFILE_URL]/icu/blob/d205e03352d0b2339d6b99d35e143e0843029cfa/docs/userguide/icu/posix.md

diff --git a/Engine/Source/Runtime/Slate/Private/Framework/Text/TextLayout.cpp b/Engine/Source/Runtime/Slate/Private/Framework/Text/TextLayout.cpp
--- a/Engine/Source/Runtime/Slate/Private/Framework/Text/TextLayout.cpp
+++ b/Engine/Source/Runtime/Slate/Private/Framework/Text/TextLayout.cpp
@@ -995,10 +995,7 @@ void FTextLayout::TransformLineText(FLineModel& LineModel) const
 {
 	auto ApplyTransformedText = [&LineModel](FString&& InNewText)
 	{
-		if (ensure(LineModel.Text->Len() == InNewText.Len()))
-		{
-			*LineModel.Text = MoveTemp(InNewText);
-		}
+		*LineModel.Text = MoveTemp(InNewText);
 	};
 
 	switch (TransformPolicy)
