PR #10039: Fix scroll box scroll lag when scrolling by holding the scroll bar with the mouse

diff --git a/Engine/Source/Runtime/Slate/Private/Widgets/Layout/SScrollBox.cpp b/Engine/Source/Runtime/Slate/Private/Widgets/Layout/SScrollBox.cpp
--- a/Engine/Source/Runtime/Slate/Private/Widgets/Layout/SScrollBox.cpp
+++ b/Engine/Source/Runtime/Slate/Private/Widgets/Layout/SScrollBox.cpp
@@ -1126,6 +1126,8 @@ int32 SScrollBox::OnPaint( const FPaintArgs& Args, const FGeometry& AllottedGeom
 
 void SScrollBox::ScrollBar_OnUserScrolled( float InScrollOffsetFraction )
 {
+	bAnimateScroll = false;
+
 	const float ContentSize = GetScrollComponentFromVector(ScrollPanel->GetDesiredSize());
 	const FGeometry ScrollPanelGeometry = FindChildGeometry(CachedGeometry, ScrollPanel.ToSharedRef());
 
