PR #10217 Common UI Analog Cursor Fixes

diff --git a/Engine/Plugins/Runtime/CommonUI/Source/CommonInput/Private/CommonInputSubsystem.cpp b/Engine/Plugins/Runtime/CommonUI/Source/CommonInput/Private/CommonInputSubsystem.cpp
--- a/Engine/Plugins/Runtime/CommonUI/Source/CommonInput/Private/CommonInputSubsystem.cpp
+++ b/Engine/Plugins/Runtime/CommonUI/Source/CommonInput/Private/CommonInputSubsystem.cpp
@@ -414,15 +414,15 @@ bool UCommonInputSubsystem::PlatformSupportsHardwareCursor() const
 
 void UCommonInputSubsystem::SetCursorPosition(FVector2D NewPosition, bool bForce)
 {
-	ULocalPlayer* LocalPlayer = GetLocalPlayerChecked();
-	if (TSharedPtr<FSlateUser> SlateUser = LocalPlayer ? LocalPlayer->GetSlateUser() : nullptr)
+	if (CommonInputPreprocessor)
 	{
-		UpdateCursorPosition(SlateUser.ToSharedRef(), NewPosition, bForce);
+		CommonInputPreprocessor->bIgnoreNextMove = true;
 	}
 
-	if (CommonInputPreprocessor)
+	ULocalPlayer* LocalPlayer = GetLocalPlayerChecked();
+	if (TSharedPtr<FSlateUser> SlateUser = LocalPlayer ? LocalPlayer->GetSlateUser() : nullptr)
 	{
-		CommonInputPreprocessor->bIgnoreNextMove = true;
+		UpdateCursorPosition(SlateUser.ToSharedRef(), NewPosition, bForce);
 	}
 }
 
diff --git a/Engine/Plugins/Runtime/CommonUI/Source/CommonUI/Private/Input/CommonAnalogCursor.cpp b/Engine/Plugins/Runtime/CommonUI/Source/CommonUI/Private/Input/CommonAnalogCursor.cpp
--- a/Engine/Plugins/Runtime/CommonUI/Source/CommonUI/Private/Input/CommonAnalogCursor.cpp
+++ b/Engine/Plugins/Runtime/CommonUI/Source/CommonUI/Private/Input/CommonAnalogCursor.cpp
@@ -75,7 +75,10 @@ void FCommonAnalogCursor::Tick(const float DeltaTime, FSlateApplication& SlateAp
 #endif
 		if (bIsAnalogMovementEnabled)
 		{
-			FAnalogCursor::Tick(DeltaTime, SlateApp, Cursor);
+			const FVector2D NewPosition = CalculateTickedCursorPosition(DeltaTime, SlateApp, SlateUser);
+
+			UCommonInputSubsystem& InputSubsystem = ActionRouter.GetInputSubsystem();
+			InputSubsystem.SetCursorPosition(NewPosition, false);
 		}
 		else
 		{
diff --git a/Engine/Plugins/Runtime/CommonUI/Source/CommonUI/Public/Input/CommonAnalogCursor.h b/Engine/Plugins/Runtime/CommonUI/Source/CommonUI/Public/Input/CommonAnalogCursor.h
--- a/Engine/Plugins/Runtime/CommonUI/Source/CommonUI/Public/Input/CommonAnalogCursor.h
+++ b/Engine/Plugins/Runtime/CommonUI/Source/CommonUI/Public/Input/CommonAnalogCursor.h
@@ -111,7 +111,7 @@ private:
 		RightShoulder = 1 << 1,
 		LeftTrigger = 1 << 2,
 		RightTrigger = 1 << 3,
-		All = 0x1111
+		All = LeftShoulder | RightShoulder | LeftTrigger | RightTrigger
 	};
 	int32 ShoulderButtonStatus = EShoulderButtonFlags::None;
 #endif
diff --git a/Engine/Source/Runtime/Slate/Private/Framework/Application/AnalogCursor.cpp b/Engine/Source/Runtime/Slate/Private/Framework/Application/AnalogCursor.cpp
--- a/Engine/Source/Runtime/Slate/Private/Framework/Application/AnalogCursor.cpp
+++ b/Engine/Source/Runtime/Slate/Private/Framework/Application/AnalogCursor.cpp
@@ -27,94 +27,7 @@ void FAnalogCursor::Tick(const float DeltaTime, FSlateApplication& SlateApp, TSh
 {
 	if (TSharedPtr<FSlateUser> SlateUser = SlateApp.GetUser(GetOwnerUserIndex()))
 	{
-		const FVector2D OldPosition = SlateUser->GetCursorPosition();
-
-		float SpeedMult = 1.0f; // Used to do a speed multiplication before adding the delta to the position to make widgets sticky
-		FVector2D AdjAnalogVals = GetAnalogValue(EAnalogStick::Left); // A copy of the analog values so I can modify them based being over a widget, not currently doing this
-
-		// Adjust analog values according to dead zone
-		const float AnalogValsSize = AdjAnalogVals.Size();
-
-		if (AnalogValsSize > 0.0f)
-		{
-			const float TargetSize = FMath::Max(AnalogValsSize - DeadZone, 0.0f) / (1.0f - DeadZone);
-			AdjAnalogVals /= AnalogValsSize;
-			AdjAnalogVals *= TargetSize;
-		}
-
-
-		// Check if there is a sticky widget beneath the cursor
-		FWidgetPath WidgetPath = SlateApp.LocateWindowUnderMouse(OldPosition, SlateApp.GetInteractiveTopLevelWindows(), false, SlateUser->GetUserIndex());
-		if (WidgetPath.IsValid())
-		{
-			const FArrangedChildren::FArrangedWidgetArray& AllArrangedWidgets = WidgetPath.Widgets.GetInternalArray();
-			for (const FArrangedWidget& ArrangedWidget : AllArrangedWidgets)
-			{
-				TSharedRef<SWidget> Widget = ArrangedWidget.Widget;
-				if (Widget->IsInteractable())
-				{
-					SpeedMult = StickySlowdown;
-					//FVector2D Adjustment = WidgetsAndCursors.Last().Geometry.Position - OldPosition; // example of calculating distance from cursor to widget center
-					break;
-				}
-			}
-		}
-
-		switch (Mode)
-		{
-		case AnalogCursorMode::Accelerated:
-		{
-			// Generate Min and Max for X to clamp the speed, this gives us instant direction change when crossing the axis
-			float CurrentMinSpeedX = 0.0f;
-			float CurrentMaxSpeedX = 0.0f;
-			if (AdjAnalogVals.X > 0.0f)
-			{
-				CurrentMaxSpeedX = AdjAnalogVals.X * MaxSpeed;
-			}
-			else
-			{
-				CurrentMinSpeedX = AdjAnalogVals.X * MaxSpeed;
-			}
-
-			// Generate Min and Max for Y to clamp the speed, this gives us instant direction change when crossing the axis
-			float CurrentMinSpeedY = 0.0f;
-			float CurrentMaxSpeedY = 0.0f;
-			if (AdjAnalogVals.Y > 0.0f)
-			{
-				CurrentMaxSpeedY = AdjAnalogVals.Y * MaxSpeed;
-			}
-			else
-			{
-				CurrentMinSpeedY = AdjAnalogVals.Y * MaxSpeed;
-			}
-
-			// Cubic acceleration curve
-			FVector2D ExpAcceleration = AdjAnalogVals * AdjAnalogVals * AdjAnalogVals * Acceleration;
-			// Preserve direction (if we use a squared equation above)
-			//ExpAcceleration.X *= FMath::Sign(AnalogValues.X);
-			//ExpAcceleration.Y *= FMath::Sign(AnalogValues.Y);
-
-			CurrentSpeed += ExpAcceleration * DeltaTime;
-
-			CurrentSpeed.X = FMath::Clamp(CurrentSpeed.X, CurrentMinSpeedX, CurrentMaxSpeedX);
-			CurrentSpeed.Y = FMath::Clamp(CurrentSpeed.Y, CurrentMinSpeedY, CurrentMaxSpeedY);
-
-			break;
-		}
-
-		case AnalogCursorMode::Direct:
-
-			CurrentSpeed = AdjAnalogVals * MaxSpeed;
-
-			break;
-		}
-
-		CurrentOffset += CurrentSpeed * DeltaTime * SpeedMult;
-		const FVector2D NewPosition = OldPosition + CurrentOffset;
-
-		// save the remaining sub-pixel offset 
-		CurrentOffset.X = FGenericPlatformMath::Frac(NewPosition.X);
-		CurrentOffset.Y = FGenericPlatformMath::Frac(NewPosition.Y);
+		const FVector2D NewPosition = CalculateTickedCursorPosition(DeltaTime, SlateApp, SlateUser);
 
 		// update the cursor position
 		UpdateCursorPosition(SlateApp, SlateUser.ToSharedRef(), NewPosition);
@@ -370,3 +283,97 @@ void FAnalogCursor::UpdateCursorPosition(FSlateApplication& SlateApp, TSharedRef
 	}
 }
 
+FVector2D FAnalogCursor::CalculateTickedCursorPosition(const float DeltaTime, FSlateApplication& SlateApp, TSharedPtr<FSlateUser> SlateUser) 
+{
+	const FVector2D OldPosition = SlateUser->GetCursorPosition();
+
+	float SpeedMult = 1.0f; // Used to do a speed multiplication before adding the delta to the position to make widgets sticky
+	FVector2D AdjAnalogVals = GetAnalogValue(EAnalogStick::Left); // A copy of the analog values so I can modify them based being over a widget, not currently doing this
+
+	// Adjust analog values according to dead zone
+	const float AnalogValsSize = AdjAnalogVals.Size();
+
+	if (AnalogValsSize > 0.0f)
+	{
+		const float TargetSize = FMath::Max(AnalogValsSize - DeadZone, 0.0f) / (1.0f - DeadZone);
+		AdjAnalogVals /= AnalogValsSize;
+		AdjAnalogVals *= TargetSize;
+	}
+
+
+	// Check if there is a sticky widget beneath the cursor
+	FWidgetPath WidgetPath = SlateApp.LocateWindowUnderMouse(OldPosition, SlateApp.GetInteractiveTopLevelWindows(), false, SlateUser->GetUserIndex());
+	if (WidgetPath.IsValid())
+	{
+		const FArrangedChildren::FArrangedWidgetArray& AllArrangedWidgets = WidgetPath.Widgets.GetInternalArray();
+		for (const FArrangedWidget& ArrangedWidget : AllArrangedWidgets)
+		{
+			TSharedRef<SWidget> Widget = ArrangedWidget.Widget;
+			if (Widget->IsInteractable())
+			{
+				SpeedMult = StickySlowdown;
+				//FVector2D Adjustment = WidgetsAndCursors.Last().Geometry.Position - OldPosition; // example of calculating distance from cursor to widget center
+				break;
+			}
+		}
+	}
+
+	switch (Mode)
+	{
+	case AnalogCursorMode::Accelerated:
+	{
+		// Generate Min and Max for X to clamp the speed, this gives us instant direction change when crossing the axis
+		float CurrentMinSpeedX = 0.0f;
+		float CurrentMaxSpeedX = 0.0f;
+		if (AdjAnalogVals.X > 0.0f)
+		{
+			CurrentMaxSpeedX = AdjAnalogVals.X * MaxSpeed;
+		}
+		else
+		{
+			CurrentMinSpeedX = AdjAnalogVals.X * MaxSpeed;
+		}
+
+		// Generate Min and Max for Y to clamp the speed, this gives us instant direction change when crossing the axis
+		float CurrentMinSpeedY = 0.0f;
+		float CurrentMaxSpeedY = 0.0f;
+		if (AdjAnalogVals.Y > 0.0f)
+		{
+			CurrentMaxSpeedY = AdjAnalogVals.Y * MaxSpeed;
+		}
+		else
+		{
+			CurrentMinSpeedY = AdjAnalogVals.Y * MaxSpeed;
+		}
+
+		// Cubic acceleration curve
+		FVector2D ExpAcceleration = AdjAnalogVals * AdjAnalogVals * AdjAnalogVals * Acceleration;
+		// Preserve direction (if we use a squared equation above)
+		//ExpAcceleration.X *= FMath::Sign(AnalogValues.X);
+		//ExpAcceleration.Y *= FMath::Sign(AnalogValues.Y);
+
+		CurrentSpeed += ExpAcceleration * DeltaTime;
+
+		CurrentSpeed.X = FMath::Clamp(CurrentSpeed.X, CurrentMinSpeedX, CurrentMaxSpeedX);
+		CurrentSpeed.Y = FMath::Clamp(CurrentSpeed.Y, CurrentMinSpeedY, CurrentMaxSpeedY);
+
+		break;
+	}
+
+	case AnalogCursorMode::Direct:
+
+		CurrentSpeed = AdjAnalogVals * MaxSpeed;
+
+		break;
+	}
+
+	CurrentOffset += CurrentSpeed * DeltaTime * SpeedMult;
+	const FVector2D NewPosition = OldPosition + CurrentOffset;
+
+	// save the remaining sub-pixel offset 
+	CurrentOffset.X = FGenericPlatformMath::Frac(NewPosition.X);
+	CurrentOffset.Y = FGenericPlatformMath::Frac(NewPosition.Y);
+
+	return NewPosition;
+}
+
diff --git a/Engine/Source/Runtime/Slate/Public/Framework/Application/AnalogCursor.h b/Engine/Source/Runtime/Slate/Public/Framework/Application/AnalogCursor.h
--- a/Engine/Source/Runtime/Slate/Public/Framework/Application/AnalogCursor.h
+++ b/Engine/Source/Runtime/Slate/Public/Framework/Application/AnalogCursor.h
@@ -80,6 +80,8 @@ protected:
 
 	virtual void UpdateCursorPosition(FSlateApplication& SlateApp, TSharedRef<FSlateUser> SlateUser, const FVector2D& NewPosition, bool bForce = false);
 
+	virtual FVector2D CalculateTickedCursorPosition(const float DeltaTime, FSlateApplication& SlateApp, TSharedPtr<FSlateUser> SlateUser);
+
 	/** Current speed of the cursor */
 	FVector2D CurrentSpeed;
 
