Spinbox: reset cached string through callback only when the display changes (not the value itself). Normally, this should be the case without a check, but since we wanted to avoid having extra delegates in MovieScene, we may call this function more often than necessary.

diff --git a/Engine/Source/Runtime/Slate/Public/Widgets/Input/SSpinBox.h b/Engine/Source/Runtime/Slate/Public/Widgets/Input/SSpinBox.h
--- a/Engine/Source/Runtime/Slate/Public/Widgets/Input/SSpinBox.h
+++ b/Engine/Source/Runtime/Slate/Public/Widgets/Input/SSpinBox.h
@@ -1212,9 +1212,15 @@ private:
 
 	/** Used to prevent per-frame re-conversion of the cached numeric value to a string. */
 	FString CachedValueString;
+
+	/** Reset the cached string when the value is the same but the display format changed (through the callback). */
 	void ResetCachedValueString() 
 	{ 
-		CachedValueString = Interface->ToString(ValueAttribute.Get()); 
+		const NumericType CurrentValue = ValueAttribute.Get();
+		if (CachedExternalValue == CurrentValue)
+		{
+			CachedValueString = Interface->ToString(CurrentValue);
+		}
 	}
 
 	/** Whetever the interfaced setting changed and the CachedValueString needs to be recomputed. */
