Slate: Deprecate the SListView functions that managed the item in the STreeView. User should use the RootItem version instead.

diff --git a/Engine/Plugins/Animation/ControlRig/Source/ControlRigEditor/Private/Editor/SRigHierarchyTreeView.cpp b/Engine/Plugins/Animation/ControlRig/Source/ControlRigEditor/Private/Editor/SRigHierarchyTreeView.cpp
--- a/Engine/Plugins/Animation/ControlRig/Source/ControlRigEditor/Private/Editor/SRigHierarchyTreeView.cpp
+++ b/Engine/Plugins/Animation/ControlRig/Source/ControlRigEditor/Private/Editor/SRigHierarchyTreeView.cpp
@@ -751,7 +751,7 @@ TArray<FRigElementKey> SRigHierarchyTreeView::GetSelectedKeys() const
 
 const TSharedPtr<FRigTreeElement>* SRigHierarchyTreeView::FindItemAtPosition(FVector2D InScreenSpacePosition) const
 {
-	if (ItemsPanel.IsValid() && HasValidItemsSource())
+	if (ItemsPanel.IsValid() && SListView<TSharedPtr<FRigTreeElement>>::HasValidItemsSource())
 	{
 		FArrangedChildren ArrangedChildren(EVisibility::Visible);
 		const int32 Index = FindChildUnderPosition(ArrangedChildren, InScreenSpacePosition);
@@ -761,7 +761,7 @@ const TSharedPtr<FRigTreeElement>* SRigHierarchyTreeView::FindItemAtPosition(FVe
 			if (ItemWidget->WeakRigTreeElement.IsValid())
 			{
 				const FRigElementKey Key = ItemWidget->WeakRigTreeElement.Pin()->Key;
-				const TSharedPtr<FRigTreeElement>* ResultPtr = GetItems().FindByPredicate([Key](const TSharedPtr<FRigTreeElement>& Item) -> bool
+				const TSharedPtr<FRigTreeElement>* ResultPtr = SListView<TSharedPtr<FRigTreeElement>>::GetItems().FindByPredicate([Key](const TSharedPtr<FRigTreeElement>& Item) -> bool
 					{
 						return Item->Key == Key;
 					});
diff --git a/Engine/Plugins/Runtime/CommonUI/Source/CommonUI/Public/CommonTreeView.h b/Engine/Plugins/Runtime/CommonUI/Source/CommonUI/Public/CommonTreeView.h
--- a/Engine/Plugins/Runtime/CommonUI/Source/CommonUI/Public/CommonTreeView.h
+++ b/Engine/Plugins/Runtime/CommonUI/Source/CommonUI/Public/CommonTreeView.h
@@ -21,7 +21,7 @@ public:
 		{
 			// Set selection to the first item in a list if no items are selected.
 			// If bReturnFocusToSelection is true find the last selected object and focus on that.
-			if (this->HasValidItemsSource() && this->GetRootItems().Num() > 0)
+			if (this->HasValidRootItemsSource() && this->GetRootItems().Num() > 0)
 			{
 				typename TListTypeTraits<ItemType>::NullableType ItemNavigatedTo = TListTypeTraits<ItemType>::MakeNullPtr();
 				if (this->GetNumItemsSelected() == 0)
diff --git a/Engine/Source/Editor/SequencerCore/Private/MVVM/Views/SOutlinerView.cpp b/Engine/Source/Editor/SequencerCore/Private/MVVM/Views/SOutlinerView.cpp
--- a/Engine/Source/Editor/SequencerCore/Private/MVVM/Views/SOutlinerView.cpp
+++ b/Engine/Source/Editor/SequencerCore/Private/MVVM/Views/SOutlinerView.cpp
@@ -575,7 +575,7 @@ FReply SOutlinerView::OnDrop(const FGeometry& MyGeometry, const FDragDropEvent&
 
 FReply SOutlinerView::OnKeyDown(const FGeometry& MyGeometry, const FKeyEvent& InKeyEvent)
 {
-	const TArrayView<const TWeakViewModelPtr<IOutlinerExtension>> ItemsSourceRef = GetItems();
+	const TArrayView<const TWeakViewModelPtr<IOutlinerExtension>> ItemsSourceRef = SListView<TWeakViewModelPtr<IOutlinerExtension>>::GetItems();
 
 	auto IsSelectable = [](const TWeakViewModelPtr<IOutlinerExtension>& WeakModel)
 	{
@@ -1087,9 +1087,9 @@ void SOutlinerView::ToggleExpandCollapseNodes(ETreeRecursion Recursion, bool bEx
 			}
 		}
 	}
-	else if (GetItems().Num() > 0)
+	else if (SListView<TWeakViewModelPtr<IOutlinerExtension>>::GetItems().Num() > 0)
 	{
-		bool bExpand = ShouldExpand(GetItems(), Recursion, Outliner);
+		bool bExpand = ShouldExpand(SListView<TWeakViewModelPtr<IOutlinerExtension>>::GetItems(), Recursion, Outliner);
 
 		if (bExpandAll)
 		{
@@ -1100,7 +1100,7 @@ void SOutlinerView::ToggleExpandCollapseNodes(ETreeRecursion Recursion, bool bEx
 			bExpand = false;
 		}
 
-		for (TWeakViewModelPtr<IOutlinerExtension> WeakItem : GetItems())
+		for (TWeakViewModelPtr<IOutlinerExtension> WeakItem : SListView<TWeakViewModelPtr<IOutlinerExtension>>::GetItems())
 		{
 			if (TViewModelPtr<IOutlinerExtension> Item = WeakItem.Pin())
 			{
diff --git a/Engine/Source/Editor/SkeletonEditor/Private/SSkeletonTree.cpp b/Engine/Source/Editor/SkeletonEditor/Private/SSkeletonTree.cpp
--- a/Engine/Source/Editor/SkeletonEditor/Private/SSkeletonTree.cpp
+++ b/Engine/Source/Editor/SkeletonEditor/Private/SSkeletonTree.cpp
@@ -94,7 +94,7 @@ public:
 			return;
 		}
 
-		const TArrayView<const ItemType> ItemsSourceRef = this->GetItems();
+		const TArrayView<const ItemType> ItemsSourceRef = this->SListView<ItemType>::GetItems();
 
 		int32 RangeStartIndex = 0;
 		if( TListTypeTraits<ItemType>::IsPtrValid(this->RangeSelectionStart) )
diff --git a/Engine/Source/Runtime/Slate/Public/Widgets/Views/STreeView.h b/Engine/Source/Runtime/Slate/Public/Widgets/Views/STreeView.h
--- a/Engine/Source/Runtime/Slate/Public/Widgets/Views/STreeView.h
+++ b/Engine/Source/Runtime/Slate/Public/Widgets/Views/STreeView.h
@@ -898,12 +898,38 @@ public:
 		return Private_IsItemExpanded( InItem );
 	}
 
-protected:
-	//~ Hide the base SetItemsSource
-	using SListView<ItemType>::SetItemsSource;
-	using SListView<ItemType>::ClearItemsSource;
-	using SListView<ItemType>::HasValidItemsSource;
-	using SListView<ItemType>::GetItems;
+public:
+	//~ Hide the base function from SListView 
+	UE_DEPRECATED(5.3, "SetItemsSource is deprecated. You probably want to use SetTreeItemsSource.")
+	void SetItemsSource(const TArray<ItemType>* InListItemsSource)
+	{
+		SListView<ItemType>::SetItemsSource(InListItemsSource);
+	}
+	UE_DEPRECATED(5.3, "SetItemsSource is deprecated. You probably want to use SetTreeItemsSource.")
+	void SetItemsSource(TSharedRef<::UE::Slate::Containers::TObservableArray<ItemType>> InListItemsSource)
+	{
+		SListView<ItemType>::SetItemsSource(InListItemsSource);
+	}
+	UE_DEPRECATED(5.3, "SetItemsSource is deprecated. You probably want to use SetTreeItemsSource.")
+	void SetItemsSource(TUniquePtr<UE::Slate::ItemsSource::IItemsSource<ItemType>> Provider)
+	{
+		SListView<ItemType>::SetItemsSource(MoveTemp(Provider));
+	}
+	UE_DEPRECATED(5.3, "ClearItemsSource is deprecated. You probably want to use ClearRootItemsSource.")
+	void ClearItemsSource()
+	{
+		SListView<ItemType>::ClearItemsSource();
+	}
+	UE_DEPRECATED(5.3, "HasValidItemsSource is deprecated. You probably want to use HasValidRootItemsSource.")
+	bool HasValidItemsSource() const
+	{
+		return SListView<ItemType>::HasValidItemsSource();
+	}
+	UE_DEPRECATED(5.3, "GetItems is deprecated. You probably want to use GetRootItems.")
+	TArrayView<const ItemType> GetItems() const
+	{
+		return SListView<ItemType>::GetItems();
+	}
 
 public:
 
