Spinbox: Update cached external value when the cached string value changes.

diff --git a/Engine/Source/Runtime/Slate/Public/Widgets/Input/SSpinBox.h b/Engine/Source/Runtime/Slate/Public/Widgets/Input/SSpinBox.h
--- a/Engine/Source/Runtime/Slate/Public/Widgets/Input/SSpinBox.h
+++ b/Engine/Source/Runtime/Slate/Public/Widgets/Input/SSpinBox.h
@@ -1217,10 +1217,8 @@ private:
 	void ResetCachedValueString() 
 	{ 
 		const NumericType CurrentValue = ValueAttribute.Get();
-		if (CachedExternalValue == CurrentValue)
-		{
-			CachedValueString = Interface->ToString(CurrentValue);
-		}
+		CachedExternalValue = CurrentValue;
+		CachedValueString = Interface->ToString(CachedExternalValue);
 	}
 
 	/** Whetever the interfaced setting changed and the CachedValueString needs to be recomputed. */
