Undo //UE5/Main/Engine/Source/Runtime/Slate/Private/Framework/Text/... changelist 25494916

#rnx

diff --git a/Engine/Source/Runtime/Slate/Private/Framework/Text/TextLayout.cpp b/Engine/Source/Runtime/Slate/Private/Framework/Text/TextLayout.cpp
--- a/Engine/Source/Runtime/Slate/Private/Framework/Text/TextLayout.cpp
+++ b/Engine/Source/Runtime/Slate/Private/Framework/Text/TextLayout.cpp
@@ -995,7 +995,10 @@ void FTextLayout::TransformLineText(FLineModel& LineModel) const
 {
 	auto ApplyTransformedText = [&LineModel](FString&& InNewText)
 	{
-		*LineModel.Text = MoveTemp(InNewText);
+		if (ensure(LineModel.Text->Len() == InNewText.Len()))
+		{
+			*LineModel.Text = MoveTemp(InNewText);
+		}
 	};
 
 	switch (TransformPolicy)
