Ability to restrict what characters can be used in a text property editor.

diff --git a/Engine/Source/Editor/PropertyEditor/Private/UserInterface/PropertyEditor/SPropertyEditorText.cpp b/Engine/Source/Editor/PropertyEditor/Private/UserInterface/PropertyEditor/SPropertyEditorText.cpp
--- a/Engine/Source/Editor/PropertyEditor/Private/UserInterface/PropertyEditor/SPropertyEditorText.cpp
+++ b/Engine/Source/Editor/PropertyEditor/Private/UserInterface/PropertyEditor/SPropertyEditorText.cpp
@@ -35,6 +35,7 @@ void SPropertyEditorText::Construct( const FArguments& InArgs, const TSharedRef<
 	static const FName NAME_MaxLength = "MaxLength";
 	static const FName NAME_MultiLine = "MultiLine";
 	static const FName NAME_PasswordField = "PasswordField";
+	static const FName NAME_AllowedCharacters = "AllowedCharacters";
 
 	bIsMultiLine = InPropertyEditor->GetPropertyHandle()->GetBoolMetaData(NAME_MultiLine);
 
@@ -45,7 +46,8 @@ void SPropertyEditorText::Construct( const FArguments& InArgs, const TSharedRef<
 	}
 
 	const bool bIsPassword = InPropertyEditor->GetPropertyHandle()->GetBoolMetaData(NAME_PasswordField);
-	
+	AllowedCharacters.InitializeFromString(InPropertyEditor->GetPropertyHandle()->GetMetaData(NAME_AllowedCharacters));
+
 	TSharedPtr<SHorizontalBox> HorizontalBox;
 	if(bIsMultiLine)
 	{
@@ -155,12 +157,33 @@ void SPropertyEditorText::OnTextCommitted( const FText& NewText, ETextCommit::Ty
 
 bool SPropertyEditorText::OnVerifyTextChanged(const FText& Text, FText& OutError)
 {
-	if (MaxLength > 0 && Text.ToString().Len() > MaxLength)
+	const FString& TextString = Text.ToString();
+
+	if (MaxLength > 0 && TextString.Len() > MaxLength)
 	{
-		OutError = FText::Format(LOCTEXT("PropertyTextTooLongError", "This value is too long ({0}/{1} characters)"), Text.ToString().Len(), MaxLength);
+		OutError = FText::Format(LOCTEXT("PropertyTextTooLongError", "This value is too long ({0}/{1} characters)"), TextString.Len(), MaxLength);
 		return false;
 	}
 
+	if (!AllowedCharacters.IsEmpty())
+	{
+		if (!TextString.IsEmpty() && !AllowedCharacters.AreAllCharsIncluded(TextString))
+		{
+			TSet<TCHAR> InvalidCharacters = AllowedCharacters.FindCharsNotIncluded(TextString);
+			FString InvalidCharactersString;
+			for (TCHAR Char : InvalidCharacters)
+			{
+				if (!InvalidCharactersString.IsEmpty())
+				{
+					InvalidCharactersString.AppendChar(TEXT(' '));
+				}
+				InvalidCharactersString.AppendChar(Char);
+			}
+			OutError = FText::Format(LOCTEXT("PropertyTextCharactersNotAllowedError", "The value may not contain the following characters: {0}"), FText::FromString(InvalidCharactersString));
+			return false;
+		}
+	}
+
 	return true;
 }
 
diff --git a/Engine/Source/Editor/PropertyEditor/Private/UserInterface/PropertyEditor/SPropertyEditorText.h b/Engine/Source/Editor/PropertyEditor/Private/UserInterface/PropertyEditor/SPropertyEditorText.h
--- a/Engine/Source/Editor/PropertyEditor/Private/UserInterface/PropertyEditor/SPropertyEditorText.h
+++ b/Engine/Source/Editor/PropertyEditor/Private/UserInterface/PropertyEditor/SPropertyEditorText.h
@@ -2,6 +2,7 @@
 #pragma once
 
 #include "Fonts/SlateFontInfo.h"
+#include "Framework/Text/CharRangeList.h"
 #include "Input/Reply.h"
 #include "Misc/Attribute.h"
 #include "Styling/AppStyle.h"
@@ -67,4 +68,7 @@ private:
 	
 	/** The maximum length of the value that can be edited, or <=0 for unlimited */
 	int32 MaxLength = 0;
+
+	/** The characters that are allowed in the value. If empty, any character is allowed. */
+	FCharRangeList AllowedCharacters;
 };
diff --git a/Engine/Source/Runtime/Slate/Private/Framework/Text/CharRangeList.cpp b/Engine/Source/Runtime/Slate/Private/Framework/Text/CharRangeList.cpp
new file mode 100644
--- /dev/null
+++ b/Engine/Source/Runtime/Slate/Private/Framework/Text/CharRangeList.cpp
@@ -0,0 +1,181 @@
+// Copyright Epic Games, Inc. All Rights Reserved.
+
+#include "Framework/Text/CharRangeList.h"
+
+namespace CharRangeConstants
+{
+	static const FString EscapedCharacters = TEXT("-\\");
+	static const TCHAR EscapeCharacter = TEXT('\\');
+	static const TCHAR RangeSeparator = TEXT('-');
+}
+
+bool FCharRange::InitializeFromString(const FString& InDefinitionString, int32 InFromIndex, int32* OutLastIndex)
+{
+	int32 Length = InDefinitionString.Len();
+	int32 Index = InFromIndex;
+	TCHAR Char;
+
+	// Parse the first character (and last if this is a single character range).
+	Char = InDefinitionString[Index++];
+	if (Char != CharRangeConstants::EscapeCharacter)
+	{
+		First = Char;
+	}
+	else
+	{
+		if (Index < Length)
+		{
+			First = InDefinitionString[Index++];
+		}
+		else
+		{
+			// Error, stray escape character found at the end.
+			*OutLastIndex = Index;
+			return false;
+		}
+	}
+
+	Last = First;
+
+	// Parse the range's last character, if present.
+	if ((Index < Length) && (InDefinitionString[Index] == CharRangeConstants::RangeSeparator))
+	{
+		Index++;
+
+		Char = InDefinitionString[Index++];
+		if (Char != CharRangeConstants::EscapeCharacter)
+		{
+			Last = Char;
+		}
+		else
+		{
+			if (Index < Length)
+			{
+				Last = InDefinitionString[Index++];
+			}
+			else
+			{
+				// Error, stray escape character found at the end.
+				*OutLastIndex = Index;
+				return false;
+			}
+		}
+	}
+
+	*OutLastIndex = Index;
+	return true;
+}
+
+FString FCharRange::ToString() const
+{
+	FString Result;
+
+	TCHAR FirstChar = (TCHAR)First;
+	TCHAR LastChar = (TCHAR)Last;
+	int32 Unused = 0;
+
+	// First character.
+	if (CharRangeConstants::EscapedCharacters.FindChar(FirstChar, Unused))
+	{
+		Result.AppendChar(CharRangeConstants::EscapeCharacter);
+	}
+	Result.AppendChar(FirstChar);
+
+	// Is this a multi-character range?
+	if (Last > First)
+	{
+		Result.AppendChar(CharRangeConstants::RangeSeparator);
+
+		// Last character.
+		if (CharRangeConstants::EscapedCharacters.FindChar(LastChar, Unused))
+		{
+			Result.AppendChar(CharRangeConstants::EscapeCharacter);
+		}
+		Result.AppendChar(LastChar);
+	}
+
+	return Result;
+}
+
+bool FCharRangeList::InitializeFromString(const FString& InDefinitionString)
+{
+	Ranges.Empty();
+
+	int32 Length = InDefinitionString.Len();
+	int32 Index = 0;
+	int32 NextIndex;
+	FCharRange NewRange;
+	while (Index < Length)
+	{
+		if (!NewRange.InitializeFromString(InDefinitionString, Index, &NextIndex))
+		{
+			return false;
+		}
+
+		Ranges.Add(NewRange);
+		Index = NextIndex;
+	}
+
+	return true;
+}
+
+FString FCharRangeList::ToString() const
+{
+	FString Result;
+	for (const FCharRange& Range : Ranges)
+	{
+		Result.Append(Range.ToString());
+	}
+
+	return Result;
+}
+
+void FCharRangeList::Empty()
+{
+	Ranges.Empty();
+}
+
+bool FCharRangeList::IsEmpty() const
+{
+	return Ranges.IsEmpty();
+}
+
+bool FCharRangeList::IsCharIncluded(TCHAR InChar) const
+{
+	for (const FCharRange& Range : Ranges)
+	{
+		if ((InChar >= Range.First) && (InChar <= Range.Last))
+		{
+			return true;
+		}
+	}
+
+	return false;
+}
+
+bool FCharRangeList::AreAllCharsIncluded(const FString& InString) const
+{
+	for (TCHAR Char : InString)
+	{
+		if (!IsCharIncluded(Char))
+		{
+			return false;
+		}
+	}
+
+	return true;
+}
+
+TSet<TCHAR> FCharRangeList::FindCharsNotIncluded(const FString& InString) const
+{
+	TSet<TCHAR> InvalidChars;
+	for (TCHAR Char : InString)
+	{
+		if (!IsCharIncluded(Char))
+		{
+			InvalidChars.Add(Char);
+		}
+	}
+
+	return InvalidChars;
+}
diff --git a/Engine/Source/Runtime/Slate/Public/Framework/Text/CharRangeList.h b/Engine/Source/Runtime/Slate/Public/Framework/Text/CharRangeList.h
new file mode 100644
--- /dev/null
+++ b/Engine/Source/Runtime/Slate/Public/Framework/Text/CharRangeList.h
@@ -0,0 +1,68 @@
+// Copyright Epic Games, Inc. All Rights Reserved.
+
+#pragma once
+
+#include "CoreMinimal.h"
+#include "CharRangeList.generated.h"
+
+
+/** Represents a range of characters, specified by the Unicode code point of the first and last characters in the range, both included.
+*   If you need a single-character range, simply use the same character for both the first and last characters. */
+USTRUCT()
+struct FCharRange
+{
+	GENERATED_USTRUCT_BODY()
+
+	/** Initializes this range from the given substring. See ToString for an overview on the required format.
+	    *OutLastIndex will contain the index of the character following the parsed range. */
+	bool InitializeFromString(const FString& InDefinitionString, int32 InFromIndex, int32* OutLastIndex);
+	/** Returns a string that represents this range, as "a-z" if the range contains multiple characters or simply "p" if it contains a single character. */
+	FString ToString() const;
+
+	/** Unicode code point of the first character in the range (inclusive). Only characters within the Basic Multilingual Plane are supported.
+	*   Note: The code point must be entered in decimal, not hexadecimal. */
+	UPROPERTY(EditAnywhere, Category = "Char Range")
+	uint16 First;
+
+	/** Unicode code point of the last character in the range (inclusive). Only characters within the Basic Multilingual Plane are supported.
+	*   Note: The code point must be entered in decimal, not hexadecimal. */
+	UPROPERTY(EditAnywhere, Category = "Char Range")
+	uint16 Last;
+};
+
+/** Represents a list of character ranges. */
+USTRUCT()
+struct SLATE_API FCharRangeList
+{
+	GENERATED_USTRUCT_BODY()
+
+	/** Initializes this instance with the character ranges represented by the passed definition string.
+	*   A definition string contains characters and ranges of characters, one after another with no special separators between them.
+	*   Characters - and \ must be escaped like this: \- and \\
+	* 
+	*   Examples:
+	*       "aT._" <-- Letters 'a' and 'T', dot and underscore.
+	*       "a-zT._" <-- All letters from 'a' to 'z', letter 'T', dot and underscore.
+	*       "a-zA-Z0-9._" <-- All lowercase and uppercase letters, all digits, dot and underscore.
+	*       "a-zA-Z0-9\-\\._" <-- All lowercase and uppercase letters, all digits, minus sign, backslash, dot and underscore.
+	*/
+	bool InitializeFromString(const FString& InDefinitionString);
+	/** Returns a definition string that represents all the character ranges in this instance. */
+	FString ToString() const;
+
+	/** Empties this instance. */
+	void Empty();
+
+	/** Returns true if this instance does not contain any characters, or false otherwise. */
+	bool IsEmpty() const;
+
+	/** Returns true if the given char is included in any of the char ranges, or false otherwise. */
+	bool IsCharIncluded(TCHAR InChar) const;
+	/** Returns true if all the characters in the given string are included in any of the char ranges, or false otherwise. */
+	bool AreAllCharsIncluded(const FString& InString) const;
+	/** Finds all the characters in the given string that are not included in any of the char ranges. */
+	TSet<TCHAR> FindCharsNotIncluded(const FString& InString) const;
+
+	UPROPERTY(EditAnywhere, Category = "Char Range List")
+	TArray<FCharRange> Ranges;
+};
