Allow spinbox to update on text change / same for numeric entry

diff --git a/Engine/Source/Runtime/Slate/Public/Widgets/Input/SNumericEntryBox.h b/Engine/Source/Runtime/Slate/Public/Widgets/Input/SNumericEntryBox.h
--- a/Engine/Source/Runtime/Slate/Public/Widgets/Input/SNumericEntryBox.h
+++ b/Engine/Source/Runtime/Slate/Public/Widgets/Input/SNumericEntryBox.h
@@ -92,6 +92,7 @@ public:
 		, _MaxSliderValue(NumericType(100))
 		, _SliderExponent(1.f)
 		, _AllowWheel(true)
+		, _BroadcastValueChangesPerKey(false)
 		, _MinDesiredValueWidth(0.f)
 		, _DisplayToggle(false)
 		, _ToggleChecked(ECheckBoxState::Checked)
@@ -158,6 +159,8 @@ public:
 		SLATE_ATTRIBUTE(NumericType, SliderExponentNeutralValue )
 		/** Whether this spin box should have mouse wheel feature enabled, defaults to true */
 		SLATE_ARGUMENT( bool, AllowWheel )
+		/** True to broadcast every time we type */
+		SLATE_ARGUMENT( bool, BroadcastValueChangesPerKey)
 		/** Step to increment or decrement the value by when scrolling the mouse wheel. If not specified will determine automatically */
 		SLATE_ATTRIBUTE( TOptional< NumericType >, WheelStep )
 		/** The minimum desired width for the value portion of the control. */
@@ -263,6 +266,7 @@ public:
 				.SliderExponent(InArgs._SliderExponent)
 				.SliderExponentNeutralValue(InArgs._SliderExponentNeutralValue)
 				.EnableWheel(InArgs._AllowWheel)
+				.BroadcastValueChangesPerKey(InArgs._BroadcastValueChangesPerKey)
 				.WheelStep(InArgs._WheelStep)
 				.OnBeginSliderMovement(InArgs._OnBeginSliderMovement)
 				.OnEndSliderMovement(InArgs._OnEndSliderMovement)
diff --git a/Engine/Source/Runtime/Slate/Public/Widgets/Input/SSpinBox.h b/Engine/Source/Runtime/Slate/Public/Widgets/Input/SSpinBox.h
--- a/Engine/Source/Runtime/Slate/Public/Widgets/Input/SSpinBox.h
+++ b/Engine/Source/Runtime/Slate/Public/Widgets/Input/SSpinBox.h
@@ -81,6 +81,7 @@ public:
 		, _SupportDynamicSliderMinValue(false)
 		, _SliderExponent(1.f)
 		, _EnableWheel(true)
+		, _BroadcastValueChangesPerKey(false)
 		, _Font(FCoreStyle::Get().GetFontStyle(TEXT("NormalFont")))
 		, _ContentPadding(FMargin(2.0f, 1.0f))
 		, _OnValueChanged()
@@ -131,6 +132,8 @@ public:
 		SLATE_ATTRIBUTE(NumericType, SliderExponentNeutralValue)
 		/** Whether this spin box should have mouse wheel feature enabled, defaults to true */
 		SLATE_ARGUMENT(bool, EnableWheel)
+		/** True to broadcast every time we type */
+		SLATE_ARGUMENT(bool, BroadcastValueChangesPerKey)
 		/** Step to increment or decrement the value by when scrolling the mouse wheel. If not specified will determine automatically */
 		SLATE_ATTRIBUTE(TOptional< NumericType >, WheelStep)
 		/** Font used to display text in the slider */
@@ -213,6 +216,7 @@ public:
 		OnDynamicSliderMinValueChanged = InArgs._OnDynamicSliderMinValueChanged;
 
 		bEnableWheel = InArgs._EnableWheel;
+		bBroadcastValueChangesPerKey = InArgs._BroadcastValueChangesPerKey;
 		WheelStep = InArgs._WheelStep;
 
 		bPreventThrottling = InArgs._PreventThrottling;
@@ -946,6 +950,15 @@ protected:
 				FString ValidData = NumValidChars > 0 ? Data.Left(NumValidChars) : GetValueAsString();
 				EditableText->SetText(FText::FromString(ValidData));
 			}
+
+			if (bBroadcastValueChangesPerKey)
+			{
+				TOptional<NumericType> NewValue = Interface->FromString(NewText.ToString(), ValueAttribute.Get());
+				if (NewValue.IsSet())
+				{
+					CommitValue(NewValue.GetValue(), (double)NewValue.GetValue(), CommittedViaCode, ETextCommit::Default);
+				}
+			}
 		}
 	}
 
@@ -1238,6 +1251,9 @@ private:
 
 	/** Does this spin box have the mouse wheel feature enabled? */
 	bool bEnableWheel = true;
+
+	/** True to broadcast every time we type. */
+	bool bBroadcastValueChangesPerKey = false;
 };
 
 template<typename NumericType>
