Fixed uninitialized data warning

#rnx

diff --git a/Engine/Source/Runtime/Slate/Public/Framework/Text/CharRangeList.h b/Engine/Source/Runtime/Slate/Public/Framework/Text/CharRangeList.h
--- a/Engine/Source/Runtime/Slate/Public/Framework/Text/CharRangeList.h
+++ b/Engine/Source/Runtime/Slate/Public/Framework/Text/CharRangeList.h
@@ -22,12 +22,12 @@ struct FCharRange
 	/** Unicode code point of the first character in the range (inclusive). Only characters within the Basic Multilingual Plane are supported.
 	*   Note: The code point must be entered in decimal, not hexadecimal. */
 	UPROPERTY(EditAnywhere, Category = "Char Range")
-	uint16 First;
+	uint16 First = 0;
 
 	/** Unicode code point of the last character in the range (inclusive). Only characters within the Basic Multilingual Plane are supported.
 	*   Note: The code point must be entered in decimal, not hexadecimal. */
 	UPROPERTY(EditAnywhere, Category = "Char Range")
-	uint16 Last;
+	uint16 Last = 0;
 };
 
 /** Represents a list of character ranges. */
