TabManager: Fixed crash where we could remove an invalid tab from the collapsed tabs array.
- The issue is that FindPotentiallyClosedTab was storing an area index that could become invalid during the call to RestoreArea.
- RestoreArea can end up creating a viewport which can tick the task graph.
- The task graph might have work that would open UI elements (ie: the message log) and those UI elements might do tab work and modify the collapsed tab array.
- The solution is to hold a shared ref to the area and remove the entry in the collapsed tabs array using that shared ref, instead of assuming its index.

diff --git a/Engine/Source/Runtime/Slate/Private/Framework/Docking/TabManager.cpp b/Engine/Source/Runtime/Slate/Private/Framework/Docking/TabManager.cpp
--- a/Engine/Source/Runtime/Slate/Private/Framework/Docking/TabManager.cpp
+++ b/Engine/Source/Runtime/Slate/Private/Framework/Docking/TabManager.cpp
@@ -1513,13 +1513,18 @@ TSharedPtr<SDockingTabStack> FTabManager::FindPotentiallyClosedTab( const FTabId
 	FTabMatcher TabMatcher( ClosedTabId );
 
 	// Search among the COLLAPSED AREAS
-	const int32 CollapsedAreaWithMatchingTab = FindTabInCollapsedAreas( TabMatcher );
-	if ( CollapsedAreaWithMatchingTab != INDEX_NONE )
+	const int32 CollapsedAreaWithMatchingTabIndex = FindTabInCollapsedAreas( TabMatcher );
+	if ( CollapsedAreaWithMatchingTabIndex != INDEX_NONE )
 	{
-		TSharedPtr<SDockingArea> RestoredArea = RestoreArea(CollapsedDockAreas[CollapsedAreaWithMatchingTab], GetPrivateApi().GetParentWindow());
+		TSharedRef<FTabManager::FArea> CollapsedAreaWithMatchingTab = CollapsedDockAreas[CollapsedAreaWithMatchingTabIndex];
+
+		TSharedPtr<SDockingArea> RestoredArea = RestoreArea(CollapsedAreaWithMatchingTab, GetPrivateApi().GetParentWindow());
 		check(RestoredArea.IsValid());
-		// We have just un-collapsed this dock area
-		CollapsedDockAreas.RemoveAt(CollapsedAreaWithMatchingTab);
+
+		// We have just un-collapsed this dock area.
+		// Don't rely on the collapsed tab index: RestoreArea() can end up kicking the task graph which could do other tab work and modify the CollapsedDockAreas array.
+		CollapsedDockAreas.Remove(CollapsedAreaWithMatchingTab);
+
 		if (RestoredArea.IsValid())
 		{
 			StackWithClosedTab = FindTabInLiveArea(TabMatcher, StaticCastSharedRef<SDockingArea>(RestoredArea->AsShared()));
