Fixed tab headers jumping when closing a tab: don't call SetTabWellHidden if it's already in the right state as it starts an animation right away.

diff --git a/Engine/Source/Runtime/Slate/Private/Framework/Docking/SDockingSplitter.cpp b/Engine/Source/Runtime/Slate/Private/Framework/Docking/SDockingSplitter.cpp
--- a/Engine/Source/Runtime/Slate/Private/Framework/Docking/SDockingSplitter.cpp
+++ b/Engine/Source/Runtime/Slate/Private/Framework/Docking/SDockingSplitter.cpp
@@ -120,7 +120,10 @@ void SDockingSplitter::AdjustDockedTabsIfNeeded()
 			if (ChildNode->GetNodeType() == DockTabStack && ChildNode->GetVisibility() == EVisibility::Visible)
 			{
 				TSharedRef<SDockingTabStack> TabStack = StaticCastSharedRef<SDockingTabStack>(ChildNode);
-				TabStack->SetTabWellHidden(false);
+				if (TabStack->IsTabWellHidden())
+				{
+					TabStack->SetTabWellHidden(false);
+				}
 				break;
 			}
 			// else if  node type is splitter, the first tab stack might be in there... check its children to see
