Fixed some places calling RemoveAll(...) on the underlying array within a FString without explicitly considering the null terminator

#rnx

diff --git a/Engine/Plugins/Animation/DeformerGraph/Source/OptimusEditor/Private/Widgets/SOptimusShaderTextSearchWidget.cpp b/Engine/Plugins/Animation/DeformerGraph/Source/OptimusEditor/Private/Widgets/SOptimusShaderTextSearchWidget.cpp
--- a/Engine/Plugins/Animation/DeformerGraph/Source/OptimusEditor/Private/Widgets/SOptimusShaderTextSearchWidget.cpp
+++ b/Engine/Plugins/Animation/DeformerGraph/Source/OptimusEditor/Private/Widgets/SOptimusShaderTextSearchWidget.cpp
@@ -66,8 +66,12 @@ void SOptimusShaderTextSearchWidget::TriggerSearch(const FText& InNewSearchText)
 	{
 		SingleLineString.GetCharArray().RemoveAll([&](const TCHAR InChar) -> bool
 		{
-			const bool bIsCharAllowed = !FChar::IsLinebreak(InChar);
-			return !bIsCharAllowed;
+			if (InChar != 0)
+			{
+				const bool bIsCharAllowed = !FChar::IsLinebreak(InChar);
+				return !bIsCharAllowed;
+			}
+			return false;
 		});
 	}
 	
diff --git a/Engine/Source/Runtime/Slate/Private/Widgets/Text/SlateEditableTextLayout.cpp b/Engine/Source/Runtime/Slate/Private/Widgets/Text/SlateEditableTextLayout.cpp
--- a/Engine/Source/Runtime/Slate/Private/Widgets/Text/SlateEditableTextLayout.cpp
+++ b/Engine/Source/Runtime/Slate/Private/Widgets/Text/SlateEditableTextLayout.cpp
@@ -2050,8 +2050,12 @@ void FSlateEditableTextLayout::InsertTextAtCursorImpl(const FString& InString)
 		const bool bIsMultiLine = OwnerWidget->IsMultiLineTextEdit();
 		SanitizedString.GetCharArray().RemoveAll([&](const TCHAR InChar) -> bool
 		{
-			const bool bIsCharAllowed = IsCharAllowed(InChar) || (bIsMultiLine || !FChar::IsLinebreak(InChar));
-			return !bIsCharAllowed;
+			if (InChar != 0)
+			{
+				const bool bIsCharAllowed = IsCharAllowed(InChar) || (bIsMultiLine && FChar::IsLinebreak(InChar));
+				return !bIsCharAllowed;
+			}
+			return false;
 		});
 	}
 
