[Slate]
* Moved dll export from types to methods/staticvar

diff --git a/Engine/Source/Runtime/Slate/Private/Framework/Application/NavigationConfig.cpp b/Engine/Source/Runtime/Slate/Private/Framework/Application/NavigationConfig.cpp
--- a/Engine/Source/Runtime/Slate/Private/Framework/Application/NavigationConfig.cpp
+++ b/Engine/Source/Runtime/Slate/Private/Framework/Application/NavigationConfig.cpp
@@ -195,3 +195,26 @@ FString FNavigationConfig::ToString() const
 
 	return Builder.ToString();
 }
+
+FTwinStickNavigationConfig::FTwinStickNavigationConfig()
+{
+	bTabNavigation = false;
+
+	KeyEventRules =
+	{
+		{EKeys::Gamepad_DPad_Left, EUINavigation::Left},
+		{EKeys::Gamepad_DPad_Right, EUINavigation::Right},
+		{EKeys::Gamepad_DPad_Up, EUINavigation::Up},
+		{EKeys::Gamepad_DPad_Down, EUINavigation::Down}
+	};
+}
+
+bool FTwinStickNavigationConfig::IsAnalogHorizontalKey(const FKey& InKey) const
+{
+	return InKey == EKeys::Gamepad_LeftX || InKey == EKeys::Gamepad_RightX;
+}
+
+bool FTwinStickNavigationConfig::IsAnalogVerticalKey(const FKey& InKey) const
+{
+	return InKey == EKeys::Gamepad_LeftY || InKey == EKeys::Gamepad_RightY;
+}
diff --git a/Engine/Source/Runtime/Slate/Private/Framework/Commands/InputChord.cpp b/Engine/Source/Runtime/Slate/Private/Framework/Commands/InputChord.cpp
--- a/Engine/Source/Runtime/Slate/Private/Framework/Commands/InputChord.cpp
+++ b/Engine/Source/Runtime/Slate/Private/Framework/Commands/InputChord.cpp
@@ -9,6 +9,15 @@
 /* FInputChord interface
  *****************************************************************************/
 
+FInputChord::FInputChord()
+	: Key(EKeys::Invalid)
+	, bShift(false)
+	, bCtrl(false)
+	, bAlt(false)
+	, bCmd(false)
+{
+}
+
 /**
  * Returns the friendly, localized string name of this key binding
  */
diff --git a/Engine/Source/Runtime/Slate/Private/Framework/Commands/UICommandInfo.cpp b/Engine/Source/Runtime/Slate/Private/Framework/Commands/UICommandInfo.cpp
--- a/Engine/Source/Runtime/Slate/Private/Framework/Commands/UICommandInfo.cpp
+++ b/Engine/Source/Runtime/Slate/Private/Framework/Commands/UICommandInfo.cpp
@@ -101,6 +101,19 @@ FUICommandInfoDecl::operator TSharedRef<FUICommandInfo>() const
 	return Info.ToSharedRef();
 }
 
+FUICommandInfo::FUICommandInfo( const FName InBindingContext )
+	: BindingContext( InBindingContext )
+	, UserInterfaceType( EUserInterfaceActionType::Button )
+	, bUseLongDisplayName( true )
+{
+	LLM_SCOPE_BYTAG(UI_Slate);
+
+	ActiveChords.Empty(2);
+	ActiveChords.Add(TSharedRef<FInputChord>(new FInputChord));
+	ActiveChords.Add(TSharedRef<FInputChord>(new FInputChord));
+
+	DefaultChords.Init(FInputChord(EKeys::Invalid, EModifierKey::None), 2);
+}
 
 
 const FText FUICommandInfo::GetInputText() const
diff --git a/Engine/Source/Runtime/Slate/Private/Framework/Docking/SDockingArea.h b/Engine/Source/Runtime/Slate/Private/Framework/Docking/SDockingArea.h
--- a/Engine/Source/Runtime/Slate/Private/Framework/Docking/SDockingArea.h
+++ b/Engine/Source/Runtime/Slate/Private/Framework/Docking/SDockingArea.h
@@ -24,7 +24,7 @@ struct FSidebarTabLists
 /**
  * Represents the root node in a hierarchy of DockNodes.
  */
-class SLATE_API SDockingArea : public SDockingSplitter
+class SDockingArea : public SDockingSplitter
 {
 public:
 
@@ -52,7 +52,7 @@ public:
 
 	SLATE_END_ARGS()
 
-	void Construct( const FArguments& InArgs, const TSharedRef<FTabManager>& InTabManager, const TSharedRef<FTabManager::FArea>& PersistentNode );
+	SLATE_API void Construct( const FArguments& InArgs, const TSharedRef<FTabManager>& InTabManager, const TSharedRef<FTabManager::FArea>& PersistentNode );
 
 	virtual Type GetNodeType() const override
 	{
@@ -60,37 +60,37 @@ public:
 	}
 	
 	/** Returns this dock area */
-	virtual TSharedPtr<SDockingArea> GetDockArea() override;
-	virtual TSharedPtr<const SDockingArea> GetDockArea() const override;
+	SLATE_API virtual TSharedPtr<SDockingArea> GetDockArea() override;
+	SLATE_API virtual TSharedPtr<const SDockingArea> GetDockArea() const override;
 
 	/** Returns the window that this dock area resides in directly and also manages */
-	TSharedPtr<SWindow> GetParentWindow() const;
+	SLATE_API TSharedPtr<SWindow> GetParentWindow() const;
 
-	virtual void OnDragEnter( const FGeometry& MyGeometry, const FDragDropEvent& DragDropEvent ) override;
-	virtual void OnDragLeave( const FDragDropEvent& DragDropEvent ) override;
+	SLATE_API virtual void OnDragEnter( const FGeometry& MyGeometry, const FDragDropEvent& DragDropEvent ) override;
+	SLATE_API virtual void OnDragLeave( const FDragDropEvent& DragDropEvent ) override;
 	
-	virtual FReply OnDrop( const FGeometry& MyGeometry, const FDragDropEvent& DragDropEvent ) override;
+	SLATE_API virtual FReply OnDrop( const FGeometry& MyGeometry, const FDragDropEvent& DragDropEvent ) override;
 
-	virtual FReply OnUserAttemptingDock( SDockingNode::RelativeDirection Direction, const FDragDropEvent& DragDropEvent ) override;
+	SLATE_API virtual FReply OnUserAttemptingDock( SDockingNode::RelativeDirection Direction, const FDragDropEvent& DragDropEvent ) override;
 
-	void OnTabFoundNewHome( const TSharedRef<SDockTab>& RelocatedTab, const TSharedRef<SWindow>& NewOwnerWindow );
+	SLATE_API void OnTabFoundNewHome( const TSharedRef<SDockTab>& RelocatedTab, const TSharedRef<SWindow>& NewOwnerWindow );
 
 	// Show the dock-from-outside dock targets
-	void ShowCross();
+	SLATE_API void ShowCross();
 
 	// Hide the dock-from-outside dock targets
-	void HideCross();
+	SLATE_API void HideCross();
 
 	/**
 	 * Removes redundant stack and splitters. Collapses any widgets that any no longer showing live content.
 	 */
-	void CleanUp( ELayoutModification RemovalMethod );
+	SLATE_API void CleanUp( ELayoutModification RemovalMethod );
 
-	void SetParentWindow( const TSharedRef<SWindow>& NewParentWindow );
+	SLATE_API void SetParentWindow( const TSharedRef<SWindow>& NewParentWindow );
 
-	virtual TSharedPtr<FTabManager::FLayoutNode> GatherPersistentLayout() const override;
+	SLATE_API virtual TSharedPtr<FTabManager::FLayoutNode> GatherPersistentLayout() const override;
 
-	TSharedRef<FTabManager> GetTabManager() const;
+	SLATE_API TSharedRef<FTabManager> GetTabManager() const;
 
 	/**
 	 * Adds a tab to a drawer in the sidebar. 
@@ -98,60 +98,60 @@ public:
 	 * @param TabToAdd	The tab to add to the sidebar
 	 * @return The location of the sidebar that the tab was added to
 	 */
-	ESidebarLocation AddTabToSidebar(TSharedRef<SDockTab> TabToAdd);
+	SLATE_API ESidebarLocation AddTabToSidebar(TSharedRef<SDockTab> TabToAdd);
 
 	/**
 	 * Restores a tab from the sidebar to its parent tab stack and removes the tab from the sidebar.  
 	 *
 	 * @return true if the tab was found in this area and restore
 	 */
-	bool RestoreTabFromSidebar(TSharedRef<SDockTab> TabToRemove);
+	SLATE_API bool RestoreTabFromSidebar(TSharedRef<SDockTab> TabToRemove);
 
 	/**
 	 * @return true if the specified tab is in the sidebar
 	 */
-	bool IsTabInSidebar(TSharedRef<SDockTab> Tab) const;
+	SLATE_API bool IsTabInSidebar(TSharedRef<SDockTab> Tab) const;
 
 	/**
  	 * Removes a tab from a sidebar
 	 * @return true if the specified tab was found and removed
 	 */
-	bool RemoveTabFromSidebar(TSharedRef<SDockTab> Tab);
+	SLATE_API bool RemoveTabFromSidebar(TSharedRef<SDockTab> Tab);
 
 	/**
 	 * Attempts to open a sidebar drawer that may the tab to open
 	 *
 	 * @true if the drawer was opened, false if the tab is not in a drawer
 	 */
-	bool TryOpenSidebarDrawer(TSharedRef<SDockTab> TabToOpen) const;
+	SLATE_API bool TryOpenSidebarDrawer(TSharedRef<SDockTab> TabToOpen) const;
 
 	/**
 	 * Adds all tabs back to a sidebar that were saved in a sidebar from a previous session
 	 */
-	void AddSidebarTabsFromRestoredLayout(const FSidebarTabLists& SidebarTabs);
+	SLATE_API void AddSidebarTabsFromRestoredLayout(const FSidebarTabLists& SidebarTabs);
 
 	/**
 	 * Gets all tabs in all sidebars in this dock area
 	 */
-	TArray<TSharedRef<SDockTab>> GetAllSidebarTabs() const;
+	SLATE_API TArray<TSharedRef<SDockTab>> GetAllSidebarTabs() const;
 
 	bool CanHaveSidebar() const { return bCanHaveSidebar; }
 protected:
-	virtual SDockingNode::ECleanupRetVal CleanUpNodes() override;
+	SLATE_API virtual SDockingNode::ECleanupRetVal CleanUpNodes() override;
 
 private:
-	EVisibility TargetCrossVisibility() const;
-	EVisibility TargetCrossCenterVisibility() const;
+	SLATE_API EVisibility TargetCrossVisibility() const;
+	SLATE_API EVisibility TargetCrossCenterVisibility() const;
 	/** Dock a tab along the outer edge of this DockArea */
-	void DockFromOutside(SDockingNode::RelativeDirection Direction, const FDragDropEvent& DragDropEvent);
+	SLATE_API void DockFromOutside(SDockingNode::RelativeDirection Direction, const FDragDropEvent& DragDropEvent);
 
 	/** We were placed in a window, and it is being destroyed */
-	void OnOwningWindowBeingDestroyed(const TSharedRef<SWindow>& WindowBeingDestroyed);
+	SLATE_API void OnOwningWindowBeingDestroyed(const TSharedRef<SWindow>& WindowBeingDestroyed);
 
 	/** We were placed in a window and it is being activated */
-	void OnOwningWindowActivated();
+	SLATE_API void OnOwningWindowActivated();
 	
-	virtual void OnLiveTabAdded() override;
+	SLATE_API virtual void OnLiveTabAdded() override;
 
 	/**
 	 * If this dock area controls a window, then we need
@@ -160,7 +160,7 @@ private:
 	 *
 	 * We also update the sidebar to account for major tabs.  Docking areas for major tabs do not have a sidebar.
 	 */
-	void UpdateWindowChromeAndSidebar();
+	SLATE_API void UpdateWindowChromeAndSidebar();
 private:	
 	/** Left and right sidebar widgets */
 	TSharedPtr<STabSidebar> LeftSidebar;
diff --git a/Engine/Source/Runtime/Slate/Private/Framework/Docking/SDockingCross.h b/Engine/Source/Runtime/Slate/Private/Framework/Docking/SDockingCross.h
--- a/Engine/Source/Runtime/Slate/Private/Framework/Docking/SDockingCross.h
+++ b/Engine/Source/Runtime/Slate/Private/Framework/Docking/SDockingCross.h
@@ -16,20 +16,20 @@ class FSlateWindowElementList;
  * gives the user a preview of what will happen if they drop on that target.
  * Dropping actually commits the layout-restructuring.
  */
-class SLATE_API SDockingCross : public SLeafWidget
+class SDockingCross : public SLeafWidget
 {
 public:
 	SLATE_BEGIN_ARGS(SDockingCross){}
 	SLATE_END_ARGS()
 
-	void Construct( const FArguments& InArgs, const TSharedPtr<class SDockingNode>& InOwnerNode );
+	SLATE_API void Construct( const FArguments& InArgs, const TSharedPtr<class SDockingNode>& InOwnerNode );
 
 	// SWidget interface
-	virtual int32 OnPaint( const FPaintArgs& Args, const FGeometry& AllottedGeometry, const FSlateRect& MyCullingRect, FSlateWindowElementList& OutDrawElements, int32 LayerId, const FWidgetStyle& InWidgetStyle, bool bParentEnabled ) const override;
-	virtual FVector2D ComputeDesiredSize(float) const override;
-	virtual void OnDragLeave( const FDragDropEvent& DragDropEvent ) override;
-	virtual FReply OnDragOver( const FGeometry& MyGeometry, const FDragDropEvent& DragDropEvent ) override;
-	virtual FReply OnDrop( const FGeometry& MyGeometry, const FDragDropEvent& DragDropEvent ) override;
+	SLATE_API virtual int32 OnPaint( const FPaintArgs& Args, const FGeometry& AllottedGeometry, const FSlateRect& MyCullingRect, FSlateWindowElementList& OutDrawElements, int32 LayerId, const FWidgetStyle& InWidgetStyle, bool bParentEnabled ) const override;
+	SLATE_API virtual FVector2D ComputeDesiredSize(float) const override;
+	SLATE_API virtual void OnDragLeave( const FDragDropEvent& DragDropEvent ) override;
+	SLATE_API virtual FReply OnDragOver( const FGeometry& MyGeometry, const FDragDropEvent& DragDropEvent ) override;
+	SLATE_API virtual FReply OnDrop( const FGeometry& MyGeometry, const FDragDropEvent& DragDropEvent ) override;
 	// End of SWidget interface
 	
 private:
diff --git a/Engine/Source/Runtime/Slate/Private/Framework/Docking/SDockingNode.h b/Engine/Source/Runtime/Slate/Private/Framework/Docking/SDockingNode.h
--- a/Engine/Source/Runtime/Slate/Private/Framework/Docking/SDockingNode.h
+++ b/Engine/Source/Runtime/Slate/Private/Framework/Docking/SDockingNode.h
@@ -15,7 +15,7 @@ class SDockingSplitter;
  * A node in the Docking/Tabbing hierarchy.
  * Any SDockingNode can be either a stack of tabs or a splitter.
  */
-class SLATE_API SDockingNode : public SCompoundWidget
+class SDockingNode : public SCompoundWidget
 {
 public:
 	enum Type
@@ -86,8 +86,8 @@ public:
 	
 
 	/** Gets the dock area that this resides in */
-	virtual TSharedPtr<SDockingArea> GetDockArea();
-	virtual TSharedPtr<const SDockingArea> GetDockArea() const;
+	SLATE_API virtual TSharedPtr<SDockingArea> GetDockArea();
+	SLATE_API virtual TSharedPtr<const SDockingArea> GetDockArea() const;
 	
 	/** Recursively searches through all children looking for child tabs */
 	virtual TArray< TSharedRef<SDockTab> > GetAllChildTabs() const {return TArray< TSharedRef<SDockTab> >();}
@@ -103,16 +103,16 @@ public:
 	 *
 	 * @return An FReply::Handled() if we successfully docked the content; FReply::Unhandled() otherwise.
 	 */
-	virtual FReply OnUserAttemptingDock( SDockingNode::RelativeDirection Direction, const FDragDropEvent& DragDropEvent );
+	SLATE_API virtual FReply OnUserAttemptingDock( SDockingNode::RelativeDirection Direction, const FDragDropEvent& DragDropEvent );
 
 	/** Should this node auto-size or be a percentage of its parent size; this setting is usually determined by users */
 	virtual SSplitter::ESizeRule GetSizeRule() const { return SSplitter::FractionOfParent; }
 
 	/** @return the numerator for the fraction of available space that this dock node should occupy. */
-	float GetSizeCoefficient() const;
+	SLATE_API float GetSizeCoefficient() const;
 
 	/** Set the coefficient size */
-	void SetSizeCoefficient( float InSizeCoefficient );
+	SLATE_API void SetSizeCoefficient( float InSizeCoefficient );
 
 	/** @returns true when this node is showing a live DockableTab; false when this nodes content is only kept around for history's sake */
 	//virtual bool IsShowingLiveTabs() const;
@@ -135,10 +135,10 @@ public:
 
 protected:
 
-	SDockingNode();
+	SLATE_API SDockingNode();
 
 	/** A live tab was added to this node or one of its descendants */
-	virtual void OnLiveTabAdded();
+	SLATE_API virtual void OnLiveTabAdded();
 
 	virtual void OnResized() {}
 
diff --git a/Engine/Source/Runtime/Slate/Private/Framework/Docking/SDockingSplitter.h b/Engine/Source/Runtime/Slate/Private/Framework/Docking/SDockingSplitter.h
--- a/Engine/Source/Runtime/Slate/Private/Framework/Docking/SDockingSplitter.h
+++ b/Engine/Source/Runtime/Slate/Private/Framework/Docking/SDockingSplitter.h
@@ -10,13 +10,13 @@
 class SDockingTabStack;
 
 /** Dynamic N-way splitter that provides the resizing functionality in the docking framework. */
-class SLATE_API SDockingSplitter : public SDockingNode
+class SDockingSplitter : public SDockingNode
 {
 public:
 	SLATE_BEGIN_ARGS(SDockingSplitter){}
 	SLATE_END_ARGS()
 
-	void Construct( const FArguments& InArgs, const TSharedRef<FTabManager::FSplitter>& PersistentNode );
+	SLATE_API void Construct( const FArguments& InArgs, const TSharedRef<FTabManager::FSplitter>& PersistentNode );
 
 	virtual Type GetNodeType() const override
 	{
@@ -30,7 +30,7 @@ public:
 	 * @param InChild      The DockNode child to add.
 	 * @param InLocation   Index at which to add; INDEX_NONE (default) adds to the end of the list
 	 */
-	void AddChildNode( const TSharedRef<SDockingNode>& InChild, int32 InLocation = INDEX_NONE);
+	SLATE_API void AddChildNode( const TSharedRef<SDockingNode>& InChild, int32 InLocation = INDEX_NONE);
 
 	/**
 	 * Replace the child InChildToReplace with Replacement
@@ -38,21 +38,21 @@ public:
 	 * @param InChildToReplace     The child of this node to replace
 	 * @param Replacement          What to replace with
 	 */
-	void ReplaceChild( const TSharedRef<SDockingNode>& InChildToReplace, const TSharedRef<SDockingNode>& Replacement );
+	SLATE_API void ReplaceChild( const TSharedRef<SDockingNode>& InChildToReplace, const TSharedRef<SDockingNode>& Replacement );
 
 	/**
 	 * Remove a child node from the splitter
 	 * 
 	 * @param ChildToRemove    The child node to remove from this splitter
 	 */
-	void RemoveChild( const TSharedRef<SDockingNode>& ChildToRemove );
+	SLATE_API void RemoveChild( const TSharedRef<SDockingNode>& ChildToRemove );
 	
 	/**
 	 * Remove a child at this index
 	 *
 	 * @param IndexOfChildToRemove    Remove a child at this index
 	 */
-	void RemoveChildAt( int32 IndexOfChildToRemove );
+	SLATE_API void RemoveChildAt( int32 IndexOfChildToRemove );
 
 	/**
 	 * Inserts NodeToPlace next to RelativeToMe.
@@ -62,55 +62,55 @@ public:
 	 * @param Direction          Where to insert relative to the other node
 	 * @param RelativeToMove     Insert relative to this node.
 	 */
-	void PlaceNode( const TSharedRef<SDockingNode>& NodeToPlace, SDockingNode::RelativeDirection Direction, const TSharedRef<SDockingNode>& RelativeToMe );
+	SLATE_API void PlaceNode( const TSharedRef<SDockingNode>& NodeToPlace, SDockingNode::RelativeDirection Direction, const TSharedRef<SDockingNode>& RelativeToMe );
 
 	/** Sets the orientation of this splitter */
-	void SetOrientation(EOrientation NewOrientation);
+	SLATE_API void SetOrientation(EOrientation NewOrientation);
 
 	/** Gets an array of all child dock nodes */
-	const TArray< TSharedRef<SDockingNode> >& GetChildNodes() const;
+	SLATE_API const TArray< TSharedRef<SDockingNode> >& GetChildNodes() const;
 
 	/** Gets an array of all child dock nodes and all of their child dock nodes and so on */
-	TArray< TSharedRef<SDockingNode> > GetChildNodesRecursively() const;
+	SLATE_API TArray< TSharedRef<SDockingNode> > GetChildNodesRecursively() const;
 	
 	/** Recursively searches through all children looking for child tabs */
-	virtual TArray< TSharedRef<SDockTab> > GetAllChildTabs() const override;
+	SLATE_API virtual TArray< TSharedRef<SDockTab> > GetAllChildTabs() const override;
 
 	/** Gets the number of tabs in all children recursively */
-	virtual int32 GetNumTabs() const override;
+	SLATE_API virtual int32 GetNumTabs() const override;
 
 	/** Should this node auto-size or be a percentage of its parent size */
-	virtual SSplitter::ESizeRule GetSizeRule() const override;
+	SLATE_API virtual SSplitter::ESizeRule GetSizeRule() const override;
 
 	/** Gets the size coefficient of a given child dock node */
-	float GetSizeCoefficientForSlot(int32 Index) const;
+	SLATE_API float GetSizeCoefficientForSlot(int32 Index) const;
 
 	/** Gets the orientation of this splitter */
-	EOrientation GetOrientation() const;
+	SLATE_API EOrientation GetOrientation() const;
 
-	virtual TSharedPtr<FTabManager::FLayoutNode> GatherPersistentLayout() const override;
+	SLATE_API virtual TSharedPtr<FTabManager::FLayoutNode> GatherPersistentLayout() const override;
 
 	/**
 	 * The tab stack which would be appropriate to use for showing the minimize/maximize/close buttons.
 	 * On Windows it is the upper right most. On Mac the upper left most.
 	 */
-	TSharedRef<SDockingTabStack> FindTabStackToHouseWindowControls() const;
+	SLATE_API TSharedRef<SDockingTabStack> FindTabStackToHouseWindowControls() const;
 
 	/** Which tab stack is appropriate for showing the app icon. */
-	TSharedRef<SDockingTabStack> FindTabStackToHouseWindowIcon() const;
+	SLATE_API TSharedRef<SDockingTabStack> FindTabStackToHouseWindowIcon() const;
 
 protected:
 
 	/** Is the docking direction (left, right, above, below) match the orientation (horizontal vs. vertical) */
-	static bool DoesDirectionMatchOrientation( SDockingNode::RelativeDirection InDirection, EOrientation InOrientation );
+	static SLATE_API bool DoesDirectionMatchOrientation( SDockingNode::RelativeDirection InDirection, EOrientation InOrientation );
 
-	static SDockingNode::ECleanupRetVal MostResponsibility( SDockingNode::ECleanupRetVal A, SDockingNode::ECleanupRetVal B );
+	static SLATE_API SDockingNode::ECleanupRetVal MostResponsibility( SDockingNode::ECleanupRetVal A, SDockingNode::ECleanupRetVal B );
 
-	virtual SDockingNode::ECleanupRetVal CleanUpNodes() override;
+	SLATE_API virtual SDockingNode::ECleanupRetVal CleanUpNodes() override;
 
-	virtual void OnResized() override;
+	SLATE_API virtual void OnResized() override;
 
-	float ComputeChildCoefficientTotal() const;
+	SLATE_API float ComputeChildCoefficientTotal() const;
 
 	enum class ETabStackToFind
 	{
@@ -119,7 +119,7 @@ protected:
 	};
 
 	/** Helper: Finds the upper left or the upper right tab stack in the hierarchy */
-	TSharedRef<SDockingNode> FindTabStack(ETabStackToFind FindMe) const;
+	SLATE_API TSharedRef<SDockingNode> FindTabStack(ETabStackToFind FindMe) const;
 
 	/** The SSplitter widget that SDockingSplitter wraps. */
 	TSharedPtr<SSplitter> Splitter;
@@ -131,7 +131,7 @@ protected:
 	 *   As the Dock area is updating after clean up of any added, removed, or newly shown Dock areas or tabs,
 	 *   the tabs/tab well may need adjusting. Use this for post clean up docked tab adjustments.
 	 */
-	void AdjustDockedTabsIfNeeded();
+	SLATE_API void AdjustDockedTabsIfNeeded();
 };
 
 
diff --git a/Engine/Source/Runtime/Slate/Private/Framework/Docking/SDockingTabStack.h b/Engine/Source/Runtime/Slate/Private/Framework/Docking/SDockingTabStack.h
--- a/Engine/Source/Runtime/Slate/Private/Framework/Docking/SDockingTabStack.h
+++ b/Engine/Source/Runtime/Slate/Private/Framework/Docking/SDockingTabStack.h
@@ -48,7 +48,7 @@ struct FDockingStackOptionalContent
  * A DockTabStack shows a row of tabs and the content of one selected tab.
  * It also supports re-arranging tabs and dragging them out for the stack.
  */
-class SLATE_API SDockingTabStack : public SDockingNode
+class SDockingTabStack : public SDockingNode
 {
 public:
 	SLATE_BEGIN_ARGS(SDockingTabStack)
@@ -61,70 +61,70 @@ public:
 		return SDockingNode::DockTabStack;
 	}
 
-	void OnLastTabRemoved();
+	SLATE_API void OnLastTabRemoved();
 
-	void OnTabClosed(const TSharedRef<SDockTab>& ClosedTab, SDockingNode::ELayoutModification RemovalMethod);
+	SLATE_API void OnTabClosed(const TSharedRef<SDockTab>& ClosedTab, SDockingNode::ELayoutModification RemovalMethod);
 
-	void OnTabRemoved( const FTabId& TabId );
+	SLATE_API void OnTabRemoved( const FTabId& TabId );
 
-	void Construct( const FArguments& InArgs, const TSharedRef<FTabManager::FStack>& PersistentNode );
+	SLATE_API void Construct( const FArguments& InArgs, const TSharedRef<FTabManager::FStack>& PersistentNode );
 
 	// TabStack methods
 
-	void OpenTab(const TSharedRef<SDockTab>& InTab, int32 InsertAtLocation = INDEX_NONE, bool bKeepInactive = false);
+	SLATE_API void OpenTab(const TSharedRef<SDockTab>& InTab, int32 InsertAtLocation = INDEX_NONE, bool bKeepInactive = false);
 
-	void AddTabWidget(const TSharedRef<SDockTab>& InTab, int32 AtLocation = INDEX_NONE, bool bKeepInactive = false);
+	SLATE_API void AddTabWidget(const TSharedRef<SDockTab>& InTab, int32 AtLocation = INDEX_NONE, bool bKeepInactive = false);
 
-	void AddSidebarTab(const TSharedRef<SDockTab>& InTab);
+	SLATE_API void AddSidebarTab(const TSharedRef<SDockTab>& InTab);
 
-	float GetTabSidebarSizeCoefficient(const TSharedRef<SDockTab>& InTab);
-	void SetTabSidebarSizeCoefficient(const TSharedRef<SDockTab>& InTab, float InSizeCoefficient);
+	SLATE_API float GetTabSidebarSizeCoefficient(const TSharedRef<SDockTab>& InTab);
+	SLATE_API void SetTabSidebarSizeCoefficient(const TSharedRef<SDockTab>& InTab, float InSizeCoefficient);
 
-	bool IsTabPinnedInSidebar(const TSharedRef<SDockTab>& InTab);
-	void SetTabPinnedInSidebar(const TSharedRef<SDockTab>& InTab, bool bPinnedInSidebar);
+	SLATE_API bool IsTabPinnedInSidebar(const TSharedRef<SDockTab>& InTab);
+	SLATE_API void SetTabPinnedInSidebar(const TSharedRef<SDockTab>& InTab, bool bPinnedInSidebar);
 
 	/** @return true if the specified tab can be moved to a sidebar */
-	bool CanMoveTabToSideBar(TSharedRef<SDockTab> Tab) const;
+	SLATE_API bool CanMoveTabToSideBar(TSharedRef<SDockTab> Tab) const;
 
 	/** Moves a specific tab from this stack into a sidebar */
-	void MoveTabToSidebar(TSharedRef<SDockTab> Tab);
+	SLATE_API void MoveTabToSidebar(TSharedRef<SDockTab> Tab);
 
 	/** Restores a specific tab from this stack to its original position */
-	void RestoreTabFromSidebar(TSharedRef<SDockTab> Tab);
+	SLATE_API void RestoreTabFromSidebar(TSharedRef<SDockTab> Tab);
 
 	/** @return All child tabs in this node */
-	const TSlotlessChildren<SDockTab>& GetTabs() const;
+	SLATE_API const TSlotlessChildren<SDockTab>& GetTabs() const;
 
-	bool HasTab( const struct FTabMatcher& TabMatcher ) const;
+	SLATE_API bool HasTab( const struct FTabMatcher& TabMatcher ) const;
 
 	/** @return the last known geometry of this TabStack */
-	FGeometry GetTabStackGeometry() const;
+	SLATE_API FGeometry GetTabStackGeometry() const;
 
-	void RemoveClosedTabsWithName( FName InName );
+	SLATE_API void RemoveClosedTabsWithName( FName InName );
 
-	bool IsShowingLiveTabs() const;
+	SLATE_API bool IsShowingLiveTabs() const;
 
-	void BringToFront( const TSharedRef<SDockTab>& TabToBringToFront );
+	SLATE_API void BringToFront( const TSharedRef<SDockTab>& TabToBringToFront );
 
 	/** Set the content that the DockNode is presenting. */
-	void SetNodeContent( const TSharedRef<SWidget>& InContent, const FDockingStackOptionalContent& OptionalContent);
+	SLATE_API void SetNodeContent( const TSharedRef<SWidget>& InContent, const FDockingStackOptionalContent& OptionalContent);
 
-	virtual FReply OnUserAttemptingDock( SDockingNode::RelativeDirection Direction, const FDragDropEvent& DragDropEvent ) override;
+	SLATE_API virtual FReply OnUserAttemptingDock( SDockingNode::RelativeDirection Direction, const FDragDropEvent& DragDropEvent ) override;
 	
 	/** Recursively searches through all children looking for child tabs */
-	virtual TArray< TSharedRef<SDockTab> > GetAllChildTabs() const override;
+	SLATE_API virtual TArray< TSharedRef<SDockTab> > GetAllChildTabs() const override;
 
 	/** Gets the number of tabs in all children recursively */
-	virtual int32 GetNumTabs() const override;
+	SLATE_API virtual int32 GetNumTabs() const override;
 
-	virtual SSplitter::ESizeRule GetSizeRule() const override;
+	SLATE_API virtual SSplitter::ESizeRule GetSizeRule() const override;
 
 	/**
 	 * Returns a bool indicating whether or not a tab is in a state where it can be hidden. A tab typically shouldn't
 	 * be hidden in certain conditions like if there is more than one tab in the well, or if it's in a floating window.
 	 *
 	 * @return returns true if the tab can be hidden, else it returns false */
-	bool CanHideTabWell() const;
+	SLATE_API bool CanHideTabWell() const;
 
 	
 	/**
@@ -132,10 +132,10 @@ public:
 	 *
 	 * @param bShouldHideTabWell the bool that defines whether or not the tab well should be hidden
 	 */
-	void SetTabWellHidden( bool bShouldHideTabWell );
-	bool IsTabWellHidden() const;
+	SLATE_API void SetTabWellHidden( bool bShouldHideTabWell );
+	SLATE_API bool IsTabWellHidden() const;
 
-	virtual TSharedPtr<FTabManager::FLayoutNode> GatherPersistentLayout() const override;
+	SLATE_API virtual TSharedPtr<FTabManager::FLayoutNode> GatherPersistentLayout() const override;
 
 	/** Elements for which we might want to reserve space. */
 	enum class EChromeElement
@@ -148,21 +148,21 @@ public:
 	 * Remove all the space that might have been reserved for
 	 * various window chrome elements (app icons, minimize, close, etc.)
 	 */
-	void ClearReservedSpace();
+	SLATE_API void ClearReservedSpace();
 	/**
 	 * Add some extra padding so that the corresponding window chrome element
 	 * does not overlap our tabs.
 	 */
-	void ReserveSpaceForWindowChrome(EChromeElement Element, bool bIncludePaddingForMenuBar, bool bOnlyMinorTabs);
+	SLATE_API void ReserveSpaceForWindowChrome(EChromeElement Element, bool bIncludePaddingForMenuBar, bool bOnlyMinorTabs);
 
 public:
 	/** SWidget interface */
-	virtual FReply OnDragOver( const FGeometry& MyGeometry, const FDragDropEvent& DragDropEvent ) override;
-	virtual void OnDragLeave( const FDragDropEvent& DragDropEvent ) override;
-	virtual FReply OnDrop( const FGeometry& MyGeometry, const FDragDropEvent& DragDropEvent ) override;
-	virtual void OnFocusChanging( const FWeakWidgetPath& PreviousFocusPath, const FWidgetPath& NewWidgetPath, const FFocusEvent& InFocusEvent ) override;
-	virtual FReply OnMouseButtonDown( const FGeometry& MyGeometry, const FPointerEvent& MouseEvent ) override;
-	virtual FReply OnKeyDown(const FGeometry& MyGeometry, const FKeyEvent& InKeyEvent) override;
+	SLATE_API virtual FReply OnDragOver( const FGeometry& MyGeometry, const FDragDropEvent& DragDropEvent ) override;
+	SLATE_API virtual void OnDragLeave( const FDragDropEvent& DragDropEvent ) override;
+	SLATE_API virtual FReply OnDrop( const FGeometry& MyGeometry, const FDragDropEvent& DragDropEvent ) override;
+	SLATE_API virtual void OnFocusChanging( const FWeakWidgetPath& PreviousFocusPath, const FWidgetPath& NewWidgetPath, const FFocusEvent& InFocusEvent ) override;
+	SLATE_API virtual FReply OnMouseButtonDown( const FGeometry& MyGeometry, const FPointerEvent& MouseEvent ) override;
+	SLATE_API virtual FReply OnKeyDown(const FGeometry& MyGeometry, const FKeyEvent& InKeyEvent) override;
 	virtual bool SupportsKeyboardFocus() const override { return true; }
 
 protected:
@@ -173,7 +173,7 @@ protected:
 		return EWindowZone::TitleBar;
 	}
 
-	void CloseForegroundTab();
+	SLATE_API void CloseForegroundTab();
 
 	enum ETabsToClose
 	{
@@ -182,27 +182,27 @@ protected:
 		CloseAllTabs
 	};
 
-	void CloseTabsToRightOfForegroundTab(ETabsToClose TabsToClose);
+	SLATE_API void CloseTabsToRightOfForegroundTab(ETabsToClose TabsToClose);
 
 	/**
 	 * Close all the background tabs.
 	 *
 	 * @param  TabsToClose   Close just document tabs or both document and tool tabs.
 	 */
-	void CloseAllButForegroundTab(ETabsToClose TabsToClose);
+	SLATE_API void CloseAllButForegroundTab(ETabsToClose TabsToClose);
 
-	FReply TabWellRightClicked( const FGeometry& TabWellGeometry, const FPointerEvent& MouseEvent );
+	SLATE_API FReply TabWellRightClicked( const FGeometry& TabWellGeometry, const FPointerEvent& MouseEvent );
 
-	virtual SDockingNode::ECleanupRetVal CleanUpNodes() override;
+	SLATE_API virtual SDockingNode::ECleanupRetVal CleanUpNodes() override;
 
-	int32 OpenPersistentTab( const FTabId& TabId, int32 OpenLocationAmongActiveTabs = INDEX_NONE );
+	SLATE_API int32 OpenPersistentTab( const FTabId& TabId, int32 OpenLocationAmongActiveTabs = INDEX_NONE );
 
-	int32 ClosePersistentTab( const FTabId& TabId );
+	SLATE_API int32 ClosePersistentTab( const FTabId& TabId );
 
-	void RemovePersistentTab( const FTabId& TabId );
+	SLATE_API void RemovePersistentTab( const FTabId& TabId );
 
 	/** Overridden from SDockingNode */
-	virtual void SetParentNode( TSharedRef<class SDockingSplitter> InParent ) override;
+	SLATE_API virtual void SetParentNode( TSharedRef<class SDockingSplitter> InParent ) override;
 
 private:
 	
@@ -217,46 +217,46 @@ private:
 	 *
 	 * @return	The newly created empty SDockingTabStack, ready for a tab to be added to it
 	 */
-	TSharedRef< SDockingTabStack > CreateNewTabStackBySplitting( const SDockingNode::RelativeDirection Direction );
+	SLATE_API TSharedRef< SDockingTabStack > CreateNewTabStackBySplitting( const SDockingNode::RelativeDirection Direction );
 
 	/** What should the content area look like for the current tab? */
-	const FSlateBrush* GetContentAreaBrush() const;
+	SLATE_API const FSlateBrush* GetContentAreaBrush() const;
 
 	/** How much padding to show around the content currently being presented */
-	FMargin GetContentPadding() const;
+	SLATE_API FMargin GetContentPadding() const;
 
 	/** Depending on the tabs we put into the tab well, we want a different background brush. */
-	const FSlateBrush* GetTabWellBrush() const;
+	SLATE_API const FSlateBrush* GetTabWellBrush() const;
 
 	/** Show the tab well? */
-	EVisibility GetTabWellVisibility() const;
+	SLATE_API EVisibility GetTabWellVisibility() const;
 
 	/** Show the stuff needed to unhide the tab well? */
-	EVisibility GetUnhideButtonVisibility() const;
+	SLATE_API EVisibility GetUnhideButtonVisibility() const;
 	
 	/** Show/Hide the tab well; do it smoothly with an animation */
-	void ToggleTabWellVisibility();
+	SLATE_API void ToggleTabWellVisibility();
 
 	/** Moves the foreground tab to a sidebar */
-	void MoveForegroundTabToSidebar();
+	SLATE_API void MoveForegroundTabToSidebar();
 
 	/** Unhides the tab well, revealing all tab headers */
-	FReply UnhideTabWell();
+	SLATE_API FReply UnhideTabWell();
 
 	/** Only allow closing the tab well when the tab allows it. */
-	bool CanCloseForegroundTab() const;
+	SLATE_API bool CanCloseForegroundTab() const;
 
 	/** Only allow closing tabs to the right when there is more than one tab open, the tab is of type Document or Major, and the tab is not furthest to the right. */
-	bool CanCloseTabsToRightOfForegroundTab() const;
+	SLATE_API bool CanCloseTabsToRightOfForegroundTab() const;
 
 	/** Only allow closing all other tabs when there are more then one tab open and the tab is of type Document or Major. */
-	bool CanCloseAllButForegroundTab() const;
+	SLATE_API bool CanCloseAllButForegroundTab() const;
 
 	/** @return true if the foreground tab can be moved to a sidebar */
-	bool CanMoveForegroundTabToSidebar() const;
+	SLATE_API bool CanMoveForegroundTabToSidebar() const;
 
 	/** @return true if the specified tab can be moved to a sidebar */
-	bool IsTabAllowedInSidebar(TSharedPtr<SDockTab> Tab) const;
+	SLATE_API bool IsTabAllowedInSidebar(TSharedPtr<SDockTab> Tab) const;
 
 	/** The tab well widget shows all tabs, keeps track of the selected tab, allows tab rearranging, etc. */
 	TSharedPtr<class SDockingTabWell> TabWell;
@@ -272,13 +272,13 @@ private:
 
 	FOverlayManagement OverlayManagement;
 
-	TSharedRef<SWidget> MakeContextMenu();
+	SLATE_API TSharedRef<SWidget> MakeContextMenu();
 
 	/** Show the docking cross */
-	void ShowCross();
+	SLATE_API void ShowCross();
 
 	/** Hide the docking cross */
-	void HideCross();
+	SLATE_API void HideCross();
 
 	/** Document Areas do not disappear when out of tabs, and instead say 'Document Area' */
 	bool bIsDocumentArea;
@@ -287,41 +287,41 @@ private:
 	FCurveSequence ShowHideTabWell;
 
 	/** Grabs the scaling factor for the tab well size from the tab well animation. */
-	FVector2D GetTabWellScale() const;
+	SLATE_API FVector2D GetTabWellScale() const;
 
 	/** Get the scale for the button that unhides the tab well */
-	FVector2D GetUnhideTabWellButtonScale() const;
+	SLATE_API FVector2D GetUnhideTabWellButtonScale() const;
 	/** Get the opacity for the button that unhides the tab well */
-	FSlateColor GetUnhideTabWellButtonOpacity() const;
+	SLATE_API FSlateColor GetUnhideTabWellButtonOpacity() const;
 
 	/** Gets the background behind the tab stack */
-	const FSlateBrush* GetTabStackBorderImage() const;
+	SLATE_API const FSlateBrush* GetTabStackBorderImage() const;
 
 	/** Visibility of TitleBar spacer based on maximize/restore status of the window.
 	 ** This gives us a little more space to grab the title bar when the window is not maximized
 	*/
-	EVisibility GetMaximizeSpacerVisibility() const;
+	SLATE_API EVisibility GetMaximizeSpacerVisibility() const;
 
 	/** Bind tab commands into the ActionList */
-	void BindTabCommands();
+	SLATE_API void BindTabCommands();
 
 	/** Attempts to close the foreground tab when the CloseMajorTab command is executed */
-	void ExecuteCloseMajorTabCommand();
+	SLATE_API void ExecuteCloseMajorTabCommand();
 
 	/** Attempts to find the foreground tab that can be closed by the CloseMajorTab command */
-	bool CanExecuteCloseMajorTabCommand();
+	SLATE_API bool CanExecuteCloseMajorTabCommand();
 
 	/** Attempts to close the active tab when the CloseMinorTab command is executed */
-	void ExecuteCloseMinorTabCommand();
+	SLATE_API void ExecuteCloseMinorTabCommand();
 
 	/** Attempts to find the active tab that can be closed by the CloseMinorTab command */
-	bool CanExecuteCloseMinorTabCommand();
+	SLATE_API bool CanExecuteCloseMinorTabCommand();
 
 	/** Called when this stack is resized. */
-	virtual void OnResized() override;
+	SLATE_API virtual void OnResized() override;
 
 #if DEBUG_TAB_MANAGEMENT
-	FString ShowPersistentTabs() const;
+	SLATE_API FString ShowPersistentTabs() const;
 #endif
 
 	/** Tab command list */
diff --git a/Engine/Source/Runtime/Slate/Private/Framework/Docking/SDockingTarget.h b/Engine/Source/Runtime/Slate/Private/Framework/Docking/SDockingTarget.h
--- a/Engine/Source/Runtime/Slate/Private/Framework/Docking/SDockingTarget.h
+++ b/Engine/Source/Runtime/Slate/Private/Framework/Docking/SDockingTarget.h
@@ -13,7 +13,7 @@
  * gives the user a preview of what will happen if they drop on that target.
  * Dropping actually commits the layout-restructuring.
  */
-class SLATE_API SDockingTarget : public SBorder
+class SDockingTarget : public SBorder
 {
 public:
 	SLATE_BEGIN_ARGS(SDockingTarget)
@@ -26,19 +26,19 @@ public:
 		SLATE_ARGUMENT( SDockingNode::RelativeDirection, DockDirection )
 	SLATE_END_ARGS()
 
-	void Construct( const FArguments& InArgs );
+	SLATE_API void Construct( const FArguments& InArgs );
 
 	// SWidget interface
-	virtual void OnDragEnter( const FGeometry& MyGeometry, const FDragDropEvent& DragDropEvent ) override;
-	virtual void OnDragLeave( const FDragDropEvent& DragDropEvent ) override;
-	virtual FReply OnDrop( const FGeometry& MyGeometry, const FDragDropEvent& DragDropEvent ) override;
+	SLATE_API virtual void OnDragEnter( const FGeometry& MyGeometry, const FDragDropEvent& DragDropEvent ) override;
+	SLATE_API virtual void OnDragLeave( const FDragDropEvent& DragDropEvent ) override;
+	SLATE_API virtual FReply OnDrop( const FGeometry& MyGeometry, const FDragDropEvent& DragDropEvent ) override;
 	// End of SWidget interface
 
 	/** @return the DockNode which this target represent */
-	TSharedPtr<SDockingNode> GetOwner() const;
+	SLATE_API TSharedPtr<SDockingNode> GetOwner() const;
 	
 	/** @return the direction which we represent (relative to the OwnerTabStack) */
-	SDockingNode::RelativeDirection GetDockDirection() const;
+	SLATE_API SDockingNode::RelativeDirection GetDockDirection() const;
 
 private:
 	const FSlateBrush* GetTargetImage() const;
diff --git a/Engine/Source/Runtime/Slate/Private/Framework/MultiBox/SButtonRowBlock.h b/Engine/Source/Runtime/Slate/Private/Framework/MultiBox/SButtonRowBlock.h
--- a/Engine/Source/Runtime/Slate/Private/Framework/MultiBox/SButtonRowBlock.h
+++ b/Engine/Source/Runtime/Slate/Private/Framework/MultiBox/SButtonRowBlock.h
@@ -81,7 +81,7 @@ private:
 /**
  * Tool bar button MultiBlock widget
  */
-class SLATE_API SButtonRowBlock
+class SButtonRowBlock
 	: public SMultiBlockBaseWidget
 {
 
@@ -94,14 +94,14 @@ public:
 	/**
 	 * Builds this MultiBlock widget up from the MultiBlock associated with it
 	 */
-	virtual void BuildMultiBlockWidget(const ISlateStyle* StyleSet, const FName& StyleName) override;
+	SLATE_API virtual void BuildMultiBlockWidget(const ISlateStyle* StyleSet, const FName& StyleName) override;
 
 	/**
 	 * Construct this widget
 	 *
 	 * @param	InArgs	The declaration data for this widget
 	 */
-	void Construct( const FArguments& InArgs );
+	SLATE_API void Construct( const FArguments& InArgs );
 
 
 protected:
@@ -110,34 +110,34 @@ protected:
 	/**
 	 * Called by Slate when this tool bar button's button is clicked
 	 */
-	FReply OnClicked();
+	SLATE_API FReply OnClicked();
 
 
 	/**
 	 * Called by Slate when this tool bar check box button is toggled
 	 */
-	void OnCheckStateChanged( const ECheckBoxState NewCheckedState );
+	SLATE_API void OnCheckStateChanged( const ECheckBoxState NewCheckedState );
 
 	/**
 	 * Called by slate to determine if this button should appear checked
 	 *
 	 * @return ECheckBoxState::Checked if it should be checked, ECheckBoxState::Unchecked if not.
 	 */
-	ECheckBoxState OnIsChecked() const;
+	SLATE_API ECheckBoxState OnIsChecked() const;
 
 	/**
 	 * Called by Slate to determine if this button is enabled
 	 * 
 	 * @return True if the menu entry is enabled, false otherwise
 	 */
-	bool IsEnabled() const;
+	SLATE_API bool IsEnabled() const;
 
 	/**
 	 * Called by Slate to determine if this button is visible
 	 *
 	 * @return EVisibility::Visible or EVisibility::Collapsed, depending on if the button should be displayed
 	 */
-	EVisibility GetVisibility() const;
+	SLATE_API EVisibility GetVisibility() const;
 private:
 	
 	/** For check toggle buttons, this is the widget that gets toggled. */
diff --git a/Engine/Source/Runtime/Slate/Private/Framework/MultiBox/SEditableTextBlock.h b/Engine/Source/Runtime/Slate/Private/Framework/MultiBox/SEditableTextBlock.h
--- a/Engine/Source/Runtime/Slate/Private/Framework/MultiBox/SEditableTextBlock.h
+++ b/Engine/Source/Runtime/Slate/Private/Framework/MultiBox/SEditableTextBlock.h
@@ -76,7 +76,7 @@ private:
 /**
  * Editable text block widget
  */
-class SLATE_API SEditableTextBlock
+class SEditableTextBlock
 	: public SMultiBlockBaseWidget
 {
 
@@ -88,7 +88,7 @@ public:
 	/**
 	 * Builds this MultiBlock widget up from the MultiBlock associated with it
 	 */
-	virtual void BuildMultiBlockWidget(const ISlateStyle* StyleSet, const FName& StyleName) override;
+	SLATE_API virtual void BuildMultiBlockWidget(const ISlateStyle* StyleSet, const FName& StyleName) override;
 
 	/**
 	 * Construct this widget
@@ -104,5 +104,5 @@ protected:
 	 * 
 	 * @return true if the widget is enabled, false otherwise
 	 */
-	bool IsEnabled() const;
+	SLATE_API bool IsEnabled() const;
 };
diff --git a/Engine/Source/Runtime/Slate/Private/Framework/MultiBox/SGroupMarkerBlock.h b/Engine/Source/Runtime/Slate/Private/Framework/MultiBox/SGroupMarkerBlock.h
--- a/Engine/Source/Runtime/Slate/Private/Framework/MultiBox/SGroupMarkerBlock.h
+++ b/Engine/Source/Runtime/Slate/Private/Framework/MultiBox/SGroupMarkerBlock.h
@@ -44,7 +44,7 @@ private:
 /**
  * Group Marker MultiBlock widget
  */
-class SLATE_API SGroupMarkerBlock 
+class SGroupMarkerBlock 
 	: public SMultiBlockBaseWidget
 {
 
@@ -54,7 +54,7 @@ public:
 	SLATE_END_ARGS()
 
 	/** FMultiBlock interface */
-	virtual void BuildMultiBlockWidget(const ISlateStyle* StyleSet, const FName& StyleName) override;
+	SLATE_API virtual void BuildMultiBlockWidget(const ISlateStyle* StyleSet, const FName& StyleName) override;
 
 	/**
 	 * Construct this widget
diff --git a/Engine/Source/Runtime/Slate/Private/Framework/MultiBox/SHeadingBlock.h b/Engine/Source/Runtime/Slate/Private/Framework/MultiBox/SHeadingBlock.h
--- a/Engine/Source/Runtime/Slate/Private/Framework/MultiBox/SHeadingBlock.h
+++ b/Engine/Source/Runtime/Slate/Private/Framework/MultiBox/SHeadingBlock.h
@@ -45,7 +45,7 @@ private:
 /**
  * Heading MultiBlock widget
  */
-class SLATE_API SHeadingBlock
+class SHeadingBlock
 	: public SMultiBlockBaseWidget
 {
 
@@ -59,7 +59,7 @@ public:
 	/**
 	 * Builds this MultiBlock widget up from the MultiBlock associated with it
 	 */
-	virtual void BuildMultiBlockWidget(const ISlateStyle* StyleSet, const FName& StyleName) override;
+	SLATE_API virtual void BuildMultiBlockWidget(const ISlateStyle* StyleSet, const FName& StyleName) override;
 
 
 	/**
@@ -67,5 +67,5 @@ public:
 	 *
 	 * @param	InArgs	The declaration data for this widget
 	 */
-	void Construct( const FArguments& InArgs );
+	SLATE_API void Construct( const FArguments& InArgs );
 };
diff --git a/Engine/Source/Runtime/Slate/Private/Framework/MultiBox/SMenuEntryBlock.h b/Engine/Source/Runtime/Slate/Private/Framework/MultiBox/SMenuEntryBlock.h
--- a/Engine/Source/Runtime/Slate/Private/Framework/MultiBox/SMenuEntryBlock.h
+++ b/Engine/Source/Runtime/Slate/Private/Framework/MultiBox/SMenuEntryBlock.h
@@ -200,7 +200,7 @@ private:
 /**
  * Menu entry MultiBlock widget
  */
-class SLATE_API SMenuEntryBlock
+class SMenuEntryBlock
 	: public SMultiBlockBaseWidget
 {
 
@@ -214,21 +214,21 @@ public:
 	/**
 	 * Builds this MultiBlock widget up from the MultiBlock associated with it
 	 */
-	virtual void BuildMultiBlockWidget(const ISlateStyle* StyleSet, const FName& StyleName) override;
+	SLATE_API virtual void BuildMultiBlockWidget(const ISlateStyle* StyleSet, const FName& StyleName) override;
 
 	/**
 	 * Construct this widget
 	 *
 	 * @param	InArgs	The declaration data for this widget
 	 */
-	void Construct( const FArguments& InArgs );
+	SLATE_API void Construct( const FArguments& InArgs );
 
 	/**
 	 * Called to create content for a pull-down or sub-menu window when it's summoned by the user
 	 *
 	 * @return	The widget content for the new menu
 	 */
-	TSharedRef< SWidget > MakeNewMenuWidget() const;
+	SLATE_API TSharedRef< SWidget > MakeNewMenuWidget() const;
 
 
 protected:
@@ -257,51 +257,51 @@ protected:
 	/**
 	 * Called by Slate when this menu entry's button is clicked
 	 */
-	FReply OnMenuItemButtonClicked();
+	SLATE_API FReply OnMenuItemButtonClicked();
 
 	/**
 	 * Called when a checkbox in the menu item or the menu item itself is clicked
 	 *
 	 * @param bCheckBoxClicked	true if a check box was clicked and not the menu item.  We dont close the menu when a check box is clicked
 	 */
-	void OnClicked( bool bCheckBoxClicked );
+	SLATE_API void OnClicked( bool bCheckBoxClicked );
 
 	/**
 	 * Called by Slate to determine if this menu entry is enabled
 	 * 
 	 * @return True if the menu entry is enabled, false otherwise
 	 */
-	bool IsEnabled() const;
+	SLATE_API bool IsEnabled() const;
 
 	/**
 	 * Called by Slate to determine if this menu entry is enabled (during menu editing)
 	 * 
 	 * @return True if the menu entry is enabled, false otherwise
 	 */
-	bool IsEnabledDuringEditMode() const;
+	SLATE_API bool IsEnabledDuringEditMode() const;
 
 	/**
 	 * Called by Slate when this check box button is toggled in a menu entry
 	 */
-	void OnCheckStateChanged( const ECheckBoxState NewCheckedState );
+	SLATE_API void OnCheckStateChanged( const ECheckBoxState NewCheckedState );
 
 	/**
 	 * Called by slate to determine if this menu entry should appear checked
 	 *
 	 * @return true if it should be checked, false if not.
 	 */
-	ECheckBoxState IsChecked() const;
+	SLATE_API ECheckBoxState IsChecked() const;
 
 	/**
 	 * In the case that we have an icon to show.  This function is called to get the image that indicates the menu item should appear checked
 	 * If we can show an actual check box, this function is not called
 	 */
-	const FSlateBrush* OnGetCheckImage() const;
+	SLATE_API const FSlateBrush* OnGetCheckImage() const;
 
 	// SWidget interface
-	virtual void OnMouseEnter( const FGeometry& MyGeometry, const FPointerEvent& MouseEvent ) override;
-	virtual void OnMouseLeave( const FPointerEvent& MouseEvent ) override;
-	virtual FReply OnKeyDown( const FGeometry& MyGeometry, const FKeyEvent& KeyEvent ) override;
+	SLATE_API virtual void OnMouseEnter( const FGeometry& MyGeometry, const FPointerEvent& MouseEvent ) override;
+	SLATE_API virtual void OnMouseLeave( const FPointerEvent& MouseEvent ) override;
+	SLATE_API virtual FReply OnKeyDown( const FGeometry& MyGeometry, const FKeyEvent& KeyEvent ) override;
 	//virtual void Tick( const FGeometry& AllottedGeometry, const double InCurrentTime, const float InDeltaTime ) override;
 	// End of SWidget interface
 
@@ -311,8 +311,8 @@ protected:
 	 *
 	 * @return	The appropriate border to use
 	 */
-	const FSlateBrush* GetMenuBarButtonBorder() const;
-	FSlateColor GetMenuBarForegroundColor() const;
+	SLATE_API const FSlateBrush* GetMenuBarButtonBorder() const;
+	SLATE_API FSlateColor GetMenuBarForegroundColor() const;
 
 	/**
 	 * Called to create content for a pull-down menu widget
@@ -320,7 +320,7 @@ protected:
 	 * @param InBuildParams	Parameters for how to build the widget
 	 * @return	The widget content for the new menu
 	 */
-	TSharedRef< SWidget > BuildMenuBarWidget( const FMenuEntryBuildParams& InBuildParams );
+	SLATE_API TSharedRef< SWidget > BuildMenuBarWidget( const FMenuEntryBuildParams& InBuildParams );
 	
 	/**
 	* Finds the STextBlock that gets displayed in the UI
@@ -328,7 +328,7 @@ protected:
 	* @param Content	Widget to check for an STextBlock
 	* @return	The STextBlock widget found
 	*/
-	TSharedRef<SWidget> FindTextBlockWidget(TSharedRef<SWidget> Content);
+	SLATE_API TSharedRef<SWidget> FindTextBlockWidget(TSharedRef<SWidget> Content);
 
 	/**
 	 * Called to create content for a menu entry inside a pull-down, context, or sub-menu
@@ -336,7 +336,7 @@ protected:
 	 * @param InBuildParams	Parameters for how to build the widget
 	 * @return The widget content for the new menu
 	 */
-	TSharedRef< SWidget > BuildMenuEntryWidget( const FMenuEntryBuildParams& InBuildParams );
+	SLATE_API TSharedRef< SWidget > BuildMenuEntryWidget( const FMenuEntryBuildParams& InBuildParams );
 
 	/**
 	 * Called to create content for a sub-menu widget (entry in a menu that opens another menu to the right of it)
@@ -344,7 +344,7 @@ protected:
 	 * @param InBuildParams	Parameters for how to build the widget
 	 * @return The widget content for the new menu
 	 */
-	TSharedRef< SWidget > BuildSubMenuWidget( const FMenuEntryBuildParams& InBuildParams );
+	SLATE_API TSharedRef< SWidget > BuildSubMenuWidget( const FMenuEntryBuildParams& InBuildParams );
 
 	/**
 	 * Requests that the sub-menu associated with this widget be toggled on or off.
@@ -353,18 +353,18 @@ protected:
 	 * @param bOpenMenu	true to open the menu, false to close the menu if one is currently open
 	 * @param bClobber true if we want to open a menu when another menu is already open
 	 */
-	void RequestSubMenuToggle( bool bOpenMenu, const bool bClobber );
+	SLATE_API void RequestSubMenuToggle( bool bOpenMenu, const bool bClobber );
 
 	/**
 	 * Cancels any open requests to toggle a sub-menu       
 	 */
-	void CancelPendingSubMenu();
+	SLATE_API void CancelPendingSubMenu();
 
 	/**
 	 * Returns whether or the sub-menu entry should appear hovered.  If the sub-menu is open we will always show the menu as hovered to indicate which sub-menu is open
 	 * In the case that the user is interacting in this menu we do not show the menu as hovered because we need to show what the user is actually selecting
 	 */
-	bool ShouldSubMenuAppearHovered() const;
+	SLATE_API bool ShouldSubMenuAppearHovered() const;
 	
 	/**
 	 * Called to query the tool tip text for this widget, but will return an empty text for menu bar items
@@ -374,17 +374,17 @@ protected:
 	 *
 	 * @return	Tool tip text, or an empty text if filtered out
 	 */
-	FText GetFilteredToolTipText( TAttribute<FText> ToolTipText ) const;
+	SLATE_API FText GetFilteredToolTipText( TAttribute<FText> ToolTipText ) const;
 
 
 	// Gets the visibility of the menu item
-	EVisibility GetVisibility() const;
+	SLATE_API EVisibility GetVisibility() const;
 
 	/** Get the selection color when the entry is hovered */
-	FSlateColor TintOnHover() const;
+	SLATE_API FSlateColor TintOnHover() const;
 
 	/** Get the inverted foreground color when the entry is hovered */
-	FSlateColor InvertOnHover() const;
+	SLATE_API FSlateColor InvertOnHover() const;
 
 	/** Updates state machine for sub-menu opening logic.  Called in the widget's Tick as well as on demand in some cases. */
 	//void UpdateSubMenuState();
diff --git a/Engine/Source/Runtime/Slate/Private/Framework/MultiBox/SMenuSeparatorBlock.h b/Engine/Source/Runtime/Slate/Private/Framework/MultiBox/SMenuSeparatorBlock.h
--- a/Engine/Source/Runtime/Slate/Private/Framework/MultiBox/SMenuSeparatorBlock.h
+++ b/Engine/Source/Runtime/Slate/Private/Framework/MultiBox/SMenuSeparatorBlock.h
@@ -39,7 +39,7 @@ private:
 /**
  * Menu separator MultiBlock widget
  */
-class SLATE_API SMenuSeparatorBlock
+class SMenuSeparatorBlock
 	: public SMultiBlockBaseWidget
 {
 
@@ -53,7 +53,7 @@ public:
 	/**
 	 * Builds this MultiBlock widget up from the MultiBlock associated with it
 	 */
-	virtual void BuildMultiBlockWidget(const ISlateStyle* StyleSet, const FName& StyleName) override;
+	SLATE_API virtual void BuildMultiBlockWidget(const ISlateStyle* StyleSet, const FName& StyleName) override;
 
 
 	/**
@@ -61,7 +61,7 @@ public:
 	 *
 	 * @param	InArgs	The declaration data for this widget
 	 */
-	void Construct( const FArguments& InArgs );
+	SLATE_API void Construct( const FArguments& InArgs );
 
 
 private:
diff --git a/Engine/Source/Runtime/Slate/Private/Framework/MultiBox/SToolBarSeparatorBlock.h b/Engine/Source/Runtime/Slate/Private/Framework/MultiBox/SToolBarSeparatorBlock.h
--- a/Engine/Source/Runtime/Slate/Private/Framework/MultiBox/SToolBarSeparatorBlock.h
+++ b/Engine/Source/Runtime/Slate/Private/Framework/MultiBox/SToolBarSeparatorBlock.h
@@ -42,7 +42,7 @@ private:
 /**
  * Toolbar separator MultiBlock widget
  */
-class SLATE_API SToolBarSeparatorBlock
+class SToolBarSeparatorBlock
 	: public SMultiBlockBaseWidget
 {
 
@@ -56,7 +56,7 @@ public:
 	/**
 	 * Builds this MultiBlock widget up from the MultiBlock associated with it
 	 */
-	virtual void BuildMultiBlockWidget(const ISlateStyle* StyleSet, const FName& StyleName) override;
+	SLATE_API virtual void BuildMultiBlockWidget(const ISlateStyle* StyleSet, const FName& StyleName) override;
 
 
 	/**
@@ -64,7 +64,7 @@ public:
 	 *
 	 * @param	InArgs	The declaration data for this widget
 	 */
-	void Construct( const FArguments& InArgs );
+	SLATE_API void Construct( const FArguments& InArgs );
 
 
 private:
diff --git a/Engine/Source/Runtime/Slate/Private/Framework/MultiBox/SUniformToolbarPanel.h b/Engine/Source/Runtime/Slate/Private/Framework/MultiBox/SUniformToolbarPanel.h
--- a/Engine/Source/Runtime/Slate/Private/Framework/MultiBox/SUniformToolbarPanel.h
+++ b/Engine/Source/Runtime/Slate/Private/Framework/MultiBox/SUniformToolbarPanel.h
@@ -21,7 +21,7 @@ class ISlateStyle;
  * This panel evenly divides up the available space between all of its children but allows for a min and max threshold to be supplied so that
  * very large children do not force siblings to grow unnecessarily large and very small children are not forced to be unnecessarily large to maintain uniformity
  */
-class SLATE_API SUniformToolbarPanel : public SPanel
+class SUniformToolbarPanel : public SPanel
 {
 public:
 	/** Stores the per-child info for this panel type */
@@ -48,7 +48,7 @@ public:
 		}
 	};
 
-	SUniformToolbarPanel();
+	SLATE_API SUniformToolbarPanel();
 
 	SLATE_BEGIN_ARGS(SUniformToolbarPanel)
 		: _Orientation(EOrientation::Orient_Horizontal)
@@ -94,20 +94,20 @@ public:
 
 	SLATE_END_ARGS()
 
-	void Construct( const FArguments& InArgs );
+	SLATE_API void Construct( const FArguments& InArgs );
 
 	//~ Begin SPanel Interface	
-	virtual void OnArrangeChildren( const FGeometry& AllottedGeometry, FArrangedChildren& ArrangedChildren ) const override;
-	virtual FChildren* GetChildren() override;
+	SLATE_API virtual void OnArrangeChildren( const FGeometry& AllottedGeometry, FArrangedChildren& ArrangedChildren ) const override;
+	SLATE_API virtual FChildren* GetChildren() override;
 	//~ End SPanel Interface
 
 	/** See SlotPadding attribute */
-	void SetSlotPadding(TAttribute<FMargin> InSlotPadding);
+	SLATE_API void SetSlotPadding(TAttribute<FMargin> InSlotPadding);
 
 	/**
 	 * Used by declarative syntax to create a Slot in the specified Column, Row.
 	 */
-	static FSlot::FSlotArguments Slot();
+	static SLATE_API FSlot::FSlotArguments Slot();
 
 	using FScopedWidgetSlotArguments = TPanelChildren<FSlot>::FScopedWidgetSlotArguments;
 	/**
@@ -115,7 +115,7 @@ public:
 	 *
 	 * @return A reference to the newly-added slot
 	 */
-	FScopedWidgetSlotArguments AddSlot();
+	SLATE_API FScopedWidgetSlotArguments AddSlot();
 	
 	/**
 	 * Removes a slot from this panel which contains the specified SWidget
@@ -123,7 +123,7 @@ public:
 	 * @param SlotWidget The widget to match when searching through the slots
 	 * @returns The true if the slot was removed and false if no slot was found matching the widget
 	 */
-	bool RemoveSlot( const TSharedRef<SWidget>& SlotWidget );
+	SLATE_API bool RemoveSlot( const TSharedRef<SWidget>& SlotWidget );
 	
 	/**
 	 * Return the index to the first widget cli
@@ -131,7 +131,7 @@ public:
 	int32 GetClippedIndex() const { return ClippedIndex; }
 
 protected:
-	virtual FVector2D ComputeDesiredSize(float) const override;
+	SLATE_API virtual FVector2D ComputeDesiredSize(float) const override;
 private:
 	/** The button that is displayed when the toolbar is clipped */
 	TSharedPtr<SComboButton> Dropdown;
diff --git a/Engine/Source/Runtime/Slate/Private/Framework/MultiBox/SWidgetBlock.h b/Engine/Source/Runtime/Slate/Private/Framework/MultiBox/SWidgetBlock.h
--- a/Engine/Source/Runtime/Slate/Private/Framework/MultiBox/SWidgetBlock.h
+++ b/Engine/Source/Runtime/Slate/Private/Framework/MultiBox/SWidgetBlock.h
@@ -69,7 +69,7 @@ private:
 /**
  * Arbitrary Widget MultiBlock widget
  */
-class SLATE_API SWidgetBlock
+class SWidgetBlock
 	: public SMultiBlockBaseWidget
 {
 
@@ -82,17 +82,17 @@ public:
 	/**
 	 * Builds this MultiBlock widget up from the MultiBlock associated with it
 	 */
-	virtual void BuildMultiBlockWidget(const ISlateStyle* StyleSet, const FName& StyleName) override;
+	SLATE_API virtual void BuildMultiBlockWidget(const ISlateStyle* StyleSet, const FName& StyleName) override;
 
 	/**
 	 * Construct this widget
 	 *
 	 * @param	InArgs	The declaration data for this widget
 	 */
-	void Construct( const FArguments& InArgs );
+	SLATE_API void Construct( const FArguments& InArgs );
 
 	
-	virtual void OnMouseEnter( const FGeometry& MyGeometry, const FPointerEvent& MouseEvent ) override;
+	SLATE_API virtual void OnMouseEnter( const FGeometry& MyGeometry, const FPointerEvent& MouseEvent ) override;
 
 protected:
 
@@ -102,5 +102,5 @@ protected:
 	* @param Content	Widget to check for an STextBlock
 	* @return	The STextBlock widget found
 	*/
-	TSharedRef<SWidget> FindTextBlockWidget(TSharedRef<SWidget> Content);
+	SLATE_API TSharedRef<SWidget> FindTextBlockWidget(TSharedRef<SWidget> Content);
 };
diff --git a/Engine/Source/Runtime/Slate/Private/Widgets/Views/SListPanel.h b/Engine/Source/Runtime/Slate/Private/Widgets/Views/SListPanel.h
--- a/Engine/Source/Runtime/Slate/Private/Widgets/Views/SListPanel.h
+++ b/Engine/Source/Runtime/Slate/Private/Widgets/Views/SListPanel.h
@@ -18,7 +18,7 @@ class FArrangedChildren;
  * Items in this panel have a uniform height.
  * Also supports offsetting its items vertically.
  */
-class SLATE_API SListPanel : public SPanel
+class SListPanel : public SPanel
 {
 public:
 	/** A ListPanel slot is very simple - it just stores a widget. */
@@ -31,11 +31,11 @@ public:
 	};
 	
 	/** Make a new ListPanel::Slot  */
-	static FSlot::FSlotArguments Slot();
+	static SLATE_API FSlot::FSlotArguments Slot();
 	
 	using FScopedWidgetSlotArguments = TPanelChildren<FSlot>::FScopedWidgetSlotArguments;
 	/** Add a slot to the ListPanel */
-	FScopedWidgetSlotArguments AddSlot(int32 InsertAtIndex = INDEX_NONE);
+	SLATE_API FScopedWidgetSlotArguments AddSlot(int32 InsertAtIndex = INDEX_NONE);
 	
 	SLATE_BEGIN_ARGS( SListPanel )
 		: _ItemWidth(0)
@@ -57,68 +57,68 @@ public:
 
 	SLATE_END_ARGS()
 	
-	SListPanel();
+	SLATE_API SListPanel();
 
 	/**
 	 * Construct the widget
 	 *
 	 * @param InArgs   A declaration from which to construct the widget
 	 */
-	void Construct( const FArguments& InArgs );
+	SLATE_API void Construct( const FArguments& InArgs );
 
 public:
 
 	// SWidget interface
-	virtual void OnArrangeChildren( const FGeometry& AllottedGeometry, FArrangedChildren& ArrangedChildren ) const override;
-	virtual FVector2D ComputeDesiredSize(float) const override;
-	virtual FChildren* GetAllChildren() override;
-	virtual FChildren* GetChildren() override;
-	virtual void Tick( const FGeometry& AllottedGeometry, const double InCurrentTime, const float InDeltaTime ) override;
+	SLATE_API virtual void OnArrangeChildren( const FGeometry& AllottedGeometry, FArrangedChildren& ArrangedChildren ) const override;
+	SLATE_API virtual FVector2D ComputeDesiredSize(float) const override;
+	SLATE_API virtual FChildren* GetAllChildren() override;
+	SLATE_API virtual FChildren* GetChildren() override;
+	SLATE_API virtual void Tick( const FGeometry& AllottedGeometry, const double InCurrentTime, const float InDeltaTime ) override;
 	// End of SWidget interface
 
 	/** Fraction of the first line that we should offset by to account for the current scroll amount. */
-	void SetFirstLineScrollOffset(float InFirstLineScrollOffset);
+	SLATE_API void SetFirstLineScrollOffset(float InFirstLineScrollOffset);
 
 	/** Set how much we should appear to have scrolled past the beginning/end of the list. */
-	void SetOverscrollAmount( float InOverscrollAmount );
+	SLATE_API void SetOverscrollAmount( float InOverscrollAmount );
 	
 	/** Remove all the children from this panel */
-	void ClearItems();
+	SLATE_API void ClearItems();
 
 	/** @return the uniform desired item dimensions used when arranging children. */
-	FTableViewDimensions GetDesiredItemDimensions() const;
+	SLATE_API FTableViewDimensions GetDesiredItemDimensions() const;
 
 	/** @return the uniform item width used when arranging children. */
-	FTableViewDimensions GetItemSize(const FGeometry& AllottedGeometry) const;
+	SLATE_API FTableViewDimensions GetItemSize(const FGeometry& AllottedGeometry) const;
 
 	/** @return the uniform item width used when arranging children. */
-	FTableViewDimensions GetItemSize(const FGeometry& AllottedGeometry, const EListItemAlignment ListItemAlignment) const;
+	SLATE_API FTableViewDimensions GetItemSize(const FGeometry& AllottedGeometry, const EListItemAlignment ListItemAlignment) const;
 
 	/** @return the horizontal padding applied to each tile item */
-	float GetItemPadding(const FGeometry& AllottedGeometry) const;
+	SLATE_API float GetItemPadding(const FGeometry& AllottedGeometry) const;
 
 	/** @return the horizontal padding applied to each tile item */
-	float GetItemPadding(const FGeometry& AllottedGeometry, const EListItemAlignment ListItemAlignment) const;
+	SLATE_API float GetItemPadding(const FGeometry& AllottedGeometry, const EListItemAlignment ListItemAlignment) const;
 
 	/** @return the horizontal padding applied to all the items on a line */
-	float GetLinePadding(const FGeometry& AllottedGeometry, const int32 LineStartIndex) const;
+	SLATE_API float GetLinePadding(const FGeometry& AllottedGeometry, const int32 LineStartIndex) const;
 
 	/** Tells the list panel whether items in the list are pending a refresh */
-	void SetRefreshPending( bool IsPendingRefresh );
+	SLATE_API void SetRefreshPending( bool IsPendingRefresh );
 
 	/** Returns true if this list panel is pending a refresh, false otherwise */
-	bool IsRefreshPending() const;
+	SLATE_API bool IsRefreshPending() const;
 
 	/** See ItemHeight attribute */
-	void SetItemHeight(TAttribute<float> Height);
+	SLATE_API void SetItemHeight(TAttribute<float> Height);
 
 	/** See ItemWidth attribute */
-	void SetItemWidth(TAttribute<float> Width);
+	SLATE_API void SetItemWidth(TAttribute<float> Width);
 
 protected:
 
 	/** @return true if this panel should arrange items as tiles placed alongside one another in each line */
-	bool ShouldArrangeAsTiles() const;
+	SLATE_API bool ShouldArrangeAsTiles() const;
 	
 protected:
 
diff --git a/Engine/Source/Runtime/Slate/Public/Framework/Application/AnalogCursor.h b/Engine/Source/Runtime/Slate/Public/Framework/Application/AnalogCursor.h
--- a/Engine/Source/Runtime/Slate/Public/Framework/Application/AnalogCursor.h
+++ b/Engine/Source/Runtime/Slate/Public/Framework/Application/AnalogCursor.h
@@ -32,38 +32,38 @@ enum class EAnalogStick : uint8
 /**
  * A class that simulates a cursor driven by an analog stick.
  */
-class SLATE_API FAnalogCursor : public IInputProcessor, public TSharedFromThis<FAnalogCursor>
+class FAnalogCursor : public IInputProcessor, public TSharedFromThis<FAnalogCursor>
 {
 public:
-	FAnalogCursor();
+	SLATE_API FAnalogCursor();
 
 	/** Dtor */
 	virtual ~FAnalogCursor()
 	{}
 
-	virtual void Tick(const float DeltaTime, FSlateApplication& SlateApp, TSharedRef<ICursor> Cursor) override;
+	SLATE_API virtual void Tick(const float DeltaTime, FSlateApplication& SlateApp, TSharedRef<ICursor> Cursor) override;
 
-	virtual bool HandleKeyDownEvent(FSlateApplication& SlateApp, const FKeyEvent& InKeyEvent) override;
-	virtual bool HandleKeyUpEvent(FSlateApplication& SlateApp, const FKeyEvent& InKeyEvent) override;
-	virtual bool HandleAnalogInputEvent(FSlateApplication& SlateApp, const FAnalogInputEvent& InAnalogInputEvent) override;
-	virtual bool HandleMouseMoveEvent(FSlateApplication& SlateApp, const FPointerEvent& MouseEvent) override;
+	SLATE_API virtual bool HandleKeyDownEvent(FSlateApplication& SlateApp, const FKeyEvent& InKeyEvent) override;
+	SLATE_API virtual bool HandleKeyUpEvent(FSlateApplication& SlateApp, const FKeyEvent& InKeyEvent) override;
+	SLATE_API virtual bool HandleAnalogInputEvent(FSlateApplication& SlateApp, const FAnalogInputEvent& InAnalogInputEvent) override;
+	SLATE_API virtual bool HandleMouseMoveEvent(FSlateApplication& SlateApp, const FPointerEvent& MouseEvent) override;
 	virtual const TCHAR* GetDebugName() const override { return TEXT("AnalogCursor"); }
 
 
 	virtual int32 GetOwnerUserIndex() const { return 0; };
 
-	void SetAcceleration(float NewAcceleration);
-	void SetMaxSpeed(float NewMaxSpeed);
-	void SetStickySlowdown(float NewStickySlowdown);
-	void SetDeadZone(float NewDeadZone);
-	void SetMode(AnalogCursorMode::Type NewMode);
+	SLATE_API void SetAcceleration(float NewAcceleration);
+	SLATE_API void SetMaxSpeed(float NewMaxSpeed);
+	SLATE_API void SetStickySlowdown(float NewStickySlowdown);
+	SLATE_API void SetDeadZone(float NewDeadZone);
+	SLATE_API void SetMode(AnalogCursorMode::Type NewMode);
 
 protected:
 
-	virtual bool IsRelevantInput(const FInputEvent& InputEvent) const;
-	virtual bool IsRelevantInput(const FKeyEvent& KeyEvent) const;
-	virtual bool IsRelevantInput(const FAnalogInputEvent& AnalogInputEvent) const;
-	virtual bool IsRelevantInput(const FPointerEvent& MouseEvent) const;
+	SLATE_API virtual bool IsRelevantInput(const FInputEvent& InputEvent) const;
+	SLATE_API virtual bool IsRelevantInput(const FKeyEvent& KeyEvent) const;
+	SLATE_API virtual bool IsRelevantInput(const FAnalogInputEvent& AnalogInputEvent) const;
+	SLATE_API virtual bool IsRelevantInput(const FPointerEvent& MouseEvent) const;
 
 	/** Getter */
 	FORCEINLINE const FVector2D& GetAnalogValues( EAnalogStick Stick = EAnalogStick::Left ) const
@@ -72,15 +72,15 @@ protected:
 	}
 
 	/** Set the cached analog stick declinations to 0 */
-	void ClearAnalogValues();
+	SLATE_API void ClearAnalogValues();
 
 	/** Handles updating the cursor position and processing a Mouse Move Event */
 	UE_DEPRECATED(4.24, "FAnalogCursor now updates cursor position based on user, not the hardware cursor specifically.")
-	virtual void UpdateCursorPosition(FSlateApplication& SlateApp, TSharedRef<ICursor> Cursor, const FVector2D& NewPosition, bool bForce = false);
+	SLATE_API virtual void UpdateCursorPosition(FSlateApplication& SlateApp, TSharedRef<ICursor> Cursor, const FVector2D& NewPosition, bool bForce = false);
 
-	virtual void UpdateCursorPosition(FSlateApplication& SlateApp, TSharedRef<FSlateUser> SlateUser, const FVector2D& NewPosition, bool bForce = false);
+	SLATE_API virtual void UpdateCursorPosition(FSlateApplication& SlateApp, TSharedRef<FSlateUser> SlateUser, const FVector2D& NewPosition, bool bForce = false);
 
-	virtual FVector2D CalculateTickedCursorPosition(const float DeltaTime, FSlateApplication& SlateApp, TSharedPtr<FSlateUser> SlateUser);
+	SLATE_API virtual FVector2D CalculateTickedCursorPosition(const float DeltaTime, FSlateApplication& SlateApp, TSharedPtr<FSlateUser> SlateUser);
 
 	/** Current speed of the cursor */
 	FVector2D CurrentSpeed;
diff --git a/Engine/Source/Runtime/Slate/Public/Framework/Application/IInputProcessor.h b/Engine/Source/Runtime/Slate/Public/Framework/Application/IInputProcessor.h
--- a/Engine/Source/Runtime/Slate/Public/Framework/Application/IInputProcessor.h
+++ b/Engine/Source/Runtime/Slate/Public/Framework/Application/IInputProcessor.h
@@ -14,7 +14,7 @@ struct FPointerEvent;
 /**
  * Interface for a Slate Input Handler
  */
-class SLATE_API IInputProcessor
+class IInputProcessor
 {
 public:
 	IInputProcessor(){};
@@ -51,4 +51,4 @@ public:
 
 	/** Debug name for logging purposes */
 	virtual const TCHAR* GetDebugName() const { return TEXT(""); }
-};
\ No newline at end of file
+};
diff --git a/Engine/Source/Runtime/Slate/Public/Framework/Application/NavigationConfig.h b/Engine/Source/Runtime/Slate/Public/Framework/Application/NavigationConfig.h
--- a/Engine/Source/Runtime/Slate/Public/Framework/Application/NavigationConfig.h
+++ b/Engine/Source/Runtime/Slate/Public/Framework/Application/NavigationConfig.h
@@ -57,37 +57,37 @@ public:
 /**
  * This class is used to control which FKeys and analog axis should move focus.
  */
-class SLATE_API FNavigationConfig : public TSharedFromThis<FNavigationConfig>
+class FNavigationConfig : public TSharedFromThis<FNavigationConfig>
 {
 public:
 	/** ctor */
-	FNavigationConfig();
+	SLATE_API FNavigationConfig();
 	/** dtor */
-	virtual ~FNavigationConfig();
+	SLATE_API virtual ~FNavigationConfig();
 
 	/** Gets the navigation direction from a given key event. */
-	virtual EUINavigation GetNavigationDirectionFromKey(const FKeyEvent& InKeyEvent) const;
+	SLATE_API virtual EUINavigation GetNavigationDirectionFromKey(const FKeyEvent& InKeyEvent) const;
 	/** Gets the navigation direction from a given analog event. */
-	virtual EUINavigation GetNavigationDirectionFromAnalog(const FAnalogInputEvent& InAnalogEvent);
+	SLATE_API virtual EUINavigation GetNavigationDirectionFromAnalog(const FAnalogInputEvent& InAnalogEvent);
 
 	/** Called when the navigation config is registered with Slate Application */
-	virtual void OnRegister();
+	SLATE_API virtual void OnRegister();
 	/** Called when the navigation config is registered with Slate Application */
-	virtual void OnUnregister();
+	SLATE_API virtual void OnUnregister();
 	/** Notified when users are removed from the system, good chance to clean up any user specific state. */
-	virtual void OnUserRemoved(int32 UserIndex);
+	SLATE_API virtual void OnUserRemoved(int32 UserIndex);
 
 	/** Notified when navigation has caused a widget change to occur */
 	virtual void OnNavigationChangedFocus(TSharedPtr<SWidget> OldWidget, TSharedPtr<SWidget> NewWidget, FFocusEvent FocusEvent) {}
 
 	/** Returns the navigation action corresponding to a key event. This version will handle multiple users correctly */
-	virtual EUINavigationAction GetNavigationActionFromKey(const FKeyEvent& InKeyEvent) const;
+	SLATE_API virtual EUINavigationAction GetNavigationActionFromKey(const FKeyEvent& InKeyEvent) const;
 
 	UE_DEPRECATED(4.24, "GetNavigationActionForKey doesn't handle multiple users properly, use GetNavigationActionFromKey instead")
-	virtual EUINavigationAction GetNavigationActionForKey(const FKey& InKey) const;
+	SLATE_API virtual EUINavigationAction GetNavigationActionForKey(const FKey& InKey) const;
 
 	/** Simplification of config as string */
-	virtual FString ToString() const;
+	SLATE_API virtual FString ToString() const;
 
 public:
 	/** Should the Tab key perform next and previous style navigation. */
@@ -116,12 +116,12 @@ protected:
 	 * that if the user moves the stick a little, we would navigate slowly, if they move it a lot, we would
 	 * repeat the navigation often.
 	 */
-	virtual float GetRepeatRateForPressure(float InPressure, int32 InRepeats) const;
+	SLATE_API virtual float GetRepeatRateForPressure(float InPressure, int32 InRepeats) const;
 
 	/**
 	 * Gets the navigation direction from the analog internally.
 	 */
-	virtual EUINavigation GetNavigationDirectionFromAnalogInternal(const FAnalogInputEvent& InAnalogEvent);
+	SLATE_API virtual EUINavigation GetNavigationDirectionFromAnalogInternal(const FAnalogInputEvent& InAnalogEvent);
 
 	virtual bool IsAnalogHorizontalKey(const FKey& InKey) const { return InKey == AnalogHorizontalKey; }
 	virtual bool IsAnalogVerticalKey  (const FKey& InKey) const { return InKey == AnalogVerticalKey;   }
@@ -132,7 +132,7 @@ protected:
 
 
 /** A navigation config that doesn't do any navigation. */
-class SLATE_API FNullNavigationConfig : public FNavigationConfig
+class FNullNavigationConfig : public FNavigationConfig
 {
 public:
 	FNullNavigationConfig()
@@ -144,23 +144,12 @@ public:
 };
 
 /** A Navigation config that supports UI Navigation with both analog sticks + D-Pad. */
-class SLATE_API FTwinStickNavigationConfig : public FNavigationConfig
+class FTwinStickNavigationConfig : public FNavigationConfig
 {
 public:
-	FTwinStickNavigationConfig()
-	{
-		bTabNavigation = false;
-
-		KeyEventRules =
-		{
-			{EKeys::Gamepad_DPad_Left, EUINavigation::Left},
-			{EKeys::Gamepad_DPad_Right, EUINavigation::Right},
-			{EKeys::Gamepad_DPad_Up, EUINavigation::Up},
-			{EKeys::Gamepad_DPad_Down, EUINavigation::Down}
-		};
-	}
+	SLATE_API FTwinStickNavigationConfig();
 	
 protected:
-	virtual bool IsAnalogHorizontalKey(const FKey& InKey) const override { return InKey == EKeys::Gamepad_LeftX || InKey == EKeys::Gamepad_RightX; }
-	virtual bool IsAnalogVerticalKey  (const FKey& InKey) const override { return InKey == EKeys::Gamepad_LeftY || InKey == EKeys::Gamepad_RightY; }
-};
\ No newline at end of file
+	SLATE_API virtual bool IsAnalogHorizontalKey(const FKey& InKey) const override;
+	SLATE_API virtual bool IsAnalogVerticalKey(const FKey& InKey) const override;
+};
