[Slate]

diff --git a/Engine/Source/Runtime/Slate/Public/Framework/Application/SlateApplication.h b/Engine/Source/Runtime/Slate/Public/Framework/Application/SlateApplication.h
--- a/Engine/Source/Runtime/Slate/Public/Framework/Application/SlateApplication.h
+++ b/Engine/Source/Runtime/Slate/Public/Framework/Application/SlateApplication.h
@@ -61,7 +61,7 @@ DECLARE_DELEGATE_RetVal(bool, FDragDropCheckingOverride);
 
 
 /** Allow widgets to find out when someone clicked outside them. Currently needed by MenuAnchros. */
-class SLATE_API FPopupSupport
+class FPopupSupport
 {
 	public:
 
@@ -69,7 +69,7 @@ class SLATE_API FPopupSupport
 	 * Given a WidgetPath that was clicked, send notifications to any subscribers that were not under the mouse.
 	 * i.e. Send the "Someone clicked outside me" notifications.
 	 */
-	void SendNotifications( const FWidgetPath& WidgetsUnderCursor );
+	SLATE_API void SendNotifications( const FWidgetPath& WidgetsUnderCursor );
 
 	/**
 	* Register for a notification when the user clicks outside a specific widget.
@@ -77,7 +77,7 @@ class SLATE_API FPopupSupport
 	* @param NotifyWhenClickedOutsideMe    When the user clicks outside this widget, fire a notification.
 	* @param InNotification                The notification to invoke.
 	*/
-	FDelegateHandle RegisterClickNotification(const TSharedRef<SWidget>& NotifyWhenClickedOutsideMe, const FOnClickedOutside& InNotification);
+	SLATE_API FDelegateHandle RegisterClickNotification(const TSharedRef<SWidget>& NotifyWhenClickedOutsideMe, const FOnClickedOutside& InNotification);
 
 	/**
 	* NOTE: Only necessary if notification no longer desired.
@@ -85,7 +85,7 @@ class SLATE_API FPopupSupport
 	*
 	* Unregister the notification because it is no longer desired.
 	*/
-	void UnregisterClickNotification(FDelegateHandle InHandle);
+	SLATE_API void UnregisterClickNotification(FDelegateHandle InHandle);
 
 	private:
 
@@ -115,7 +115,7 @@ class SLATE_API FPopupSupport
 /**
  * Interface for a Slate Input Mapping.
  */
-class SLATE_API ISlateInputManager
+class ISlateInputManager
 {
 public:
 	virtual int32 GetUserIndexForMouse() const = 0;
@@ -140,7 +140,7 @@ public:
 	virtual TOptional<int32> GetUserIndexForController(int32 ControllerId, FKey InKey) const = 0;
 };
 
-class SLATE_API FSlateDefaultInputMapping : public ISlateInputManager
+class FSlateDefaultInputMapping : public ISlateInputManager
 {
 public:
 	virtual int32 GetUserIndexForMouse() const override { return 0; }
@@ -180,14 +180,14 @@ enum class ESlateTickType : uint8
 
 ENUM_CLASS_FLAGS(ESlateTickType);
 
-class SLATE_API FSlateApplication
+class FSlateApplication
 	: public FSlateApplicationBase
 	, public FGenericApplicationMessageHandler
 {
 public:
 
 	/** Virtual destructor. */
-	virtual ~FSlateApplication();
+	SLATE_API virtual ~FSlateApplication();
 
 public:
 
@@ -225,14 +225,14 @@ public:
 
 public:
 
-	static void Create();
-	static TSharedRef<FSlateApplication> Create(const TSharedRef<class GenericApplication>& InPlatformApplication);
+	static SLATE_API void Create();
+	static SLATE_API TSharedRef<FSlateApplication> Create(const TSharedRef<class GenericApplication>& InPlatformApplication);
 
-	static TSharedRef<FSlateApplication> InitializeAsStandaloneApplication(const TSharedRef< class FSlateRenderer >& PlatformRenderer);
+	static SLATE_API TSharedRef<FSlateApplication> InitializeAsStandaloneApplication(const TSharedRef< class FSlateRenderer >& PlatformRenderer);
 	
-	static TSharedRef<FSlateApplication> InitializeAsStandaloneApplication(const TSharedRef< class FSlateRenderer >& PlatformRenderer, const TSharedRef<class GenericApplication>& InPlatformApplication);
+	static SLATE_API TSharedRef<FSlateApplication> InitializeAsStandaloneApplication(const TSharedRef< class FSlateRenderer >& PlatformRenderer, const TSharedRef<class GenericApplication>& InPlatformApplication);
 
-	static void InitializeCoreStyle();
+	static SLATE_API void InitializeCoreStyle();
 
 	/**
 	 * Returns true if a Slate application instance is currently initialized and ready
@@ -257,16 +257,16 @@ public:
 		return *CurrentApplication;
 	}
 
-	static void Shutdown(bool bShutdownPlatform = true);
+	static SLATE_API void Shutdown(bool bShutdownPlatform = true);
 
 	/** @return the global tab manager */
-	static TSharedRef<class FGlobalTabmanager> GetGlobalTabManager();
+	static SLATE_API TSharedRef<class FGlobalTabmanager> GetGlobalTabManager();
 
 	/** Initializes high dpi support for the process */
-	static void InitHighDPI(const bool bForceEnable);
+	static SLATE_API void InitHighDPI(const bool bForceEnable);
 
 	/** @return the root style node, which is the entry point to the style graph representing all the current style rules. */
-	const class FStyleNode* GetRootStyle() const;
+	SLATE_API const class FStyleNode* GetRootStyle() const;
 
 	/**
 	 * Initializes the renderer responsible for drawing all elements in this application
@@ -274,16 +274,16 @@ public:
 	 * @param InRenderer The renderer to use.
 	 * @param bQuietMode Don't show any message boxes when initialization fails.
 	 */
-	virtual bool InitializeRenderer( TSharedRef<FSlateRenderer> InRenderer, bool bQuietMode = false );
+	SLATE_API virtual bool InitializeRenderer( TSharedRef<FSlateRenderer> InRenderer, bool bQuietMode = false );
 
 	/** Set the slate sound provider that the slate app should use. */
-	virtual void InitializeSound( const TSharedRef<ISlateSoundDevice>& InSlateSoundDevice );
+	SLATE_API virtual void InitializeSound( const TSharedRef<ISlateSoundDevice>& InSlateSoundDevice );
 
 	/** Play SoundToPlay. Interrupt previous sound if one is playing. */
-	void PlaySound( const FSlateSound& SoundToPlay, int32 UserIndex = 0 ) const;
+	SLATE_API void PlaySound( const FSlateSound& SoundToPlay, int32 UserIndex = 0 ) const;
 
 	/** @return The duration of the given sound resource */
-	float GetSoundDuration(const FSlateSound& Sound) const;
+	SLATE_API float GetSoundDuration(const FSlateSound& Sound) const;
 
 	IInputInterface* GetInputInterface() const { return PlatformApplication->GetInputInterface(); }
 
@@ -300,7 +300,7 @@ public:
 	 *
 	 * @param MouseCoordinate		The new position.
 	 */
-	void SetCursorPos( const FVector2D& MouseCoordinate ) override;
+	SLATE_API void SetCursorPos( const FVector2D& MouseCoordinate ) override;
 
 	/** 
 	 * Replace the IPlatformTextField implementation with a custom one.  The implementation used by default is platform specific.
@@ -308,45 +308,45 @@ public:
 	 *
 	 * @param PlatformTextField		The platform specific implementation of the IPlatformTextField.
 	 */
-	void OverridePlatformTextField(TUniquePtr<IPlatformTextField> PlatformTextField);
+	SLATE_API void OverridePlatformTextField(TUniquePtr<IPlatformTextField> PlatformTextField);
 	
 	/** 
 	 *	Updates the cursor user's cursor to either the platform cursor or fake cursor
 	 */
-	void UsePlatformCursorForCursorUser(bool bUsePlatformCursor);
+	SLATE_API void UsePlatformCursorForCursorUser(bool bUsePlatformCursor);
 
 	/** Changes the cursor type to Default (Visible) or None (Not Visible)*/
-	void SetPlatformCursorVisibility(bool bNewVisibility);
+	SLATE_API void SetPlatformCursorVisibility(bool bNewVisibility);
 
 	/** Polls game devices for input */
-	void PollGameDeviceState();
+	SLATE_API void PollGameDeviceState();
 
 	/** Occurs before Tick(), after all pointer and keyboard input has been processed. */
-	void FinishedInputThisFrame();
+	SLATE_API void FinishedInputThisFrame();
 
 	/** Ticks this application */
-	void Tick(ESlateTickType TickType = ESlateTickType::All);
+	SLATE_API void Tick(ESlateTickType TickType = ESlateTickType::All);
 
 	/** Pumps OS messages when a modal window or intra-frame debugging session exists */
-	void PumpMessages();
+	SLATE_API void PumpMessages();
 
 	/** Returns true if this slate application is ready to open modal windows */
-	bool CanAddModalWindow() const;
+	SLATE_API bool CanAddModalWindow() const;
 
 	/** Returns true if this slate application is ready to display windows. */
-	bool CanDisplayWindows() const;
+	SLATE_API bool CanDisplayWindows() const;
 
 	/** Returns navigation direction matching a key event, this is determined in the FNavigationConfig */
-	virtual EUINavigation GetNavigationDirectionFromKey(const FKeyEvent& InKeyEvent) const override;
+	SLATE_API virtual EUINavigation GetNavigationDirectionFromKey(const FKeyEvent& InKeyEvent) const override;
 
 	/** Returns navigation direction matching an anlog event, this is determined in the FNavigationConfig */
-	virtual EUINavigation GetNavigationDirectionFromAnalog(const FAnalogInputEvent& InAnalogEvent) override;
+	SLATE_API virtual EUINavigation GetNavigationDirectionFromAnalog(const FAnalogInputEvent& InAnalogEvent) override;
 
 	/** Returns the navigation action corresponding to a key event. This version will handle multiple users correctly */
-	virtual EUINavigationAction GetNavigationActionFromKey(const FKeyEvent& InKeyEvent) const override;
+	SLATE_API virtual EUINavigationAction GetNavigationActionFromKey(const FKeyEvent& InKeyEvent) const override;
 
 	UE_DEPRECATED(4.24, "GetNavigationActionForKey doesn't handle multiple users properly, use GetNavigationActionFromKey instead")
-	virtual EUINavigationAction GetNavigationActionForKey(const FKey& InKey) const override;
+	SLATE_API virtual EUINavigationAction GetNavigationActionForKey(const FKey& InKey) const override;
 
 	/**
 	 * Adds a modal window to the application.  
@@ -356,13 +356,13 @@ public:
 	 * @param InParentWindow	The parent of the modal window.  All modal windows must have a parent.
 	 * @param bSlowTaskWindow	true if the window is for a slow task and this function should return before the window is closed
 	 */
-	void AddModalWindow( TSharedRef<SWindow> InSlateWindow, const TSharedPtr<const SWidget> InParentWidget, bool bSlowTaskWindow = false );
+	SLATE_API void AddModalWindow( TSharedRef<SWindow> InSlateWindow, const TSharedPtr<const SWidget> InParentWidget, bool bSlowTaskWindow = false );
 
 	/** Sets the delegate for when a modal window stack begins */
-	void SetModalWindowStackStartedDelegate(FModalWindowStackStarted StackStartedDelegate);
+	SLATE_API void SetModalWindowStackStartedDelegate(FModalWindowStackStarted StackStartedDelegate);
 
 	/** Sets the delegate for when a modal window stack ends */
-	void SetModalWindowStackEndedDelegate(FModalWindowStackEnded StackEndedDelegate);
+	SLATE_API void SetModalWindowStackEndedDelegate(FModalWindowStackEnded StackEndedDelegate);
 
 	/**
 	 * Associates a top level Slate Window with a native window, and "natively" parents that window to the specified Slate window.
@@ -374,7 +374,7 @@ public:
 	 *
 	 * @return a reference to the SWindow that was just added.
 	 */
-	TSharedRef<SWindow> AddWindowAsNativeChild( TSharedRef<SWindow> InSlateWindow, TSharedRef<SWindow> InParentWindow, const bool bShowImmediately = true );
+	SLATE_API TSharedRef<SWindow> AddWindowAsNativeChild( TSharedRef<SWindow> InSlateWindow, TSharedRef<SWindow> InParentWindow, const bool bShowImmediately = true );
 
 	/**
 	 * Creates a new Menu and adds it to the menu stack.
@@ -390,7 +390,7 @@ public:
 	 * @param Method				An optional popup method override. If not set, the widgets in the InOwnerPath will be queried for this.
 	 * @param bIsCollapsedByParent	Is this menu collapsed when a parent menu receives focus/activation? If false, only focus/activation outside the entire stack will auto collapse it.
 	 */
-	TSharedPtr<IMenu> PushMenu(const TSharedRef<SWidget>& InParentWidget, const FWidgetPath& InOwnerPath, const TSharedRef<SWidget>& InContent, const UE::Slate::FDeprecateVector2DParameter& SummonLocation, const FPopupTransitionEffect& TransitionEffect, const bool bFocusImmediately = true, const UE::Slate::FDeprecateVector2DParameter& SummonLocationSize = FVector2f::ZeroVector, TOptional<EPopupMethod> Method = TOptional<EPopupMethod>(), const bool bIsCollapsedByParent = true);
+	SLATE_API TSharedPtr<IMenu> PushMenu(const TSharedRef<SWidget>& InParentWidget, const FWidgetPath& InOwnerPath, const TSharedRef<SWidget>& InContent, const UE::Slate::FDeprecateVector2DParameter& SummonLocation, const FPopupTransitionEffect& TransitionEffect, const bool bFocusImmediately = true, const UE::Slate::FDeprecateVector2DParameter& SummonLocationSize = FVector2f::ZeroVector, TOptional<EPopupMethod> Method = TOptional<EPopupMethod>(), const bool bIsCollapsedByParent = true);
 
 	/**
 	 * Creates a new Menu and adds it to the menu stack under the specified parent menu.
@@ -404,7 +404,7 @@ public:
 	 * @param SummonLocationSize	An optional rect which describes an area in which the menu may not appear
 	 * @param bIsCollapsedByParent	Is this menu collapsed when a parent menu receives focus/activation? If false, only focus/activation outside the entire stack will auto collapse it.
 	 */
-	TSharedPtr<IMenu> PushMenu(const TSharedPtr<IMenu>& InParentMenu, const TSharedRef<SWidget>& InContent, const UE::Slate::FDeprecateVector2DParameter& SummonLocation, const FPopupTransitionEffect& TransitionEffect, const bool bFocusImmediately = true, const UE::Slate::FDeprecateVector2DParameter& SummonLocationSize = FVector2f::ZeroVector, const bool bIsCollapsedByParent = true);
+	SLATE_API TSharedPtr<IMenu> PushMenu(const TSharedPtr<IMenu>& InParentMenu, const TSharedRef<SWidget>& InContent, const UE::Slate::FDeprecateVector2DParameter& SummonLocation, const FPopupTransitionEffect& TransitionEffect, const bool bFocusImmediately = true, const UE::Slate::FDeprecateVector2DParameter& SummonLocationSize = FVector2f::ZeroVector, const bool bIsCollapsedByParent = true);
 
 	/**
 	 * Creates a new hosted Menu and adds it to the menu stack.
@@ -419,7 +419,7 @@ public:
 	 * @param ShouldThrottle		Should we throttle engine ticking to maximize the menu responsiveness
 	 * @param bIsCollapsedByParent	Is this menu collapsed when a parent menu receives focus/activation? If false, only focus/activation outside the entire stack will auto collapse it.
 	 */
-	TSharedPtr<IMenu> PushHostedMenu(const TSharedRef<SWidget>& InParentWidget, const FWidgetPath& InOwnerPath, const TSharedRef<IMenuHost>& InMenuHost, const TSharedRef<SWidget>& InContent, TSharedPtr<SWidget>& OutWrappedContent, const FPopupTransitionEffect& TransitionEffect, EShouldThrottle ShouldThrottle, const bool bIsCollapsedByParent = true);
+	SLATE_API TSharedPtr<IMenu> PushHostedMenu(const TSharedRef<SWidget>& InParentWidget, const FWidgetPath& InOwnerPath, const TSharedRef<IMenuHost>& InMenuHost, const TSharedRef<SWidget>& InContent, TSharedPtr<SWidget>& OutWrappedContent, const FPopupTransitionEffect& TransitionEffect, EShouldThrottle ShouldThrottle, const bool bIsCollapsedByParent = true);
 	
 	/**
 	 * Creates a new hosted child Menu and adds it to the menu stack under the specified parent menu.
@@ -433,13 +433,13 @@ public:
 	 * @param ShouldThrottle		Should we throttle engine ticking to maximize the menu responsiveness
 	 * @param bIsCollapsedByParent	Is this menu collapsed when a parent menu receives focus/activation? If false, only focus/activation outside the entire stack will auto collapse it.
 	 */	
-	TSharedPtr<IMenu> PushHostedMenu(const TSharedPtr<IMenu>& InParentMenu, const TSharedRef<IMenuHost>& InMenuHost, const TSharedRef<SWidget>& InContent, TSharedPtr<SWidget>& OutWrappedContent, const FPopupTransitionEffect& TransitionEffect, EShouldThrottle ShouldThrottle, const bool bIsCollapsedByParent = true);
+	SLATE_API TSharedPtr<IMenu> PushHostedMenu(const TSharedPtr<IMenu>& InParentMenu, const TSharedRef<IMenuHost>& InMenuHost, const TSharedRef<SWidget>& InContent, TSharedPtr<SWidget>& OutWrappedContent, const FPopupTransitionEffect& TransitionEffect, EShouldThrottle ShouldThrottle, const bool bIsCollapsedByParent = true);
 
 	/** @return Returns whether the menu has child menus. */
-	bool HasOpenSubMenus(TSharedPtr<IMenu> InMenu) const;
+	SLATE_API bool HasOpenSubMenus(TSharedPtr<IMenu> InMenu) const;
 
 	/** @return	Returns true if there are any pop-up menus summoned */
-	virtual bool AnyMenusVisible() const override;
+	SLATE_API virtual bool AnyMenusVisible() const override;
 
 	/**
 	 * Attempt to locate a menu that contains the specified widget
@@ -447,30 +447,30 @@ public:
 	 * @param InWidgetPath Path to the widget to use for the search
 	 * @return the menu in which the widget resides, or nullptr
 	 */
-	TSharedPtr<IMenu> FindMenuInWidgetPath(const FWidgetPath& InWidgetPath) const;
+	SLATE_API TSharedPtr<IMenu> FindMenuInWidgetPath(const FWidgetPath& InWidgetPath) const;
 
 	/** @return	Returns a ptr to the window that is currently the host of the menu stack or null if no menus are visible */
-	TSharedPtr<SWindow> GetVisibleMenuWindow() const;
+	SLATE_API TSharedPtr<SWindow> GetVisibleMenuWindow() const;
 
 	/** @return	Returns a ptr to the widget that created the opened root menu or null if one is not opened */
-	TSharedPtr<SWidget> GetMenuHostWidget() const;
+	SLATE_API TSharedPtr<SWidget> GetMenuHostWidget() const;
 
 	/** Dismisses all open menus */
-	void DismissAllMenus();
+	SLATE_API void DismissAllMenus();
 
 	/**
 	 * Dismisses a menu and all its children
 	 *
 	 * @param InFromMenu	The menu to dismiss, any children, grandchildren etc will also be dismissed
 	 */
-	void DismissMenu(const TSharedPtr<IMenu>& InFromMenu);
+	SLATE_API void DismissMenu(const TSharedPtr<IMenu>& InFromMenu);
 
 	/**
 	 * Dismisses a menu and all its children. The menu is determined by looking for menus in the parent chain of the widget.
 	 *
 	 * @param InWidgetInMenu	The widget whose path is search upwards for a menu. That menu will then be dismissed.
 	 */
-	void DismissMenuByWidget(const TSharedRef<SWidget>& InWidgetInMenu);
+	SLATE_API void DismissMenuByWidget(const TSharedRef<SWidget>& InWidgetInMenu);
 
 	/**
 	 * HACK: Don't use this unless shutting down a game viewport
@@ -478,19 +478,19 @@ public:
 	 *
 	 * @param WindowToDestroy		Window for destruction
 	 */
-	void DestroyWindowImmediately( TSharedRef<SWindow> WindowToDestroy );
+	SLATE_API void DestroyWindowImmediately( TSharedRef<SWindow> WindowToDestroy );
 
 	/**
 	 * Disable Slate components when an external, non-slate, modal window is brought up.  In the case of multiple
 	 * external modal windows, we will only increment our tracking counter.
 	 */
-	void ExternalModalStart();
+	SLATE_API void ExternalModalStart();
 
 	/**
 	 * Re-enable disabled Slate components when a non-slate modal window is dismissed.  Slate components
 	 * will only be re-enabled when all tracked external modal windows have been dismissed.
 	 */
-	void ExternalModalStop();
+	SLATE_API void ExternalModalStop();
 
 	/** Event before slate application ticks. */
 	DECLARE_EVENT_OneParam(FSlateApplication, FSlateTickEvent, float);
@@ -519,14 +519,14 @@ public:
 	 * This has to be done explicitly instead of using the FRenderResource mechanism because FViewportRHI's are managed by the game thread.
 	 * This is needed before destroying the RHI device. 
 	 */
-	void InvalidateAllViewports();
+	SLATE_API void InvalidateAllViewports();
 
 	/**
 	 * Registers a game viewport with the Slate application so that specific messages can be routed directly to a viewport
 	 * 
 	 * @param InViewport	The viewport to register.  Note there is currently only one registered viewport
 	 */
-	void RegisterGameViewport( TSharedRef<SViewport> InViewport );
+	SLATE_API void RegisterGameViewport( TSharedRef<SViewport> InViewport );
 
 	/**
 	 * Registers a viewport with the Slate application so that specific messages can be routed directly to a viewport
@@ -534,56 +534,56 @@ public:
 	 * 
 	 * @param InViewport	The viewport to register.  Note there is currently only one registered viewport
 	 */
-	void RegisterViewport(TSharedRef<SViewport> InViewport);
+	SLATE_API void RegisterViewport(TSharedRef<SViewport> InViewport);
 
 	/**
 	 * Returns the game viewport registered with the slate application
 	 *
 	 * @return registered game viewport
 	 */
-	TSharedPtr<SViewport> GetGameViewport() const;
+	SLATE_API TSharedPtr<SViewport> GetGameViewport() const;
 
 	/** 
 	 * Unregisters the current game viewport from Slate.  This method sends a final deactivation message to the viewport
 	 * to allow it to do a final cleanup before being closed.
 	 */
-	void UnregisterGameViewport();
+	SLATE_API void UnregisterGameViewport();
 
 	/**
 	 * Register another window that may be visible in a non-top level way that still needs to be able to maintain focus paths.
 	 * Generally speaking - this is for Virtual Windows that are created to render in the 3D world slate content.
 	 */
-	void RegisterVirtualWindow(TSharedRef<SWindow> InWindow);
+	SLATE_API void RegisterVirtualWindow(TSharedRef<SWindow> InWindow);
 
 	/**
 	 * Unregister a virtual window.
 	 */
-	void UnregisterVirtualWindow(TSharedRef<SWindow> InWindow);
+	SLATE_API void UnregisterVirtualWindow(TSharedRef<SWindow> InWindow);
 
 	/**
 	 * Flushes the render state of slate, releasing accesses and flushing all render commands.
 	 */
-	void FlushRenderState();
+	SLATE_API void FlushRenderState();
 
 	/**
 	 * Sets specified user focus to the SWidget representing the currently active game viewport
 	 */
-	void SetUserFocusToGameViewport(uint32 UserIndex, EFocusCause ReasonFocusIsChanging = EFocusCause::SetDirectly);
+	SLATE_API void SetUserFocusToGameViewport(uint32 UserIndex, EFocusCause ReasonFocusIsChanging = EFocusCause::SetDirectly);
 
 	/**
 	 * Sets all users focus to the SWidget representing the currently active game viewport
 	 */
-	void SetAllUserFocusToGameViewport(EFocusCause ReasonFocusIsChanging = EFocusCause::SetDirectly);
+	SLATE_API void SetAllUserFocusToGameViewport(EFocusCause ReasonFocusIsChanging = EFocusCause::SetDirectly);
 
 	/**
 	 * Activates the Game Viewport if it is properly childed under a window
 	 */
-	void ActivateGameViewport();
+	SLATE_API void ActivateGameViewport();
 
 	/**
 	 * True if transforming mouse input coordinates to account for fullscreen distortions
 	 */
-	bool GetTransformFullscreenMouseInput() const;
+	SLATE_API bool GetTransformFullscreenMouseInput() const;
 
 #if WITH_SLATE_DEBUGGING
 	/**
@@ -592,7 +592,7 @@ public:
 	 *
 	 * @param InNavigationConfig Navigation config to dump info to log for
 	 */
-	void TryDumpNavigationConfig(TSharedPtr<FNavigationConfig> InNavigationConfig) const;
+	SLATE_API void TryDumpNavigationConfig(TSharedPtr<FNavigationConfig> InNavigationConfig) const;
 #endif // WITH_SLATE_DEBUGGING
 
 	/**
@@ -602,7 +602,7 @@ public:
 	 * @param WidgetToFocus the widget to set focus to
 	 * @param ReasonFocusIsChanging the contextual reason for the focus change
 	 */
-	bool SetUserFocus(uint32 UserIndex, const TSharedPtr<SWidget>& WidgetToFocus, EFocusCause ReasonFocusIsChanging = EFocusCause::SetDirectly);
+	SLATE_API bool SetUserFocus(uint32 UserIndex, const TSharedPtr<SWidget>& WidgetToFocus, EFocusCause ReasonFocusIsChanging = EFocusCause::SetDirectly);
 
 	/**
 	 * Sets focus for all users to the SWidget passed in.
@@ -610,25 +610,25 @@ public:
 	 * @param WidgetToFocus the widget to set focus to
 	 * @param ReasonFocusIsChanging the contextual reason for the focus change
 	 */
-	void SetAllUserFocus(const TSharedPtr<SWidget>& WidgetToFocus, EFocusCause ReasonFocusIsChanging = EFocusCause::SetDirectly);
+	SLATE_API void SetAllUserFocus(const TSharedPtr<SWidget>& WidgetToFocus, EFocusCause ReasonFocusIsChanging = EFocusCause::SetDirectly);
 
 	/** Releases the users focus from whatever it currently is on. */
-	void ClearUserFocus(uint32 UserIndex, EFocusCause ReasonFocusIsChanging = EFocusCause::SetDirectly);
+	SLATE_API void ClearUserFocus(uint32 UserIndex, EFocusCause ReasonFocusIsChanging = EFocusCause::SetDirectly);
 
 	/** Releases the focus for all users from whatever it currently is on. */
-	void ClearAllUserFocus(EFocusCause ReasonFocusIsChanging = EFocusCause::SetDirectly);
+	SLATE_API void ClearAllUserFocus(EFocusCause ReasonFocusIsChanging = EFocusCause::SetDirectly);
 
 	/**
 	 * Sets the Keyboard focus to the specified SWidget
 	 */
-	bool SetKeyboardFocus(const TSharedPtr<SWidget>& OptionalWidgetToFocus, EFocusCause ReasonFocusIsChanging = EFocusCause::SetDirectly);
+	SLATE_API bool SetKeyboardFocus(const TSharedPtr<SWidget>& OptionalWidgetToFocus, EFocusCause ReasonFocusIsChanging = EFocusCause::SetDirectly);
 
 	/**
 	 * Clears keyboard focus, if any widget is currently focused
 	 *
 	 * @param ReasonFocusIsChanging The reason that keyboard focus is changing
 	 */
-	void ClearKeyboardFocus(const EFocusCause ReasonFocusIsChanging = EFocusCause::SetDirectly);
+	SLATE_API void ClearKeyboardFocus(const EFocusCause ReasonFocusIsChanging = EFocusCause::SetDirectly);
 
 #if WITH_EDITOR
 	/**
@@ -658,7 +658,7 @@ public:
 	 *
 	 * @return  State of modifier keys
 	 */
-	FModifierKeysState GetModifierKeys() const;
+	SLATE_API FModifierKeysState GetModifierKeys() const;
 
 	/**
 	 *	Restores all input settings to their original values
@@ -666,14 +666,14 @@ public:
 	 *  Clears all user focus
 	 *  Shows the mouse, clears any locks and captures, turns off high precision input
 	 */
-	void ResetToDefaultInputSettings();
+	SLATE_API void ResetToDefaultInputSettings();
 	
 	/**
 	 *	Restores all pointer input settings to their original values
 	 * 
 	 *  Shows the mouse, clears any locks and captures, turns off high precision input
 	 */
-	void ResetToDefaultPointerInputSettings();
+	SLATE_API void ResetToDefaultPointerInputSettings();
 
 	/**
 	 * Mouse capture
@@ -686,27 +686,27 @@ public:
 	bool GetHandleDeviceInputWhenApplicationNotActive() const { return bHandleDeviceInputWhenApplicationNotActive; }
 
 	/** returning platform-specific value designating window that captures mouse, or nullptr if mouse isn't captured */
-	virtual void* GetMouseCaptureWindow() const;
+	SLATE_API virtual void* GetMouseCaptureWindow() const;
 
 	/** Releases capture for every pointer on every user from whatever it currently is on. */
-	void ReleaseAllPointerCapture();
+	SLATE_API void ReleaseAllPointerCapture();
 	UE_DEPRECATED(4.23, "ReleaseMouseCapture has been renamed to ReleaseAllPointerCapture()")
-	void ReleaseMouseCapture();
+	SLATE_API void ReleaseMouseCapture();
 
 	/** Releases capture for every pointer belonging to the given user index particular user. */
-	void ReleaseAllPointerCapture(int32 UserIndex);
+	SLATE_API void ReleaseAllPointerCapture(int32 UserIndex);
 	UE_DEPRECATED(4.23, "ReleaseMouseCaptureForUser has been renamed to ReleaseAllPointerCapture(int32 UserIndex)")
-	void ReleaseMouseCaptureForUser(int32 UserIndex);
+	SLATE_API void ReleaseMouseCaptureForUser(int32 UserIndex);
 
 	/** @return The active modal window or nullptr if there is no modal window. */
-	TSharedPtr<SWindow> GetActiveModalWindow() const;
+	SLATE_API TSharedPtr<SWindow> GetActiveModalWindow() const;
 
 	/**
 	 * Assign a delegate to be called when this application is requesting an exit (e.g. when the last window is closed).
 	 *
 	 * @param OnExitRequestHandler  Function to execute when the application wants to quit
 	 */
-	void SetExitRequestedHandler( const FSimpleDelegate& OnExitRequestedHandler );
+	SLATE_API void SetExitRequestedHandler( const FSimpleDelegate& OnExitRequestedHandler );
 		
 	/**
 	 * @todo slate: Remove this method or make it private.
@@ -717,7 +717,7 @@ public:
 	 * @param  OutWidgetPath  The generated widget path
 	 * @param  VisibilityFilter	Widgets must have this type of visibility to be included the path
 	 */
-	bool GeneratePathToWidgetUnchecked( TSharedRef<const SWidget> InWidget, FWidgetPath& OutWidgetPath, EVisibility VisibilityFilter = EVisibility::Visible ) const;
+	SLATE_API bool GeneratePathToWidgetUnchecked( TSharedRef<const SWidget> InWidget, FWidgetPath& OutWidgetPath, EVisibility VisibilityFilter = EVisibility::Visible ) const;
 	
 	/**
 	 * @todo slate: Remove this method or make it private.
@@ -728,7 +728,7 @@ public:
 	 * @param  OutWidgetPath  The generated widget path
 	 * @param  VisibilityFilter	Widgets must have this type of visibility to be included the path
 	 */
-	void GeneratePathToWidgetChecked( TSharedRef<const SWidget> InWidget, FWidgetPath& OutWidgetPath, EVisibility VisibilityFilter = EVisibility::Visible ) const;
+	SLATE_API void GeneratePathToWidgetChecked( TSharedRef<const SWidget> InWidget, FWidgetPath& OutWidgetPath, EVisibility VisibilityFilter = EVisibility::Visible ) const;
 	
 	/**
 	 * Finds the window that the provided widget resides in
@@ -736,7 +736,7 @@ public:
 	 * @param InWidget		The widget to find the window for
 	 * @return The window where the widget resides, or null if the widget wasn't found.  Remember, a widget might not be found simply because its parent decided not to report the widget in ArrangeChildren.
 	 */
-	virtual TSharedPtr<SWindow> FindWidgetWindow( TSharedRef<const SWidget> InWidget ) const override;
+	SLATE_API virtual TSharedPtr<SWindow> FindWidgetWindow( TSharedRef<const SWidget> InWidget ) const override;
 
 	/**
 	 * Finds the window that the provided widget resides in
@@ -746,7 +746,7 @@ public:
 	 * @return The window where the widget resides, or null if the widget wasn't found.  Remember, a widget might not be found simply because its parent decided not to report the widget in ArrangeChildren.
 	 */
 	UE_DEPRECATED(4.23, "The FindWidgetWindow method that takes an FWidgetPath has been deprecated.  If you dont need the widget path, use FindWidgetWindow(MyWidget) instead.  If you need the path use GeneratePathToWidget")
-	TSharedPtr<SWindow> FindWidgetWindow( TSharedRef<const SWidget> InWidget, FWidgetPath& OutWidgetPath) const;
+	SLATE_API TSharedPtr<SWindow> FindWidgetWindow( TSharedRef<const SWidget> InWidget, FWidgetPath& OutWidgetPath) const;
 
 	/**
 	 * @return True if the application is currently routing high precision mouse movement events (OS specific)
@@ -774,7 +774,7 @@ public:
 	 *
 	 * @param WidgetReflector The widget reflector to set.
 	 */
-	void SetWidgetReflector( const TSharedRef<IWidgetReflector>& WidgetReflector );
+	SLATE_API void SetWidgetReflector( const TSharedRef<IWidgetReflector>& WidgetReflector );
 
 	/** @param AccessDelegate The delegate to pass along to the widget reflector */
 	void SetWidgetReflectorSourceAccessDelegate(FAccessSourceCode AccessDelegate)
@@ -800,16 +800,16 @@ public:
 	virtual void GetInitialDisplayMetrics( FDisplayMetrics& OutDisplayMetrics ) const { PlatformApplication->GetInitialDisplayMetrics( OutDisplayMetrics ); }
 
 	/** Are we drag-dropping right now? */
-	bool IsDragDropping() const;
+	SLATE_API bool IsDragDropping() const;
 
 	/** Are we drag-dropping and are we affected by this pointer event? */
-	bool IsDragDroppingAffected(const FPointerEvent& InPointerEvent) const;
+	SLATE_API bool IsDragDroppingAffected(const FPointerEvent& InPointerEvent) const;
 
 	/** Get the current drag-dropping content */
-	TSharedPtr<class FDragDropOperation> GetDragDroppingContent() const;
+	SLATE_API TSharedPtr<class FDragDropOperation> GetDragDroppingContent() const;
 
 	/** Cancels any in flight drag and drops */
-	void CancelDragDrop();
+	SLATE_API void CancelDragDrop();
 
 	/**
 	 * Returns the attribute that can be used by widgets to check if the application is in normal execution mode
@@ -831,14 +831,14 @@ public:
 	 * Enters debugging mode which is a special state that causes the
 	 * Slate application to tick in place which in the middle of a stack frame
 	 */
-	void EnterDebuggingMode();
+	SLATE_API void EnterDebuggingMode();
 
 	/**
 	 * Leaves debugging mode
 	 *
 	 * @param bLeavingDebugForSingleStep	Whether or not we are leaving debug mode due to single stepping
 	 */
-	void LeaveDebuggingMode( bool bLeavingDebugForSingleStep = false );
+	SLATE_API void LeaveDebuggingMode( bool bLeavingDebugForSingleStep = false );
 	
 	/**
 	 * Calculates the popup window position from the passed in window position and size. 
@@ -852,7 +852,7 @@ public:
 	 *								If horizontal it will attempt to open below the anchor but will open above if there is no room.
 	 * @return The adjusted position
 	 */
-	virtual UE::Slate::FDeprecateVector2DResult CalculatePopupWindowPosition( const FSlateRect& InAnchor, const UE::Slate::FDeprecateVector2DParameter& InSize, bool bAutoAdjustForDPIScale = true, const UE::Slate::FDeprecateVector2DParameter& InProposedPlacement = FVector2f::ZeroVector, const EOrientation Orientation = Orient_Vertical) const;
+	SLATE_API virtual UE::Slate::FDeprecateVector2DResult CalculatePopupWindowPosition( const FSlateRect& InAnchor, const UE::Slate::FDeprecateVector2DParameter& InSize, bool bAutoAdjustForDPIScale = true, const UE::Slate::FDeprecateVector2DParameter& InProposedPlacement = FVector2f::ZeroVector, const EOrientation Orientation = Orient_Vertical) const;
 
 	/**
 	 * Calculates the tooltip window position.
@@ -861,7 +861,7 @@ public:
 	 * @param InSize The size of the tooltip window.
 	 * @return The suggested position.
 	 */
-	virtual UE::Slate::FDeprecateVector2DResult CalculateTooltipWindowPosition( const FSlateRect& InAnchorRect, const UE::Slate::FDeprecateVector2DParameter& InSize, bool bAutoAdjustForDPIScale) const;
+	SLATE_API virtual UE::Slate::FDeprecateVector2DResult CalculateTooltipWindowPosition( const FSlateRect& InAnchorRect, const UE::Slate::FDeprecateVector2DParameter& InSize, bool bAutoAdjustForDPIScale) const;
 
 	/**
 	 * Is the window in the app's destroy queue? If so it will be destroyed next tick.
@@ -869,13 +869,13 @@ public:
 	 * @param Window		The window to find in the destroy list
 	 * @return				true if Window is in the destroy queue
 	 */
-	bool IsWindowInDestroyQueue(TSharedRef<SWindow> Window) const;
+	SLATE_API bool IsWindowInDestroyQueue(TSharedRef<SWindow> Window) const;
 
 	/** @return	Returns true if the application's average frame rate is at least as high as our target frame rate that satisfies our requirement for a smooth and responsive UI experience */
-	bool IsRunningAtTargetFrameRate() const;
+	SLATE_API bool IsRunningAtTargetFrameRate() const;
 
 	/** @return Returns true if transition effects for new menus and windows should be played */
-	bool AreMenuAnimationsEnabled() const;
+	SLATE_API bool AreMenuAnimationsEnabled() const;
 
 	/** 
 	 * Sets whether transition effects for new menus and windows should be played.  This can be called at any time.
@@ -883,19 +883,19 @@ public:
 	 * @param	bEnableAnimations	True if animations should be used, otherwise false.
 	 */
 	UE_DEPRECATED(5.0, "Enable Window Animations is no longer used and is a no-op so calling this function is no longer necessary.")
-	void EnableMenuAnimations( const bool bEnableAnimations );
+	SLATE_API void EnableMenuAnimations( const bool bEnableAnimations );
 
-	void SetPlatformApplication(const TSharedRef<class GenericApplication>& InPlatformApplication);
+	SLATE_API void SetPlatformApplication(const TSharedRef<class GenericApplication>& InPlatformApplication);
 
 	/**
 	 * Replace the current platform application with a custom version.
 	 * @param InPlatformApplication - The replacement platform application.
 	 */
-	void OverridePlatformApplication(TSharedPtr<class GenericApplication> InPlatformApplication);
+	SLATE_API void OverridePlatformApplication(TSharedPtr<class GenericApplication> InPlatformApplication);
 
 	/** Set the global application icon */
 	UE_DEPRECATED(4.26, "SetAppIcon has been deprecated.  Set \"AppIcon\" in your applications style to override the icon")
-	void SetAppIcon(const FSlateBrush* const InAppIcon);
+	SLATE_API void SetAppIcon(const FSlateBrush* const InAppIcon);
 
 	/** Sets the display state of external UI such as Steam. */
 	void ExternalUIChange(bool bIsOpening)
@@ -909,13 +909,13 @@ public:
 	 * @param bShow	true to show the keyboard, false to hide it
 	 * @param TextEntryWidget The widget that will receive the input from the virtual keyboard
 	 */
-	void ShowVirtualKeyboard( bool bShow, int32 UserIndex, TSharedPtr<IVirtualKeyboardEntry> TextEntryWidget = nullptr );
+	SLATE_API void ShowVirtualKeyboard( bool bShow, int32 UserIndex, TSharedPtr<IVirtualKeyboardEntry> TextEntryWidget = nullptr );
 
 	/** @return true if the current platform allows cursor positioning in editable text boxes */
-	bool AllowMoveCursor();
+	SLATE_API bool AllowMoveCursor();
 
 	/** Get the work area that has the largest intersection with the specified rectangle */
-	FSlateRect GetWorkArea( const FSlateRect& InRect ) const;
+	SLATE_API FSlateRect GetWorkArea( const FSlateRect& InRect ) const;
 
 	/**
 	 * Shows or hides an onscreen keyboard
@@ -928,10 +928,10 @@ public:
 	}
 
 	/** @return true if the current platform allows source code access */
-	bool SupportsSourceAccess() const;
+	SLATE_API bool SupportsSourceAccess() const;
 
 	/** Opens the current platform's code editing IDE (if necessary) and focuses the specified line in the specified file. Will only work if SupportsSourceAccess() returns true */
-	void GotoLineInSource(const FString& FileName, int32 LineNumber) const;
+	SLATE_API void GotoLineInSource(const FString& FileName, int32 LineNumber) const;
 
 	/** @return Service that supports popups; helps with auto-hiding of popups */
 	FPopupSupport& GetPopupSupport() { return PopupSupport; }
@@ -939,7 +939,7 @@ public:
 	/**
 	 * Forces the window to redraw immediately.
 	 */
-	void ForceRedrawWindow( const TSharedRef<SWindow>& InWindowToDraw );
+	SLATE_API void ForceRedrawWindow( const TSharedRef<SWindow>& InWindowToDraw );
 
 	/**
 	 * Takes a screenshot of the widget writing the results into the color buffer provided.  Note that the format is BGRA.
@@ -947,7 +947,7 @@ public:
 	 * 
 	 * @return true if taking the screenshot was successful.
 	 */
-	bool TakeScreenshot(const TSharedRef<SWidget>& Widget, TArray<FColor>& OutColorData, FIntVector& OutSize);
+	SLATE_API bool TakeScreenshot(const TSharedRef<SWidget>& Widget, TArray<FColor>& OutColorData, FIntVector& OutSize);
 	
 	/**
 	 * Takes a screenshot of the widget writing the results into the color buffer provided.  This is to be used with HDR buffers
@@ -955,7 +955,7 @@ public:
 	 *
 	 * @return true if taking the screenshot was successful.
 	 */
-	bool TakeHDRScreenshot(const TSharedRef<SWidget>& Widget, TArray<FLinearColor>& OutColorData, FIntVector& OutSize);
+	SLATE_API bool TakeHDRScreenshot(const TSharedRef<SWidget>& Widget, TArray<FLinearColor>& OutColorData, FIntVector& OutSize);
 
 	/**
 	 * Takes a screenshot of the widget writing the results into the color buffer provided, this version allows you to provide 
@@ -963,7 +963,7 @@ public:
 	 *
 	 * @return true if taking the screenshot was successful.
 	 */
-	bool TakeScreenshot(const TSharedRef<SWidget>& Widget, const FIntRect& InnerWidgetArea, TArray<FColor>& OutColorData, FIntVector& OutSize);
+	SLATE_API bool TakeScreenshot(const TSharedRef<SWidget>& Widget, const FIntRect& InnerWidgetArea, TArray<FColor>& OutColorData, FIntVector& OutSize);
 
 	/**
 	  * Takes a screenshot of the widget writing the results into the color buffer provided, this version allows you to provide
@@ -971,7 +971,7 @@ public:
 	  *
 	  * @return true if taking the screenshot was successful.
 	*/
-	bool TakeHDRScreenshot(const TSharedRef<SWidget>& Widget, const FIntRect& InnerWidgetArea, TArray<FLinearColor>& OutColorData, FIntVector& OutSize);
+	SLATE_API bool TakeHDRScreenshot(const TSharedRef<SWidget>& Widget, const FIntRect& InnerWidgetArea, TArray<FLinearColor>& OutColorData, FIntVector& OutSize);
 
 	/** Gets the user at the given index, null if the user does not exist. */
 	FORCEINLINE TSharedPtr<const FSlateUser> GetUser(int32 UserIndex) const
@@ -982,7 +982,7 @@ public:
 	{
 		return Users.IsValidIndex(UserIndex) ? Users[UserIndex] : nullptr;
 	}
-	TSharedPtr<FSlateUser> GetUser(FPlatformUserId PlatformUser);
+	SLATE_API TSharedPtr<FSlateUser> GetUser(FPlatformUserId PlatformUser);
 	FORCEINLINE TSharedPtr<const FSlateUser> GetUser(const FInputEvent& InputEvent) const { return GetUser(InputEvent.GetUserIndex()); }
 	FORCEINLINE TSharedPtr<FSlateUser> GetUser(const FInputEvent& InputEvent) { return GetUser(InputEvent.GetUserIndex()); }
 	
@@ -1005,9 +1005,9 @@ public:
 		return nullptr;
 	}
 	
-	TSharedPtr<FSlateUser> GetUserFromPlatformUser(FPlatformUserId PlatformUser);
+	SLATE_API TSharedPtr<FSlateUser> GetUserFromPlatformUser(FPlatformUserId PlatformUser);
 	
-	TSharedPtr<const FSlateUser> GetUserFromPlatformUser(FPlatformUserId PlatformUser) const;
+	SLATE_API TSharedPtr<const FSlateUser> GetUserFromPlatformUser(FPlatformUserId PlatformUser) const;
 
 	/** Get the standard 'default' user (there's always guaranteed to be at least one). */
 	FORCEINLINE TSharedPtr<const FSlateUser> GetCursorUser() const
@@ -1027,16 +1027,16 @@ public:
 	 * @return a handle for the existing or newly created virtual slate user.  This is handy when you need to create
 	 * virtual hardware users for slate components in the virtual world that may need to be interacted with with virtual hardware.
 	 */
-	TSharedRef<FSlateVirtualUserHandle> FindOrCreateVirtualUser(int32 VirtualUserIndex);
+	SLATE_API TSharedRef<FSlateVirtualUserHandle> FindOrCreateVirtualUser(int32 VirtualUserIndex);
 
 	//@todo DanH: This is only getting called by the WidgetInteractionComponent. There's some piping to be set up here for this to be fully accurate for splitscreen
-	void UnregisterUser(int32 UserIndex);
+	SLATE_API void UnregisterUser(int32 UserIndex);
 
 	/** Allows you do some operations for every registered user. */
-	void ForEachUser(TFunctionRef<void(FSlateUser&)> InPredicate, bool bIncludeVirtualUsers = false);
+	SLATE_API void ForEachUser(TFunctionRef<void(FSlateUser&)> InPredicate, bool bIncludeVirtualUsers = false);
 	
 	UE_DEPRECATED(4.23, "ForEachUser now provides an FSlateUser& parameter to the lambda instead of an FSlateUser*")
-	void ForEachUser(TFunctionRef<void(FSlateUser*)> InPredicate, bool bIncludeVirtualUsers = false);
+	SLATE_API void ForEachUser(TFunctionRef<void(FSlateUser*)> InPredicate, bool bIncludeVirtualUsers = false);
 
 	/**
 	 * Gets time step in seconds if a fixed delta time is wanted.
@@ -1051,7 +1051,7 @@ public:
 	 * Sets time step in seconds if a fixed delta time is wanted.
 	 * @param seconds Time step in seconds for fixed delta time.
 	 */
-	static void SetFixedDeltaTime(double InSeconds);
+	static SLATE_API void SetFixedDeltaTime(double InSeconds);
 
 protected:
 	/**
@@ -1059,102 +1059,102 @@ protected:
 	 * a user entry if it gets input from a controller for that index.  Might happen if the user
 	 * allocates the virtual user.
 	 */
-	TSharedRef<FSlateUser> RegisterNewUser(int32 UserIndex, bool bIsVirtual = false);
+	SLATE_API TSharedRef<FSlateUser> RegisterNewUser(int32 UserIndex, bool bIsVirtual = false);
 
 	/**
 	 * Register a new user with Slate.  Normally this is unnecessary as Slate automatically adds
 	 * a user entry if it gets input from a controller for that index.  Might happen if the user
 	 * allocates the virtual user.
 	 */
-	TSharedRef<FSlateUser> RegisterNewUser(FPlatformUserId PlatformUserId, bool bIsVirtual = false);
+	SLATE_API TSharedRef<FSlateUser> RegisterNewUser(FPlatformUserId PlatformUserId, bool bIsVirtual = false);
 
 	/**
 	 * Locates the SlateUser object corresponding to the index, creating a new one if it doesn't exist.
 	 * Asserts if given an invalid (ie, negative) index.
 	 */
-	TSharedRef<FSlateUser> GetOrCreateUser(int32 UserIndex);
+	SLATE_API TSharedRef<FSlateUser> GetOrCreateUser(int32 UserIndex);
 	
 	/**
 	 * Locates the SlateUser object corresponding to the index, creating a new one if it doesn't exist.
 	 * Asserts if given an invalid (ie, negative) index.
 	 */
-	TSharedRef<FSlateUser> GetOrCreateUser(FPlatformUserId PlatformUserId);
+	SLATE_API TSharedRef<FSlateUser> GetOrCreateUser(FPlatformUserId PlatformUserId);
 
 	/**
 	 * Locates the SlateUser object corresponding to the input device id, creating a new one if it doesn't exist.
 	 * This will call GetOrCreateUser for the owning Platform User of the given input device.
 	 * Asserts if given an invalid (ie, negative) index.
 	 */
-	TSharedRef<FSlateUser> GetOrCreateUser(FInputDeviceId DeviceId);
+	SLATE_API TSharedRef<FSlateUser> GetOrCreateUser(FInputDeviceId DeviceId);
 	
 	FORCEINLINE TSharedRef<FSlateUser> GetOrCreateUser(const FInputEvent& InputEvent) { return GetOrCreateUser(InputEvent.GetUserIndex()); }
 
 	friend class FEventRouter;
 
 	/** Transforms a pointer event to account for non-standard viewport resolutions */
-	FPointerEvent TransformPointerEvent(const FPointerEvent& PointerEvent, const TSharedPtr<SWindow>& Window) const;
+	SLATE_API FPointerEvent TransformPointerEvent(const FPointerEvent& PointerEvent, const TSharedPtr<SWindow>& Window) const;
 
-	virtual bool DoesWidgetHaveMouseCaptureByUser(const TSharedPtr<const SWidget> Widget, int32 UserIndex, TOptional<int32> PointerIndex) const override;
-	virtual bool DoesWidgetHaveMouseCapture(const TSharedPtr<const SWidget> Widget) const override;
+	SLATE_API virtual bool DoesWidgetHaveMouseCaptureByUser(const TSharedPtr<const SWidget> Widget, int32 UserIndex, TOptional<int32> PointerIndex) const override;
+	SLATE_API virtual bool DoesWidgetHaveMouseCapture(const TSharedPtr<const SWidget> Widget) const override;
 
-	virtual TOptional<EFocusCause> HasUserFocus(const TSharedPtr<const SWidget> Widget, int32 UserIndex) const override;
-	virtual TOptional<EFocusCause> HasAnyUserFocus(const TSharedPtr<const SWidget> Widget) const override;
-	virtual bool IsWidgetDirectlyHovered(const TSharedPtr<const SWidget> Widget) const override;
-	virtual bool ShowUserFocus(const TSharedPtr<const SWidget> Widget) const override;
+	SLATE_API virtual TOptional<EFocusCause> HasUserFocus(const TSharedPtr<const SWidget> Widget, int32 UserIndex) const override;
+	SLATE_API virtual TOptional<EFocusCause> HasAnyUserFocus(const TSharedPtr<const SWidget> Widget) const override;
+	SLATE_API virtual bool IsWidgetDirectlyHovered(const TSharedPtr<const SWidget> Widget) const override;
+	SLATE_API virtual bool ShowUserFocus(const TSharedPtr<const SWidget> Widget) const override;
 
-	TSharedRef<FNavigationConfig> GetRelevantNavConfig(int32 UserIndex) const;
+	SLATE_API TSharedRef<FNavigationConfig> GetRelevantNavConfig(int32 UserIndex) const;
 
 	/** Called when the slate application is being shut down. */
-	void OnShutdown();
+	SLATE_API void OnShutdown();
 
-	void DestroyRenderer();
+	SLATE_API void DestroyRenderer();
 
 	/** Advances time for the application. */
-	void TickTime();
+	SLATE_API void TickTime();
 
 	/**
 	 * Pumps and ticks the platform.
 	 */
-	void TickPlatform(float DeltaTime);
+	SLATE_API void TickPlatform(float DeltaTime);
 
 	/**
 	 * Ticks and paints the actual Slate portion of the application.
 	 */
-	void TickAndDrawWidgets(float DeltaTime);
+	SLATE_API void TickAndDrawWidgets(float DeltaTime);
 
 	/** Draws Slate windows. Should only be called by the application's main loop or renderer. */
-	void DrawWindows();
+	SLATE_API void DrawWindows();
 
 	/**
 	 * Draws slate windows, optionally only drawing the passed in window
 	 */
-	void PrivateDrawWindows( TSharedPtr<SWindow> DrawOnlyThisWindow = nullptr );
+	SLATE_API void PrivateDrawWindows( TSharedPtr<SWindow> DrawOnlyThisWindow = nullptr );
 
 	/**
 	 * Pre-pass step before drawing windows to compute geometry size and reshape autosized windows
 	 */
-	void DrawPrepass( TSharedPtr<SWindow> DrawOnlyThisWindow );
+	SLATE_API void DrawPrepass( TSharedPtr<SWindow> DrawOnlyThisWindow );
 
 	/**
 	 * Draws a window and its children
 	 */
-	void DrawWindowAndChildren( const TSharedRef<SWindow>& WindowToDraw, struct FDrawWindowArgs& DrawWindowArgs );
+	SLATE_API void DrawWindowAndChildren( const TSharedRef<SWindow>& WindowToDraw, struct FDrawWindowArgs& DrawWindowArgs );
 
 	/**
 	 * Gets all visible child windows of a window.
 	 */
-	void GetAllVisibleChildWindows(TArray< TSharedRef<SWindow> >& OutWindows, TSharedRef<SWindow> CurrentWindow);
+	SLATE_API void GetAllVisibleChildWindows(TArray< TSharedRef<SWindow> >& OutWindows, TSharedRef<SWindow> CurrentWindow);
 
 	/** Engages or disengages application throttling based on user behavior */
-	void ThrottleApplicationBasedOnMouseMovement();
+	SLATE_API void ThrottleApplicationBasedOnMouseMovement();
 
-	virtual FWidgetPath LocateWidgetInWindow(UE::Slate::FDeprecateVector2DParameter ScreenspaceMouseCoordinate, const TSharedRef<SWindow>& Window, bool bIgnoreEnabledStatus, int32 UserIndex) const override;
+	SLATE_API virtual FWidgetPath LocateWidgetInWindow(UE::Slate::FDeprecateVector2DParameter ScreenspaceMouseCoordinate, const TSharedRef<SWindow>& Window, bool bIgnoreEnabledStatus, int32 UserIndex) const override;
 
 	/**
 	 * Sets up any values that need to be based on the physical dimensions of the device.  
 	 * Such as dead zones associated with precise tapping...etc
 	 */
-	void SetupPhysicalSensitivities();
+	SLATE_API void SetupPhysicalSensitivities();
 
 public:
 
@@ -1165,7 +1165,7 @@ public:
 	 * @param  bIsSynthetic  True when the even is synthesized by slate.
 	 * @return  Was this event handled by the Slate application?
 	 */
-	bool ProcessMouseMoveEvent( const FPointerEvent& MouseEvent, bool bIsSynthetic = false );
+	SLATE_API bool ProcessMouseMoveEvent( const FPointerEvent& MouseEvent, bool bIsSynthetic = false );
 
 	/**
 	 * Called by the native application in response to a mouse button press. Routs the event to Slate Widgets.
@@ -1175,7 +1175,7 @@ public:
 	 * @param  InMouseEvent    Mouse event
 	 * @return  Was this event handled by the Slate application?
 	 */
-	bool ProcessMouseButtonDownEvent(const TSharedPtr< FGenericWindow >& PlatformWindow, const FPointerEvent& InMouseEvent);
+	SLATE_API bool ProcessMouseButtonDownEvent(const TSharedPtr< FGenericWindow >& PlatformWindow, const FPointerEvent& InMouseEvent);
 
 	/**
 	 * Called by the native application in response to a mouse button release. Routs the event to Slate Widgets.
@@ -1183,7 +1183,7 @@ public:
 	 * @param  InMouseEvent  Mouse event
 	 * @return  Was this event handled by the Slate application?
 	 */
-	bool ProcessMouseButtonUpEvent( const FPointerEvent& MouseEvent );
+	SLATE_API bool ProcessMouseButtonUpEvent( const FPointerEvent& MouseEvent );
 
 	/**
 	 * Called by the native application in response to a mouse release. Routs the event to Slate Widgets.
@@ -1191,7 +1191,7 @@ public:
 	 * @param  InMouseEvent  Mouse event
 	 * @return  Was this event handled by the Slate application?
 	 */
-	bool ProcessMouseButtonDoubleClickEvent( const TSharedPtr< FGenericWindow >& PlatformWindow, const FPointerEvent& InMouseEvent );
+	SLATE_API bool ProcessMouseButtonDoubleClickEvent( const TSharedPtr< FGenericWindow >& PlatformWindow, const FPointerEvent& InMouseEvent );
 	
 	/**
 	 * Called by the native application in response to a mouse wheel spin or a touch gesture. Routs the event to Slate Widgets.
@@ -1200,7 +1200,7 @@ public:
 	 * @param  InGestureEvent  Optional gesture event details
 	 * @return  Was this event handled by the Slate application?
 	 */
-	bool ProcessMouseWheelOrGestureEvent( const FPointerEvent& InWheelEvent, const FPointerEvent* InGestureEvent );
+	SLATE_API bool ProcessMouseWheelOrGestureEvent( const FPointerEvent& InWheelEvent, const FPointerEvent* InGestureEvent );
 
 	/**
 	 * Called when a character is entered
@@ -1208,7 +1208,7 @@ public:
 	 * @param  InCharacterEvent  Character event
 	 * @return  Was this event handled by the Slate application?
 	 */
-	bool ProcessKeyCharEvent( const FCharacterEvent& InCharacterEvent );
+	SLATE_API bool ProcessKeyCharEvent( const FCharacterEvent& InCharacterEvent );
 
 	/**
 	 * Called when a key is pressed
@@ -1216,7 +1216,7 @@ public:
 	 * @param  InKeyEvent  Keyb event
 	 * @return  Was this event handled by the Slate application?
 	 */
-	bool ProcessKeyDownEvent( const FKeyEvent& InKeyEvent );
+	SLATE_API bool ProcessKeyDownEvent( const FKeyEvent& InKeyEvent );
 
 	/**
 	 * Called when a key is released
@@ -1224,7 +1224,7 @@ public:
 	 * @param  InKeyEvent  Key event
 	 * @return  Was this event handled by the Slate application?
 	 */
-	bool ProcessKeyUpEvent( const FKeyEvent& InKeyEvent );
+	SLATE_API bool ProcessKeyUpEvent( const FKeyEvent& InKeyEvent );
 	
 	/**
 	 * Called when a analog input values change
@@ -1232,7 +1232,7 @@ public:
 	 * @param  InAnalogInputEvent Analog input event
 	 * @return  Was this event handled by the Slate application?
 	 */
-	bool ProcessAnalogInputEvent(const FAnalogInputEvent& InAnalogInputEvent);
+	SLATE_API bool ProcessAnalogInputEvent(const FAnalogInputEvent& InAnalogInputEvent);
 
 	/**
 	 * Called when a drag from an external (non-slate) source enters a window
@@ -1241,35 +1241,35 @@ public:
 	 * @param DragDropEvent  Describes the mouse state (position, pressed buttons, etc) and associated payload
 	 * @return true if the drag enter was handled and can be processed by some widget in this window; false otherwise
 	 */
-	bool ProcessDragEnterEvent( TSharedRef<SWindow> WindowEntered, const FDragDropEvent& DragDropEvent );
+	SLATE_API bool ProcessDragEnterEvent( TSharedRef<SWindow> WindowEntered, const FDragDropEvent& DragDropEvent );
 
 	/**
 	 * Called when a touchpad touch is started (finger down) when polling game device state
 	 * 
 	 * @param ControllerEvent	The touch event generated
 	 */
-	void ProcessTouchStartedEvent( const TSharedPtr< FGenericWindow >& PlatformWindow, const FPointerEvent& InTouchEvent );
+	SLATE_API void ProcessTouchStartedEvent( const TSharedPtr< FGenericWindow >& PlatformWindow, const FPointerEvent& InTouchEvent );
 
 	/**
 	 * Called when a touchpad touch is moved  (finger moved) when polling game device state
 	 * 
 	 * @param ControllerEvent	The touch event generated
 	 */
-	void ProcessTouchMovedEvent( const FPointerEvent& InTouchEvent );
+	SLATE_API void ProcessTouchMovedEvent( const FPointerEvent& InTouchEvent );
 
 	/**
 	 * Called when a touchpad touch is ended (finger lifted) when polling game device state
 	 * 
 	 * @param ControllerEvent	The touch event generated
 	 */
-	void ProcessTouchEndedEvent( const FPointerEvent& InTouchEvent );
+	SLATE_API void ProcessTouchEndedEvent( const FPointerEvent& InTouchEvent );
 
 	/**
 	 * Called when motion is detected (controller or device) when polling game device state
 	 * 
 	 * @param MotionEvent		The motion event generated
 	 */
-	void ProcessMotionDetectedEvent( const FMotionEvent& InMotionEvent );
+	SLATE_API void ProcessMotionDetectedEvent( const FMotionEvent& InMotionEvent );
 
 	/**
 	 * Called by the native application in response to an activation or deactivation. 
@@ -1277,14 +1277,14 @@ public:
 	 * @param ActivateEvent Information about the window activation/deactivation
 	 * @return  Was this event handled by the Slate application?
 	 */
-	bool ProcessWindowActivatedEvent( const FWindowActivateEvent& ActivateEvent );
+	SLATE_API bool ProcessWindowActivatedEvent( const FWindowActivateEvent& ActivateEvent );
 	
 	/**
 	 * Called when the application is activated (i.e. one of its windows becomes active) or deactivated.
 	 *
 	 * @param InAppActivated Whether the application was activated.
 	 */
-	void ProcessApplicationActivationEvent( bool InAppActivated );
+	SLATE_API void ProcessApplicationActivationEvent( bool InAppActivated );
 
 	/**
 	 * Returns true if the we're currently processing mouse, keyboard, touch or gamepad input.
@@ -1299,7 +1299,7 @@ public:
 	 * Sets the navigation config.  If you need to control navigation config dynamically, you
 	 * should subclass FNavigationConfig to be dynamically adjustable to your needs.
 	 */
-	void SetNavigationConfig(TSharedRef<FNavigationConfig> InNavigationConfig);
+	SLATE_API void SetNavigationConfig(TSharedRef<FNavigationConfig> InNavigationConfig);
 
 	/**
 	 * Sets the navigation config factory.  If you need to control navigation config dynamically, you
@@ -1311,14 +1311,14 @@ public:
 public:
 
 	/** Closes all active windows immediately */
-	void CloseAllWindowsImmediately();
+	SLATE_API void CloseAllWindowsImmediately();
 
 	/**
 	 * Destroy the native and slate windows in the array provided.
 	 *
 	 * @param WindowsToDestroy   Destroy these windows
 	 */
-	void DestroyWindowsImmediately();
+	SLATE_API void DestroyWindowsImmediately();
 
 	/**
 	 * Apply any requests from the Reply to the application. E.g. Capture mouse
@@ -1328,7 +1328,7 @@ public:
 	 * @param UserIndex			 User index that generated the event we are replying to (defaults to 0, at least for now)
 	 * @param PointerIndex		 Pointer index that generated the event we are replying to
 	 */
-	void ProcessExternalReply(const FWidgetPath& CurrentEventPath, const FReply TheReply, const int32 UserIndex = 0, const int32 PointerIndex = 10 /* todo: use the enum */);
+	SLATE_API void ProcessExternalReply(const FWidgetPath& CurrentEventPath, const FReply TheReply, const int32 UserIndex = 0, const int32 PointerIndex = 10 /* todo: use the enum */);
 
 	/**
 	 * Apply any requests from the Reply to the application. E.g. Capture mouse
@@ -1339,17 +1339,17 @@ public:
 	 * @param InMouseEvent       Optional mouse event that caused this action.
 	 * @param UserIndex			 User index that generated the event we are replying to (defaults to 0, at least for now)
 	 */
-	void ProcessReply(const FWidgetPath& CurrentEventPath, const FReply& TheReply, const FWidgetPath* WidgetsUnderMouse, const FPointerEvent* InMouseEvent, const uint32 UserIndex = 0);
+	SLATE_API void ProcessReply(const FWidgetPath& CurrentEventPath, const FReply& TheReply, const FWidgetPath* WidgetsUnderMouse, const FPointerEvent* InMouseEvent, const uint32 UserIndex = 0);
 	
 	/** Bubble a request for which cursor to display for widgets under the mouse or the widget that captured the mouse. */
-	void QueryCursor();
+	SLATE_API void QueryCursor();
 
 	/**
 	 * Apply any requests from the CursorReply
 	 *
 	 * @param CursorReply        The reply generated by an event that was being processed.
 	 */
-	void ProcessCursorReply(const FCursorReply& CursorReply);
+	SLATE_API void ProcessCursorReply(const FCursorReply& CursorReply);
 	
 	/**
 	 * Spawns a tool tip window.  If an existing tool tip window is open, it will be dismissed first.
@@ -1357,30 +1357,30 @@ public:
 	 * @param	InToolTip           Widget to display.
 	 * @param	InSpawnLocation     Screen space location to show the tool tip (window top left)
 	 */
-	void SpawnToolTip( const TSharedRef<IToolTip>& InToolTip, const UE::Slate::FDeprecateVector2DParameter& InSpawnLocation );
+	SLATE_API void SpawnToolTip( const TSharedRef<IToolTip>& InToolTip, const UE::Slate::FDeprecateVector2DParameter& InSpawnLocation );
 
 	/** Closes the open tool-tip, if a tool-tip is open */
-	void CloseToolTip();
+	SLATE_API void CloseToolTip();
 
-	void UpdateToolTip( bool bAllowSpawningOfNewToolTips );
+	SLATE_API void UpdateToolTip( bool bAllowSpawningOfNewToolTips );
 
 	/** @return an array of top-level windows that can be interacted with. e.g. when a modal window is up, only return the modal window */
-	TArray< TSharedRef<SWindow> > GetInteractiveTopLevelWindows();
+	SLATE_API TArray< TSharedRef<SWindow> > GetInteractiveTopLevelWindows();
 
 	/** Gets all visible slate windows ordered from back to front based on child hierarchies */
-	void GetAllVisibleWindowsOrdered(TArray< TSharedRef<SWindow> >& OutWindows);
+	SLATE_API void GetAllVisibleWindowsOrdered(TArray< TSharedRef<SWindow> >& OutWindows);
 	
 	/** @return true if mouse events are being turned into touch events, and touch UI should be forced on */
-	bool IsFakingTouchEvents() const;
+	SLATE_API bool IsFakingTouchEvents() const;
 
 	/** Sets whether the application is treating mouse events as imitating touch events.  Optional CursorLocation can be supplied to override the platform's belief of where the cursor is */
-	void SetGameIsFakingTouchEvents(const bool bIsFaking, FVector2D* CursorLocation = nullptr);
+	SLATE_API void SetGameIsFakingTouchEvents(const bool bIsFaking, FVector2D* CursorLocation = nullptr);
 
 	/** Sets the handler for otherwise unhandled key down events. This is used by the editor to provide a global action list, if the key was not consumed by any widget. */
-	void SetUnhandledKeyDownEventHandler( const FOnKeyEvent& NewHandler );
+	SLATE_API void SetUnhandledKeyDownEventHandler( const FOnKeyEvent& NewHandler );
 
 	/** Sets the handler for otherwise unhandled key down events. This is used by the editor to provide a global action list, if the key was not consumed by any widget. */
-	void SetUnhandledKeyUpEventHandler(const FOnKeyEvent& NewHandler);
+	SLATE_API void SetUnhandledKeyUpEventHandler(const FOnKeyEvent& NewHandler);
 
 	/** @return the last time a user interacted with a keyboard, mouse, touch device, or controller */
 	double GetLastUserInteractionTime() const { return LastUserInteractionTime; }
@@ -1390,17 +1390,17 @@ public:
 	FSlateLastUserInteractionTimeUpdateEvent& GetLastUserInteractionTimeUpdateEvent() { return LastUserInteractionTimeUpdateEvent; }
 
 	/** @return the deadzone size for dragging in screen pixels (aka virtual desktop pixels) */
-	float GetDragTriggerDistance() const;
+	SLATE_API float GetDragTriggerDistance() const;
 
 	/** @return the deadzone size squared for dragging in screen pixels (aka virtual desktop pixels) */
-	float GetDragTriggerDistanceSquared() const;
+	SLATE_API float GetDragTriggerDistanceSquared() const;
 
 	/** @return true if the difference between the ScreenSpaceOrigin and the ScreenSpacePosition is larger than the trigger distance for dragging in Slate. */
-	bool HasTraveledFarEnoughToTriggerDrag(const FPointerEvent& PointerEvent, const UE::Slate::FDeprecateVector2DParameter ScreenSpaceOrigin) const;
-	bool HasTraveledFarEnoughToTriggerDrag(const FPointerEvent& PointerEvent, const UE::Slate::FDeprecateVector2DParameter ScreenSpaceOrigin, EOrientation Orientation) const;
+	SLATE_API bool HasTraveledFarEnoughToTriggerDrag(const FPointerEvent& PointerEvent, const UE::Slate::FDeprecateVector2DParameter ScreenSpaceOrigin) const;
+	SLATE_API bool HasTraveledFarEnoughToTriggerDrag(const FPointerEvent& PointerEvent, const UE::Slate::FDeprecateVector2DParameter ScreenSpaceOrigin, EOrientation Orientation) const;
 
 	/** Set the size of the deadzone for dragging in screen pixels */
-	void SetDragTriggerDistance( float ScreenPixels );
+	SLATE_API void SetDragTriggerDistance( float ScreenPixels );
 	
 	/** 
 	 * Adds input pre-processor if unique. 
@@ -1408,29 +1408,29 @@ public:
 	 * @param Index				Where to insert the InputProcessor, when sorting is needed. Default index will add at the end.
 	 * @return True if added to list of input pre-processors, false if not
 	 */
-	bool RegisterInputPreProcessor(TSharedPtr<class IInputProcessor> InputProcessor, const int32 Index = INDEX_NONE);
+	SLATE_API bool RegisterInputPreProcessor(TSharedPtr<class IInputProcessor> InputProcessor, const int32 Index = INDEX_NONE);
 
 	/**
 	 * Removes an input pre-processor.
 	 * @param InputProcessor	The input pre-processor to Remove.
 	 */
-	void UnregisterInputPreProcessor(TSharedPtr<class IInputProcessor> InputProcessor);
+	SLATE_API void UnregisterInputPreProcessor(TSharedPtr<class IInputProcessor> InputProcessor);
 
 	/**
 	 * Get the index of a registered pre-processor.
 	 * @param InputProcessor	The input pre-processor to find.
 	 * @return The index of the pre-processor, or INDEX_NONE if not registered.
 	 */
-	int32 FindInputPreProcessor(TSharedPtr<class IInputProcessor> InputProcessor) const;
+	SLATE_API int32 FindInputPreProcessor(TSharedPtr<class IInputProcessor> InputProcessor) const;
 
 	/** Sets the hit detection radius of the cursor */
-	void SetCursorRadius(float NewRadius);
+	SLATE_API void SetCursorRadius(float NewRadius);
 
 	/** Getter for the cursor radius */
-	float GetCursorRadius() const;
+	SLATE_API float GetCursorRadius() const;
 
-	void SetAllowTooltips(bool bCanShow);
-	bool GetAllowTooltips() const;
+	SLATE_API void SetAllowTooltips(bool bCanShow);
+	SLATE_API bool GetAllowTooltips() const;
 	
 	bool IsRenderingOffScreen() const { return bRenderOffScreen; }
 
@@ -1443,7 +1443,7 @@ public:
 		return bAppIsActive;
 	}
 
-	virtual TSharedRef<SWindow> AddWindow( TSharedRef<SWindow> InSlateWindow, const bool bShowImmediately = true ) override;
+	SLATE_API virtual TSharedRef<SWindow> AddWindow( TSharedRef<SWindow> InSlateWindow, const bool bShowImmediately = true ) override;
 
 	virtual void ArrangeWindowToFrontVirtual( TArray<TSharedRef<SWindow>>& Windows, const TSharedRef<SWindow>& WindowToBringToFront ) override
 	{
@@ -1465,22 +1465,22 @@ public:
 		return CurrentTime;
 	}
 
-	virtual TSharedPtr<SWindow> GetActiveTopLevelWindow() const override;
+	SLATE_API virtual TSharedPtr<SWindow> GetActiveTopLevelWindow() const override;
 
-	virtual TSharedPtr<SWindow> GetActiveTopLevelRegularWindow() const override;
+	SLATE_API virtual TSharedPtr<SWindow> GetActiveTopLevelRegularWindow() const override;
 
-	virtual const FSlateBrush* GetAppIcon() const override;
-	virtual const FSlateBrush* GetAppIconSmall() const override;
+	SLATE_API virtual const FSlateBrush* GetAppIcon() const override;
+	SLATE_API virtual const FSlateBrush* GetAppIconSmall() const override;
 
 	virtual float GetApplicationScale() const override { return Scale; }
 	virtual bool GetSoftwareCursorAvailable() const override { return bSoftwareCursorAvailable; }
-	virtual EVisibility GetSoftwareCursorVis() const override;
+	SLATE_API virtual EVisibility GetSoftwareCursorVis() const override;
 
-	virtual UE::Slate::FDeprecateVector2DResult GetCursorPos() const override;
-	virtual UE::Slate::FDeprecateVector2DResult GetLastCursorPos() const override;
-	virtual UE::Slate::FDeprecateVector2DResult GetCursorSize() const override;
+	SLATE_API virtual UE::Slate::FDeprecateVector2DResult GetCursorPos() const override;
+	SLATE_API virtual UE::Slate::FDeprecateVector2DResult GetLastCursorPos() const override;
+	SLATE_API virtual UE::Slate::FDeprecateVector2DResult GetCursorSize() const override;
 
-	virtual TSharedPtr<SWidget> GetKeyboardFocusedWidget() const override;
+	SLATE_API virtual TSharedPtr<SWidget> GetKeyboardFocusedWidget() const override;
 
 	virtual EWindowTransparency GetWindowTransparencySupport() const override
 	{
@@ -1489,31 +1489,31 @@ public:
 
 protected:
 
-	virtual TSharedPtr< SWidget > GetMouseCaptorImpl() const override;
+	SLATE_API virtual TSharedPtr< SWidget > GetMouseCaptorImpl() const override;
 
 public:
 
 	//~ Begin FSlateApplicationBase Interface
 
-	virtual bool HasAnyMouseCaptor() const override;
-	virtual bool HasUserMouseCapture(int32 UserIndex) const override;
-	virtual FSlateRect GetPreferredWorkArea() const override;
-	virtual bool HasFocusedDescendants( const TSharedRef<const SWidget>& Widget ) const override;
-	virtual bool HasUserFocusedDescendants(const TSharedRef< const SWidget >& Widget, int32 UserIndex) const override;
-	virtual bool IsExternalUIOpened() override;
-	virtual FWidgetPath LocateWindowUnderMouse( UE::Slate::FDeprecateVector2DParameter ScreenspaceMouseCoordinate, const TArray<TSharedRef<SWindow>>& Windows, bool bIgnoreEnabledStatus = false, int32 UserIndex = INDEX_NONE) override;
-	virtual bool IsWindowHousingInteractiveTooltip(const TSharedRef<const SWindow>& WindowToTest) const override;
-	virtual TSharedRef<SImage> MakeImage( const TAttribute<const FSlateBrush*>& Image, const TAttribute<FSlateColor>& Color, const TAttribute<EVisibility>& Visibility ) const override;
-	virtual TSharedRef<SWidget> MakeWindowTitleBar(const FWindowTitleBarArgs& InArgs, TSharedPtr<IWindowTitleBar>& OutTitleBar) const override;
-	virtual TSharedRef<IToolTip> MakeToolTip( const TAttribute<FText>& ToolTipText ) override;
-	virtual TSharedRef<IToolTip> MakeToolTip( const FText& ToolTipText ) override;
-	virtual void RequestDestroyWindow( TSharedRef<SWindow> WindowToDestroy ) override;
-	virtual bool SetKeyboardFocus( const FWidgetPath& InFocusPath, const EFocusCause InCause ) override;
-	virtual bool SetUserFocus(const uint32 InUserIndex, const FWidgetPath& InFocusPath, const EFocusCause InCause) override;
-	virtual void SetAllUserFocus(const FWidgetPath& InFocusPath, const EFocusCause InCause) override;
-	virtual void SetAllUserFocusAllowingDescendantFocus(const FWidgetPath& InFocusPath, const EFocusCause InCause) override;
-	virtual TSharedPtr<SWidget> GetUserFocusedWidget(uint32 UserIndex) const override;
-	virtual TSharedPtr<SWidget> GetCurrentDebugContextWidget() const override;
+	SLATE_API virtual bool HasAnyMouseCaptor() const override;
+	SLATE_API virtual bool HasUserMouseCapture(int32 UserIndex) const override;
+	SLATE_API virtual FSlateRect GetPreferredWorkArea() const override;
+	SLATE_API virtual bool HasFocusedDescendants( const TSharedRef<const SWidget>& Widget ) const override;
+	SLATE_API virtual bool HasUserFocusedDescendants(const TSharedRef< const SWidget >& Widget, int32 UserIndex) const override;
+	SLATE_API virtual bool IsExternalUIOpened() override;
+	SLATE_API virtual FWidgetPath LocateWindowUnderMouse( UE::Slate::FDeprecateVector2DParameter ScreenspaceMouseCoordinate, const TArray<TSharedRef<SWindow>>& Windows, bool bIgnoreEnabledStatus = false, int32 UserIndex = INDEX_NONE) override;
+	SLATE_API virtual bool IsWindowHousingInteractiveTooltip(const TSharedRef<const SWindow>& WindowToTest) const override;
+	SLATE_API virtual TSharedRef<SImage> MakeImage( const TAttribute<const FSlateBrush*>& Image, const TAttribute<FSlateColor>& Color, const TAttribute<EVisibility>& Visibility ) const override;
+	SLATE_API virtual TSharedRef<SWidget> MakeWindowTitleBar(const FWindowTitleBarArgs& InArgs, TSharedPtr<IWindowTitleBar>& OutTitleBar) const override;
+	SLATE_API virtual TSharedRef<IToolTip> MakeToolTip( const TAttribute<FText>& ToolTipText ) override;
+	SLATE_API virtual TSharedRef<IToolTip> MakeToolTip( const FText& ToolTipText ) override;
+	SLATE_API virtual void RequestDestroyWindow( TSharedRef<SWindow> WindowToDestroy ) override;
+	SLATE_API virtual bool SetKeyboardFocus( const FWidgetPath& InFocusPath, const EFocusCause InCause ) override;
+	SLATE_API virtual bool SetUserFocus(const uint32 InUserIndex, const FWidgetPath& InFocusPath, const EFocusCause InCause) override;
+	SLATE_API virtual void SetAllUserFocus(const FWidgetPath& InFocusPath, const EFocusCause InCause) override;
+	SLATE_API virtual void SetAllUserFocusAllowingDescendantFocus(const FWidgetPath& InFocusPath, const EFocusCause InCause) override;
+	SLATE_API virtual TSharedPtr<SWidget> GetUserFocusedWidget(uint32 UserIndex) const override;
+	SLATE_API virtual TSharedPtr<SWidget> GetCurrentDebugContextWidget() const override;
 	virtual const TArray<TSharedRef<SWindow>> GetTopLevelWindows() const override { return SlateWindows; }
 
 	DECLARE_EVENT_OneParam(FSlateApplication, FApplicationActivationStateChangedEvent, const bool /*IsActive*/)
@@ -1523,57 +1523,57 @@ public:
 
 	//~ Begin FGenericApplicationMessageHandler Interface
 
-	virtual bool ShouldProcessUserInputMessages( const TSharedPtr< FGenericWindow >& PlatformWindow ) const override;
-	virtual bool OnKeyChar( const TCHAR Character, const bool IsRepeat ) override;
-	virtual bool OnKeyDown( const int32 KeyCode, const uint32 CharacterCode, const bool IsRepeat ) override;
-	virtual bool OnKeyUp( const int32 KeyCode, const uint32 CharacterCode, const bool IsRepeat ) override;
-	virtual void OnInputLanguageChanged() override;
-	virtual bool OnMouseDown( const TSharedPtr< FGenericWindow >& PlatformWindow, const EMouseButtons::Type Button ) override;
-	virtual bool OnMouseDown( const TSharedPtr< FGenericWindow >& PlatformWindow, const EMouseButtons::Type Button, const FVector2D CursorPos ) override;
-	virtual bool OnMouseUp( const EMouseButtons::Type Button ) override;
-	virtual bool OnMouseUp( const EMouseButtons::Type Button, const FVector2D CursorPos ) override;
-	virtual bool OnMouseDoubleClick( const TSharedPtr< FGenericWindow >& PlatformWindow, const EMouseButtons::Type Button ) override;
-	virtual bool OnMouseDoubleClick( const TSharedPtr< FGenericWindow >& PlatformWindow, const EMouseButtons::Type Button, const FVector2D CursorPos ) override;
-	virtual bool OnMouseWheel( const float Delta ) override;
-	virtual bool OnMouseWheel( const float Delta, const FVector2D CursorPos ) override;
-	virtual bool OnMouseMove() override;
-	virtual bool OnRawMouseMove( const int32 X, const int32 Y ) override;
-	virtual bool OnCursorSet() override;
-	virtual bool OnTouchGesture( EGestureEvent GestureType, const FVector2D& Delta, float WheelDelta, bool bIsDirectionInvertedFromDevice ) override;
-	virtual bool OnTouchStarted( const TSharedPtr< FGenericWindow >& Window, const FVector2D& Location, float Force, int32 TouchIndex, FPlatformUserId PlatformUserId, FInputDeviceId DeviceId ) override;
-	virtual bool OnTouchMoved( const FVector2D& Location, float Force, int32 TouchIndex, FPlatformUserId PlatformUserId, FInputDeviceId DeviceID ) override;
-	virtual bool OnTouchEnded( const FVector2D& Location, int32 TouchIndex, FPlatformUserId PlatformUserId, FInputDeviceId DeviceID ) override;
-	virtual bool OnTouchForceChanged(const FVector2D& Location, float Force, int32 TouchIndex, FPlatformUserId PlatformUserId, FInputDeviceId DeviceID) override;
-	virtual bool OnTouchFirstMove(const FVector2D& Location, float Force, int32 TouchIndex, FPlatformUserId PlatformUserId, FInputDeviceId DeviceID) override;
-	virtual void ShouldSimulateGesture(EGestureEvent Gesture, bool bEnable) override;
-	virtual bool OnMotionDetected(const FVector& Tilt, const FVector& RotationRate, const FVector& Gravity, const FVector& Acceleration, FPlatformUserId PlatformUserId, FInputDeviceId InputDeviceId) override;
-	virtual bool OnControllerAnalog(FGamepadKeyNames::Type KeyName, FPlatformUserId PlatformUserId, FInputDeviceId InputDeviceId, float AnalogValue) override;
-	virtual bool OnControllerButtonPressed(FGamepadKeyNames::Type KeyName, FPlatformUserId PlatformUserId, FInputDeviceId InputDeviceId, bool IsRepeat) override;
-	virtual bool OnControllerButtonReleased(FGamepadKeyNames::Type KeyName, FPlatformUserId PlatformUserId, FInputDeviceId InputDeviceId, bool IsRepeat) override;
-	virtual bool OnSizeChanged( const TSharedRef< FGenericWindow >& PlatformWindow, const int32 Width, const int32 Height, bool bWasMinimized = false ) override;
-	virtual void OnOSPaint( const TSharedRef< FGenericWindow >& PlatformWindow ) override;
-	virtual FWindowSizeLimits GetSizeLimitsForWindow(const TSharedRef<FGenericWindow>& Window) const override;
-	virtual void OnResizingWindow( const TSharedRef< FGenericWindow >& PlatformWindow ) override;
-	virtual bool BeginReshapingWindow( const TSharedRef< FGenericWindow >& PlatformWindow ) override;
-	virtual void FinishedReshapingWindow( const TSharedRef< FGenericWindow >& PlatformWindow ) override;
-	virtual void SignalSystemDPIChanged(const TSharedRef<FGenericWindow>& Window) override;
-	virtual void HandleDPIScaleChanged(const TSharedRef<FGenericWindow>& Window) override;
-	virtual void OnMovedWindow( const TSharedRef< FGenericWindow >& PlatformWindow, const int32 X, const int32 Y ) override;
-	virtual bool OnWindowActivationChanged( const TSharedRef< FGenericWindow >& PlatformWindow, const EWindowActivation ActivationType ) override;
-	virtual bool OnApplicationActivationChanged( const bool IsActive ) override;
-	virtual bool OnConvertibleLaptopModeChanged() override;
-	virtual EWindowZone::Type GetWindowZoneForPoint( const TSharedRef< FGenericWindow >& PlatformWindow, const int32 X, const int32 Y ) override;
-	virtual void OnWindowClose( const TSharedRef< FGenericWindow >& PlatformWindow ) override;
-	virtual EDropEffect::Type OnDragEnterText( const TSharedRef< FGenericWindow >& Window, const FString& Text ) override;
-	virtual EDropEffect::Type OnDragEnterFiles( const TSharedRef< FGenericWindow >& Window, const TArray< FString >& Files ) override;
-	virtual EDropEffect::Type OnDragEnterExternal( const TSharedRef< FGenericWindow >& Window, const FString& Text, const TArray< FString >& Files ) override;
-
-	EDropEffect::Type OnDragEnter( const TSharedRef< SWindow >& Window, const TSharedRef<FExternalDragOperation>& DragDropOperation );
-
-	virtual EDropEffect::Type OnDragOver( const TSharedPtr< FGenericWindow >& Window ) override;
-	virtual void OnDragLeave( const TSharedPtr< FGenericWindow >& Window ) override;
-	virtual EDropEffect::Type OnDragDrop( const TSharedPtr< FGenericWindow >& Window ) override;
-	virtual bool OnWindowAction( const TSharedRef< FGenericWindow >& PlatformWindow, const EWindowAction::Type InActionType ) override;
+	SLATE_API virtual bool ShouldProcessUserInputMessages( const TSharedPtr< FGenericWindow >& PlatformWindow ) const override;
+	SLATE_API virtual bool OnKeyChar( const TCHAR Character, const bool IsRepeat ) override;
+	SLATE_API virtual bool OnKeyDown( const int32 KeyCode, const uint32 CharacterCode, const bool IsRepeat ) override;
+	SLATE_API virtual bool OnKeyUp( const int32 KeyCode, const uint32 CharacterCode, const bool IsRepeat ) override;
+	SLATE_API virtual void OnInputLanguageChanged() override;
+	SLATE_API virtual bool OnMouseDown( const TSharedPtr< FGenericWindow >& PlatformWindow, const EMouseButtons::Type Button ) override;
+	SLATE_API virtual bool OnMouseDown( const TSharedPtr< FGenericWindow >& PlatformWindow, const EMouseButtons::Type Button, const FVector2D CursorPos ) override;
+	SLATE_API virtual bool OnMouseUp( const EMouseButtons::Type Button ) override;
+	SLATE_API virtual bool OnMouseUp( const EMouseButtons::Type Button, const FVector2D CursorPos ) override;
+	SLATE_API virtual bool OnMouseDoubleClick( const TSharedPtr< FGenericWindow >& PlatformWindow, const EMouseButtons::Type Button ) override;
+	SLATE_API virtual bool OnMouseDoubleClick( const TSharedPtr< FGenericWindow >& PlatformWindow, const EMouseButtons::Type Button, const FVector2D CursorPos ) override;
+	SLATE_API virtual bool OnMouseWheel( const float Delta ) override;
+	SLATE_API virtual bool OnMouseWheel( const float Delta, const FVector2D CursorPos ) override;
+	SLATE_API virtual bool OnMouseMove() override;
+	SLATE_API virtual bool OnRawMouseMove( const int32 X, const int32 Y ) override;
+	SLATE_API virtual bool OnCursorSet() override;
+	SLATE_API virtual bool OnTouchGesture( EGestureEvent GestureType, const FVector2D& Delta, float WheelDelta, bool bIsDirectionInvertedFromDevice ) override;
+	SLATE_API virtual bool OnTouchStarted( const TSharedPtr< FGenericWindow >& Window, const FVector2D& Location, float Force, int32 TouchIndex, FPlatformUserId PlatformUserId, FInputDeviceId DeviceId ) override;
+	SLATE_API virtual bool OnTouchMoved( const FVector2D& Location, float Force, int32 TouchIndex, FPlatformUserId PlatformUserId, FInputDeviceId DeviceID ) override;
+	SLATE_API virtual bool OnTouchEnded( const FVector2D& Location, int32 TouchIndex, FPlatformUserId PlatformUserId, FInputDeviceId DeviceID ) override;
+	SLATE_API virtual bool OnTouchForceChanged(const FVector2D& Location, float Force, int32 TouchIndex, FPlatformUserId PlatformUserId, FInputDeviceId DeviceID) override;
+	SLATE_API virtual bool OnTouchFirstMove(const FVector2D& Location, float Force, int32 TouchIndex, FPlatformUserId PlatformUserId, FInputDeviceId DeviceID) override;
+	SLATE_API virtual void ShouldSimulateGesture(EGestureEvent Gesture, bool bEnable) override;
+	SLATE_API virtual bool OnMotionDetected(const FVector& Tilt, const FVector& RotationRate, const FVector& Gravity, const FVector& Acceleration, FPlatformUserId PlatformUserId, FInputDeviceId InputDeviceId) override;
+	SLATE_API virtual bool OnControllerAnalog(FGamepadKeyNames::Type KeyName, FPlatformUserId PlatformUserId, FInputDeviceId InputDeviceId, float AnalogValue) override;
+	SLATE_API virtual bool OnControllerButtonPressed(FGamepadKeyNames::Type KeyName, FPlatformUserId PlatformUserId, FInputDeviceId InputDeviceId, bool IsRepeat) override;
+	SLATE_API virtual bool OnControllerButtonReleased(FGamepadKeyNames::Type KeyName, FPlatformUserId PlatformUserId, FInputDeviceId InputDeviceId, bool IsRepeat) override;
+	SLATE_API virtual bool OnSizeChanged( const TSharedRef< FGenericWindow >& PlatformWindow, const int32 Width, const int32 Height, bool bWasMinimized = false ) override;
+	SLATE_API virtual void OnOSPaint( const TSharedRef< FGenericWindow >& PlatformWindow ) override;
+	SLATE_API virtual FWindowSizeLimits GetSizeLimitsForWindow(const TSharedRef<FGenericWindow>& Window) const override;
+	SLATE_API virtual void OnResizingWindow( const TSharedRef< FGenericWindow >& PlatformWindow ) override;
+	SLATE_API virtual bool BeginReshapingWindow( const TSharedRef< FGenericWindow >& PlatformWindow ) override;
+	SLATE_API virtual void FinishedReshapingWindow( const TSharedRef< FGenericWindow >& PlatformWindow ) override;
+	SLATE_API virtual void SignalSystemDPIChanged(const TSharedRef<FGenericWindow>& Window) override;
+	SLATE_API virtual void HandleDPIScaleChanged(const TSharedRef<FGenericWindow>& Window) override;
+	SLATE_API virtual void OnMovedWindow( const TSharedRef< FGenericWindow >& PlatformWindow, const int32 X, const int32 Y ) override;
+	SLATE_API virtual bool OnWindowActivationChanged( const TSharedRef< FGenericWindow >& PlatformWindow, const EWindowActivation ActivationType ) override;
+	SLATE_API virtual bool OnApplicationActivationChanged( const bool IsActive ) override;
+	SLATE_API virtual bool OnConvertibleLaptopModeChanged() override;
+	SLATE_API virtual EWindowZone::Type GetWindowZoneForPoint( const TSharedRef< FGenericWindow >& PlatformWindow, const int32 X, const int32 Y ) override;
+	SLATE_API virtual void OnWindowClose( const TSharedRef< FGenericWindow >& PlatformWindow ) override;
+	SLATE_API virtual EDropEffect::Type OnDragEnterText( const TSharedRef< FGenericWindow >& Window, const FString& Text ) override;
+	SLATE_API virtual EDropEffect::Type OnDragEnterFiles( const TSharedRef< FGenericWindow >& Window, const TArray< FString >& Files ) override;
+	SLATE_API virtual EDropEffect::Type OnDragEnterExternal( const TSharedRef< FGenericWindow >& Window, const FString& Text, const TArray< FString >& Files ) override;
+
+	SLATE_API EDropEffect::Type OnDragEnter( const TSharedRef< SWindow >& Window, const TSharedRef<FExternalDragOperation>& DragDropOperation );
+
+	SLATE_API virtual EDropEffect::Type OnDragOver( const TSharedPtr< FGenericWindow >& Window ) override;
+	SLATE_API virtual void OnDragLeave( const TSharedPtr< FGenericWindow >& Window ) override;
+	SLATE_API virtual EDropEffect::Type OnDragDrop( const TSharedPtr< FGenericWindow >& Window ) override;
+	SLATE_API virtual bool OnWindowAction( const TSharedRef< FGenericWindow >& PlatformWindow, const EWindowAction::Type InActionType ) override;
 
 public:
 
@@ -1585,7 +1585,7 @@ public:
 	 * 
 	 * @return The reply returned by the widget that handled the event
 	 */
-	FReply RoutePointerDownEvent(const FWidgetPath& WidgetsUnderPointer, const FPointerEvent& PointerEvent);
+	SLATE_API FReply RoutePointerDownEvent(const FWidgetPath& WidgetsUnderPointer, const FPointerEvent& PointerEvent);
 
 	/**
 	 * Directly routes a pointer up event to the widgets in the specified widget path
@@ -1595,7 +1595,7 @@ public:
 	 *
 	 * @return The reply from the event
 	 */
-	FReply RoutePointerUpEvent(const FWidgetPath& WidgetsUnderPointer, const FPointerEvent& PointerEvent);
+	SLATE_API FReply RoutePointerUpEvent(const FWidgetPath& WidgetsUnderPointer, const FPointerEvent& PointerEvent);
 
 	/**
 	 * Directly routes a pointer move event to the widgets in the specified widget path
@@ -1604,7 +1604,7 @@ public:
 	 * @param PointerEvent		The event data that is is routed to the widget path
 	 * @param bIsSynthetic		Whether or not the move event is synthetic.  Synthetic pointer moves used simulate an event without the pointer actually moving 
 	 */
-	bool RoutePointerMoveEvent( const FWidgetPath& WidgetsUnderPointer, const FPointerEvent& PointerEvent, bool bIsSynthetic );
+	SLATE_API bool RoutePointerMoveEvent( const FWidgetPath& WidgetsUnderPointer, const FPointerEvent& PointerEvent, bool bIsSynthetic );
 
 	/**
 	 * Directly routes a pointer double click event to the widgets in the specified widget path
@@ -1612,7 +1612,7 @@ public:
 	 * @param WidgetsUnderPointer	The path of widgets the event is routed to.
 	 * @param PointerEvent		The event data that is is routed to the widget path
 	 */
-	FReply RoutePointerDoubleClickEvent( const FWidgetPath& WidgetsUnderPointer, const FPointerEvent& PointerEvent );
+	SLATE_API FReply RoutePointerDoubleClickEvent( const FWidgetPath& WidgetsUnderPointer, const FPointerEvent& PointerEvent );
 
 	/**
 	 * Directly routes a pointer mouse wheel or gesture event to the widgets in the specified widget path.
@@ -1621,40 +1621,40 @@ public:
 	 * @param InWheelEvent			The event data that is is routed to the widget path
 	 * @param InGestureEvent		The event data that is is routed to the widget path
 	 */
-	FReply RouteMouseWheelOrGestureEvent(const FWidgetPath& WidgetsUnderPointer, const FPointerEvent& InWheelEvent, const FPointerEvent* InGestureEvent = nullptr);
+	SLATE_API FReply RouteMouseWheelOrGestureEvent(const FWidgetPath& WidgetsUnderPointer, const FPointerEvent& InWheelEvent, const FPointerEvent* InGestureEvent = nullptr);
 
 	/**
 	 * @return int user index that the mouse is mapped to.
 	 */
-	int32 GetUserIndexForMouse() const;
+	SLATE_API int32 GetUserIndexForMouse() const;
 
 	/**
 	 * @return int user index that the keyboard is mapped to.
 	 */
-	int32 GetUserIndexForKeyboard() const;
+	SLATE_API int32 GetUserIndexForKeyboard() const;
 
 	/**
 	 * @return InputDeviceId that the mouse is mapped to
 	 */
-	FInputDeviceId GetInputDeviceIdForMouse() const;
+	SLATE_API FInputDeviceId GetInputDeviceIdForMouse() const;
 
 	/**
 	 * @return InputDeviceId that the keyboard is mapped to
 	 */
-	FInputDeviceId GetInputDeviceIdForKeyboard() const;
+	SLATE_API FInputDeviceId GetInputDeviceIdForKeyboard() const;
 
 	/** @return int user index that this controller is mapped to. */
-	int32 GetUserIndexForController(int32 ControllerId) const;
+	SLATE_API int32 GetUserIndexForController(int32 ControllerId) const;
 
 	/** @return Gets the slate user index that this input device is mapped to */	
-	TOptional<int32> GetUserIndexForInputDevice(FInputDeviceId InputDeviceId) const;
+	SLATE_API TOptional<int32> GetUserIndexForInputDevice(FInputDeviceId InputDeviceId) const;
 	/** @return Gets the slate user index that this platform user id mapped to*/	
-	TOptional<int32> GetUserIndexForPlatformUser(FPlatformUserId PlatformUser) const;
+	SLATE_API TOptional<int32> GetUserIndexForPlatformUser(FPlatformUserId PlatformUser) const;
 	
-	TOptional<int32> GetUserIndexForController(int32 ControllerId, FKey InKey) const;
+	SLATE_API TOptional<int32> GetUserIndexForController(int32 ControllerId, FKey InKey) const;
 
 	/** Establishes the input mapping object used to map input sources to SlateUser indices */
-	void SetInputManager(TSharedRef<ISlateInputManager> InputManager);
+	SLATE_API void SetInputManager(TSharedRef<ISlateInputManager> InputManager);
 
 	/**
 	 * Register for a notification when the window action occurs.
@@ -1663,7 +1663,7 @@ public:
 	 *
 	 * @return Handle to the registered delegate.
 	 */
-	FDelegateHandle RegisterOnWindowActionNotification(const FOnWindowAction& Notification);
+	SLATE_API FDelegateHandle RegisterOnWindowActionNotification(const FOnWindowAction& Notification);
 
 	/** Event type for when Slate is ticking during a modal dialog loop */
 	DECLARE_EVENT_OneParam(FSlateApplication, FOnModalLoopTickEvent, float);
@@ -1680,7 +1680,7 @@ public:
 
 	* @param Handle Handle to the delegate to unregister.
 	*/
-	void UnregisterOnWindowActionNotification(FDelegateHandle Handle);
+	SLATE_API void UnregisterOnWindowActionNotification(FDelegateHandle Handle);
 
 	/**
 	 * Attempts to navigate directly to the given widget
@@ -1689,7 +1689,7 @@ public:
 	 * @param NavigationDestination The navigation destination widget
 	 * @param NavigationSource The source type of the navigation
 	 */
-	void NavigateToWidget(const uint32 UserIndex, const TSharedPtr<SWidget>& NavigationDestination, ENavigationSource NavigationSource = ENavigationSource::FocusedWidget);
+	SLATE_API void NavigateToWidget(const uint32 UserIndex, const TSharedPtr<SWidget>& NavigationDestination, ENavigationSource NavigationSource = ENavigationSource::FocusedWidget);
 
 	/**
 	 * Attempts to navigate directly to the widget currently under the given user's cursor
@@ -1698,66 +1698,66 @@ public:
 	 * @param InNavigationType The navigation type / direction
 	 * @param InWindow The window to do the navigation within
 	 */
-	void NavigateFromWidgetUnderCursor(const uint32 InUserIndex, EUINavigation InNavigationType, TSharedRef<SWindow> InWindow);
+	SLATE_API void NavigateFromWidgetUnderCursor(const uint32 InUserIndex, EUINavigation InNavigationType, TSharedRef<SWindow> InWindow);
 
 	/**
 	* Given an optional widget, try and get the most suitable parent window to use with dialogs (such as file and directory pickers).
 	* This will first try and get the window that owns the widget (if provided), before falling back to using the MainFrame window.
 	*/
-	TSharedPtr<SWindow> FindBestParentWindowForDialogs(const TSharedPtr<SWidget>& InWidget, const ESlateParentWindowSearchMethod InParentWindowSearchMethod = ESlateParentWindowSearchMethod::ActiveWindow);
+	SLATE_API TSharedPtr<SWindow> FindBestParentWindowForDialogs(const TSharedPtr<SWidget>& InWidget, const ESlateParentWindowSearchMethod InParentWindowSearchMethod = ESlateParentWindowSearchMethod::ActiveWindow);
 
 	/**
 	* Given an optional widget, try and get the most suitable parent window handle to use with dialogs (such as file and directory pickers).
 	* This will first try and get the window that owns the widget (if provided), before falling back to using the MainFrame window.
 	*/
-	const void* FindBestParentWindowHandleForDialogs(const TSharedPtr<SWidget>& InWidget, const ESlateParentWindowSearchMethod InParentWindowSearchMethod = ESlateParentWindowSearchMethod::ActiveWindow);
+	SLATE_API const void* FindBestParentWindowHandleForDialogs(const TSharedPtr<SWidget>& InWidget, const ESlateParentWindowSearchMethod InParentWindowSearchMethod = ESlateParentWindowSearchMethod::ActiveWindow);
 
 public:
 #if WITH_EDITORONLY_DATA
 	FDragDropCheckingOverride OnDragDropCheckOverride;
 #endif
 
-	const TSet<FKey>& GetPressedMouseButtons() const;
+	SLATE_API const TSet<FKey>& GetPressedMouseButtons() const;
 
 private:
 
-	TSharedRef< FGenericWindow > MakeWindow( TSharedRef<SWindow> InSlateWindow, const bool bShowImmediately );
+	SLATE_API TSharedRef< FGenericWindow > MakeWindow( TSharedRef<SWindow> InSlateWindow, const bool bShowImmediately );
 
 	/**
 	 * Destroys an SWindow, removing it and all its children from the Slate window list.  Notifies the native window to destroy itself and releases rendering resources
 	 *
 	 * @param DestroyedWindow The window to destroy
 	 */
-	void PrivateDestroyWindow( const TSharedRef<SWindow>& DestroyedWindow );
+	SLATE_API void PrivateDestroyWindow( const TSharedRef<SWindow>& DestroyedWindow );
 
 	/**
 	 * Attempts to navigate to the next widget in the direction specified
 	 *
 	 * @return if a new widget was navigated too
 	 */
-	bool AttemptNavigation(const FWidgetPath& NavigationSource, const FNavigationEvent& NavigationEvent, const FNavigationReply& NavigationReply, const FArrangedWidget& BoundaryWidget);
+	SLATE_API bool AttemptNavigation(const FWidgetPath& NavigationSource, const FNavigationEvent& NavigationEvent, const FNavigationReply& NavigationReply, const FArrangedWidget& BoundaryWidget);
 
 	/**
 	 * Executes a navigate to the specified widget if possible
 	 *
 	 * @return if the widget was navigated too
 	 */
-	bool ExecuteNavigation(const FWidgetPath& NavigationSource, TSharedPtr<SWidget> DestinationWidget, const uint32 UserIndex, bool bAlwaysHandleNavigationAttempt);
+	SLATE_API bool ExecuteNavigation(const FWidgetPath& NavigationSource, TSharedPtr<SWidget> DestinationWidget, const uint32 UserIndex, bool bAlwaysHandleNavigationAttempt);
 
 private:
-	FSlateApplication();
-	void SetLastUserInteractionTime(const double InCurrentTime);	
-	bool SetUserFocus(FSlateUser& User, const FWidgetPath& InFocusPath, const EFocusCause InCause);
+	SLATE_API FSlateApplication();
+	SLATE_API void SetLastUserInteractionTime(const double InCurrentTime);	
+	SLATE_API bool SetUserFocus(FSlateUser& User, const FWidgetPath& InFocusPath, const EFocusCause InCause);
 
 private:
 	/**
 	 * Will be invoked when the size of the geometry of the virtual
 	 * desktop changes (e.g. resolution change or monitors re-arranged)
 	 */
-	void OnVirtualDesktopSizeChanged(const FDisplayMetrics& NewDisplayMetric);
+	SLATE_API void OnVirtualDesktopSizeChanged(const FDisplayMetrics& NewDisplayMetric);
 
 	/** Application singleton */
-	static TSharedPtr< FSlateApplication > CurrentApplication;
+	static SLATE_API TSharedPtr< FSlateApplication > CurrentApplication;
 
 	TSet<FKey> PressedMouseButtons;
 
@@ -2010,7 +2010,7 @@ private:
 	uint64 PlatformMouseMovementEvents = 0;
 
 	/** Constant delta time used on every Slate widget Tick if CVarSlateUseFixedDeltaTime is enabled. */
-	static double FixedDeltaTime;
+	static SLATE_API double FixedDeltaTime;
 	/**
 	 * A helper class to wrap the list of input pre-processors. 
 	 */
