[Slate]

diff --git a/Engine/Source/Runtime/Slate/Public/Framework/Application/SlateUser.h b/Engine/Source/Runtime/Slate/Public/Framework/Application/SlateUser.h
--- a/Engine/Source/Runtime/Slate/Public/Framework/Application/SlateUser.h
+++ b/Engine/Source/Runtime/Slate/Public/Framework/Application/SlateUser.h
@@ -14,10 +14,10 @@ class FMenuStack;
 class FNavigationConfig;
 
 /** Handle to a virtual user of slate. */
-class SLATE_API FSlateVirtualUserHandle
+class FSlateVirtualUserHandle
 {
 public:
-	virtual ~FSlateVirtualUserHandle();
+	SLATE_API virtual ~FSlateVirtualUserHandle();
 
 	FORCEINLINE int32 GetUserIndex() const { return UserIndex; }
 	FORCEINLINE int32 GetVirtualUserIndex() const { return VirtualUserIndex; }
@@ -37,96 +37,96 @@ private:
  * Slate's representation of an individual input-providing user.  
  * As new input sources are connected, new SlateUsers are created.
  */
-class SLATE_API FSlateUser : public TSharedFromThis<FSlateUser>
+class FSlateUser : public TSharedFromThis<FSlateUser>
 {
 public:
-	virtual ~FSlateUser();
+	SLATE_API virtual ~FSlateUser();
 
 	FORCEINLINE int32 GetUserIndex() const { return UserIndex; }
 	FORCEINLINE FPlatformUserId GetPlatformUserId() const { return PlatformUser; }
 	FORCEINLINE bool IsVirtualUser() const { return !Cursor.IsValid(); }
 
-	TSharedPtr<SWidget> GetFocusedWidget() const;
-	bool ShouldShowFocus(TSharedPtr<const SWidget> Widget) const;
-	bool SetFocus(const TSharedRef<SWidget>& WidgetToFocus, EFocusCause ReasonFocusIsChanging = EFocusCause::SetDirectly);
-	void ClearFocus(EFocusCause ReasonFocusIsChanging = EFocusCause::SetDirectly);
+	SLATE_API TSharedPtr<SWidget> GetFocusedWidget() const;
+	SLATE_API bool ShouldShowFocus(TSharedPtr<const SWidget> Widget) const;
+	SLATE_API bool SetFocus(const TSharedRef<SWidget>& WidgetToFocus, EFocusCause ReasonFocusIsChanging = EFocusCause::SetDirectly);
+	SLATE_API void ClearFocus(EFocusCause ReasonFocusIsChanging = EFocusCause::SetDirectly);
 
 	/** Returns the cause for which the provided widget was focused, or nothing if the given widget is not the current focus target. */
-	TOptional<EFocusCause> HasFocus(TSharedPtr<const SWidget> Widget) const;
+	SLATE_API TOptional<EFocusCause> HasFocus(TSharedPtr<const SWidget> Widget) const;
 
 	/** Returns true if the given widget is in the focus path, but is not the focused widget itself. */
-	bool HasFocusedDescendants(TSharedRef<const SWidget> Widget) const;
+	SLATE_API bool HasFocusedDescendants(TSharedRef<const SWidget> Widget) const;
 
 	/** Returns true if the given widget is anywhere in the focus path, including the focused widget itself. */
-	bool IsWidgetInFocusPath(TSharedPtr<const SWidget> Widget) const;
+	SLATE_API bool IsWidgetInFocusPath(TSharedPtr<const SWidget> Widget) const;
 	
-	bool HasAnyCapture() const;
-	bool HasCursorCapture() const;
-	bool HasCapture(uint32 PointerIndex) const;
+	SLATE_API bool HasAnyCapture() const;
+	SLATE_API bool HasCursorCapture() const;
+	SLATE_API bool HasCapture(uint32 PointerIndex) const;
 	
-	bool DoesWidgetHaveAnyCapture(TSharedPtr<const SWidget> Widget) const;
-	bool DoesWidgetHaveCursorCapture(TSharedPtr<const SWidget> Widget) const;
-	bool DoesWidgetHaveCapture(TSharedPtr<const SWidget> Widget, uint32 PointerIndex) const;
+	SLATE_API bool DoesWidgetHaveAnyCapture(TSharedPtr<const SWidget> Widget) const;
+	SLATE_API bool DoesWidgetHaveCursorCapture(TSharedPtr<const SWidget> Widget) const;
+	SLATE_API bool DoesWidgetHaveCapture(TSharedPtr<const SWidget> Widget, uint32 PointerIndex) const;
 	
-	bool SetCursorCaptor(TSharedRef<const SWidget> Widget, const FWidgetPath& EventPath);
-	bool SetPointerCaptor(uint32 PointerIndex, TSharedRef<const SWidget> Widget, const FWidgetPath& EventPath);
+	SLATE_API bool SetCursorCaptor(TSharedRef<const SWidget> Widget, const FWidgetPath& EventPath);
+	SLATE_API bool SetPointerCaptor(uint32 PointerIndex, TSharedRef<const SWidget> Widget, const FWidgetPath& EventPath);
 	
-	void ReleaseAllCapture();
-	void ReleaseCursorCapture();
-	void ReleaseCapture(uint32 PointerIndex);
+	SLATE_API void ReleaseAllCapture();
+	SLATE_API void ReleaseCursorCapture();
+	SLATE_API void ReleaseCapture(uint32 PointerIndex);
 
-	TArray<FWidgetPath> GetCaptorPaths();
-	FWidgetPath GetCursorCaptorPath(FWeakWidgetPath::EInterruptedPathHandling::Type InterruptedPathHandling = FWeakWidgetPath::EInterruptedPathHandling::Truncate, const FPointerEvent* PointerEvent = nullptr);
-	FWidgetPath GetCaptorPath(uint32 PointerIndex, FWeakWidgetPath::EInterruptedPathHandling::Type InterruptedPathHandling = FWeakWidgetPath::EInterruptedPathHandling::Truncate, const FPointerEvent* PointerEvent = nullptr);
+	SLATE_API TArray<FWidgetPath> GetCaptorPaths();
+	SLATE_API FWidgetPath GetCursorCaptorPath(FWeakWidgetPath::EInterruptedPathHandling::Type InterruptedPathHandling = FWeakWidgetPath::EInterruptedPathHandling::Truncate, const FPointerEvent* PointerEvent = nullptr);
+	SLATE_API FWidgetPath GetCaptorPath(uint32 PointerIndex, FWeakWidgetPath::EInterruptedPathHandling::Type InterruptedPathHandling = FWeakWidgetPath::EInterruptedPathHandling::Truncate, const FPointerEvent* PointerEvent = nullptr);
 
-	FWeakWidgetPath GetWeakCursorCapturePath() const;
-	FWeakWidgetPath GetWeakCapturePath(uint32 PointerIndex) const;
+	SLATE_API FWeakWidgetPath GetWeakCursorCapturePath() const;
+	SLATE_API FWeakWidgetPath GetWeakCapturePath(uint32 PointerIndex) const;
 
-	TArray<TSharedRef<SWidget>> GetCaptorWidgets() const;
-	TSharedPtr<SWidget> GetCursorCaptor() const;
-	TSharedPtr<SWidget> GetPointerCaptor(uint32 PointerIndex) const;
+	SLATE_API TArray<TSharedRef<SWidget>> GetCaptorWidgets() const;
+	SLATE_API TSharedPtr<SWidget> GetCursorCaptor() const;
+	SLATE_API TSharedPtr<SWidget> GetPointerCaptor(uint32 PointerIndex) const;
 
-	void SetCursorVisibility(bool bDrawCursor);
+	SLATE_API void SetCursorVisibility(bool bDrawCursor);
 
-	void SetCursorPosition(int32 PosX, int32 PosY);
-	void SetCursorPosition(const UE::Slate::FDeprecateVector2DParameter& NewCursorPos);
-	void SetPointerPosition(uint32 PointerIndex, int32 PosX, int32 PosY);
-	void SetPointerPosition(uint32 PointerIndex, const UE::Slate::FDeprecateVector2DParameter& NewPointerPos);
+	SLATE_API void SetCursorPosition(int32 PosX, int32 PosY);
+	SLATE_API void SetCursorPosition(const UE::Slate::FDeprecateVector2DParameter& NewCursorPos);
+	SLATE_API void SetPointerPosition(uint32 PointerIndex, int32 PosX, int32 PosY);
+	SLATE_API void SetPointerPosition(uint32 PointerIndex, const UE::Slate::FDeprecateVector2DParameter& NewPointerPos);
 	
-	UE::Slate::FDeprecateVector2DResult GetCursorPosition() const;
-	UE::Slate::FDeprecateVector2DResult GetPreviousCursorPosition() const;
+	SLATE_API UE::Slate::FDeprecateVector2DResult GetCursorPosition() const;
+	SLATE_API UE::Slate::FDeprecateVector2DResult GetPreviousCursorPosition() const;
 	
-	UE::Slate::FDeprecateVector2DResult GetPointerPosition(uint32 PointerIndex) const;
-	UE::Slate::FDeprecateVector2DResult GetPreviousPointerPosition(uint32 PointerIndex) const;
+	SLATE_API UE::Slate::FDeprecateVector2DResult GetPointerPosition(uint32 PointerIndex) const;
+	SLATE_API UE::Slate::FDeprecateVector2DResult GetPreviousPointerPosition(uint32 PointerIndex) const;
 
-	bool IsWidgetUnderCursor(TSharedPtr<const SWidget> Widget) const;
-	bool IsWidgetUnderPointer(TSharedPtr<const SWidget> Widget, uint32 PointerIndex) const;
-	bool IsWidgetUnderAnyPointer(TSharedPtr<const SWidget> Widget) const;
+	SLATE_API bool IsWidgetUnderCursor(TSharedPtr<const SWidget> Widget) const;
+	SLATE_API bool IsWidgetUnderPointer(TSharedPtr<const SWidget> Widget, uint32 PointerIndex) const;
+	SLATE_API bool IsWidgetUnderAnyPointer(TSharedPtr<const SWidget> Widget) const;
 
-	bool IsWidgetDirectlyUnderCursor(TSharedPtr<const SWidget> Widget) const;
-	bool IsWidgetDirectlyUnderPointer(TSharedPtr<const SWidget> Widget, uint32 PointerIndex) const;
-	bool IsWidgetDirectlyUnderAnyPointer(TSharedPtr<const SWidget> Widget) const;
+	SLATE_API bool IsWidgetDirectlyUnderCursor(TSharedPtr<const SWidget> Widget) const;
+	SLATE_API bool IsWidgetDirectlyUnderPointer(TSharedPtr<const SWidget> Widget, uint32 PointerIndex) const;
+	SLATE_API bool IsWidgetDirectlyUnderAnyPointer(TSharedPtr<const SWidget> Widget) const;
 
-	FWeakWidgetPath GetLastWidgetsUnderCursor() const;
-	FWeakWidgetPath GetLastWidgetsUnderPointer(uint32 PointerIndex) const;
+	SLATE_API FWeakWidgetPath GetLastWidgetsUnderCursor() const;
+	SLATE_API FWeakWidgetPath GetLastWidgetsUnderPointer(uint32 PointerIndex) const;
 	const TMap<uint32, FWeakWidgetPath>& GetWidgetsUnderPointerLastEventByIndex() const { return WidgetsUnderPointerLastEventByIndex; }
 
 	TSharedPtr<FDragDropOperation> GetDragDropContent() const { return DragDropContent; }
-	bool IsDragDropping() const;
-	bool IsDragDroppingAffected(const FPointerEvent& InPointerEvent) const;
-	void CancelDragDrop();
+	SLATE_API bool IsDragDropping() const;
+	SLATE_API bool IsDragDroppingAffected(const FPointerEvent& InPointerEvent) const;
+	SLATE_API void CancelDragDrop();
 	
-	void ShowTooltip(const TSharedRef<IToolTip>& InTooltip, const UE::Slate::FDeprecateVector2DParameter& InSpawnLocation);
-	void CloseTooltip();
+	SLATE_API void ShowTooltip(const TSharedRef<IToolTip>& InTooltip, const UE::Slate::FDeprecateVector2DParameter& InSpawnLocation);
+	SLATE_API void CloseTooltip();
 
 	const FGestureDetector& GetGestureDetector() const { return GestureDetector; }
 
-	void SetUserNavigationConfig(TSharedPtr<FNavigationConfig> InNavigationConfig);
+	SLATE_API void SetUserNavigationConfig(TSharedPtr<FNavigationConfig> InNavigationConfig);
 	TSharedPtr<FNavigationConfig> GetUserNavigationConfig() const { return UserNavigationConfig; }
 
 SLATE_SCOPE:
 	static TSharedRef<FSlateUser> Create(int32 InUserIndex, TSharedPtr<ICursor> InCursor);
-	static TSharedRef<FSlateUser> Create(FPlatformUserId InPlatformUserId, TSharedPtr<ICursor> InCursor);
+	static SLATE_API TSharedRef<FSlateUser> Create(FPlatformUserId InPlatformUserId, TSharedPtr<ICursor> InCursor);
 	
 	FORCEINLINE bool HasValidFocusPath() const { return WeakFocusPath.IsValid(); }
 	FORCEINLINE const FWeakWidgetPath& GetWeakFocusPath() const { return WeakFocusPath; }
@@ -140,48 +140,48 @@ SLATE_SCOPE:
 		return StrongFocusPath.ToSharedRef();
 	}
 
-	void DrawWindowlessDragDropContent(const TSharedRef<SWindow>& WindowToDraw, FSlateWindowElementList& WindowElementList, int32& MaxLayerId);
-	void DrawCursor(const TSharedRef<SWindow>& WindowToDraw, FSlateWindowElementList& WindowElementList, int32& MaxLayerId);
+	SLATE_API void DrawWindowlessDragDropContent(const TSharedRef<SWindow>& WindowToDraw, FSlateWindowElementList& WindowElementList, int32& MaxLayerId);
+	SLATE_API void DrawCursor(const TSharedRef<SWindow>& WindowToDraw, FSlateWindowElementList& WindowElementList, int32& MaxLayerId);
 
-	void QueueSyntheticCursorMove();
-	bool SynthesizeCursorMoveIfNeeded();
+	SLATE_API void QueueSyntheticCursorMove();
+	SLATE_API bool SynthesizeCursorMoveIfNeeded();
 
 	TSharedPtr<ICursor> GetCursor() const { return Cursor; }
-	void LockCursor(const TSharedRef<SWidget>& Widget);
-	void UnlockCursor();
-	void UpdateCursor();
-	void ProcessCursorReply(const FCursorReply& CursorReply);
+	SLATE_API void LockCursor(const TSharedRef<SWidget>& Widget);
+	SLATE_API void UnlockCursor();
+	SLATE_API void UpdateCursor();
+	SLATE_API void ProcessCursorReply(const FCursorReply& CursorReply);
 	void RequestCursorQuery() { bQueryCursorRequested = true; }
-	void QueryCursor();
+	SLATE_API void QueryCursor();
 	void OverrideCursor(const TSharedPtr<ICursor> InCursor) { Cursor = InCursor; }
 	
-	void SetFocusPath(const FWidgetPath& NewFocusPath, EFocusCause InFocusCause, bool bInShowFocus);
+	SLATE_API void SetFocusPath(const FWidgetPath& NewFocusPath, EFocusCause InFocusCause, bool bInShowFocus);
 	
-	void FinishFrame();
-	void NotifyWindowDestroyed(TSharedRef<SWindow> DestroyedWindow);
+	SLATE_API void FinishFrame();
+	SLATE_API void NotifyWindowDestroyed(TSharedRef<SWindow> DestroyedWindow);
 
-	bool IsTouchPointerActive(int32 TouchPointerIndex) const;
+	SLATE_API bool IsTouchPointerActive(int32 TouchPointerIndex) const;
 
-	void NotifyTouchStarted(const FPointerEvent& TouchEvent);
-	void NotifyPointerMoveBegin(const FPointerEvent& PointerEvent);
-	void NotifyPointerMoveComplete(const FPointerEvent& PointerEvent, const FWidgetPath& WidgetsUnderPointer);
-	void NotifyPointerReleased(const FPointerEvent& PointerEvent, const FWidgetPath& WidgetsUnderCursor, TSharedPtr<FDragDropOperation> DroppedContent, bool bWasHandled);
-	void UpdatePointerPosition(const FPointerEvent& PointerEvent);
+	SLATE_API void NotifyTouchStarted(const FPointerEvent& TouchEvent);
+	SLATE_API void NotifyPointerMoveBegin(const FPointerEvent& PointerEvent);
+	SLATE_API void NotifyPointerMoveComplete(const FPointerEvent& PointerEvent, const FWidgetPath& WidgetsUnderPointer);
+	SLATE_API void NotifyPointerReleased(const FPointerEvent& PointerEvent, const FWidgetPath& WidgetsUnderCursor, TSharedPtr<FDragDropOperation> DroppedContent, bool bWasHandled);
+	SLATE_API void UpdatePointerPosition(const FPointerEvent& PointerEvent);
 
-	void StartDragDetection(const FWidgetPath& PathToWidget, int32 PointerIndex, FKey DragButton, UE::Slate::FDeprecateVector2DParameter StartLocation);
-	FWidgetPath DetectDrag(const FPointerEvent& PointerEvent, float DragTriggerDistance);
-	bool IsDetectingDrag(uint32 PointerIndex) const;
-	void ResetDragDetection();
+	SLATE_API void StartDragDetection(const FWidgetPath& PathToWidget, int32 PointerIndex, FKey DragButton, UE::Slate::FDeprecateVector2DParameter StartLocation);
+	SLATE_API FWidgetPath DetectDrag(const FPointerEvent& PointerEvent, float DragTriggerDistance);
+	SLATE_API bool IsDetectingDrag(uint32 PointerIndex) const;
+	SLATE_API void ResetDragDetection();
 
-	void SetDragDropContent(TSharedRef<FDragDropOperation> InDragDropContent);
-	void ResetDragDropContent();
+	SLATE_API void SetDragDropContent(TSharedRef<FDragDropOperation> InDragDropContent);
+	SLATE_API void ResetDragDropContent();
 
 	int32 GetFocusVersion() const { return FocusVersion; }
 	void IncrementFocusVersion() { FocusVersion++; }
 
-	void UpdateTooltip(const FMenuStack& MenuStack, bool bCanSpawnNewTooltip);
-	void ResetTooltipWindow();
-	bool IsWindowHousingInteractiveTooltip(const TSharedRef<const SWindow>& WindowToTest) const;
+	SLATE_API void UpdateTooltip(const FMenuStack& MenuStack, bool bCanSpawnNewTooltip);
+	SLATE_API void ResetTooltipWindow();
+	SLATE_API bool IsWindowHousingInteractiveTooltip(const TSharedRef<const SWindow>& WindowToTest) const;
 
 	FGestureDetector& GetGestureDetector() { return GestureDetector; }
 
@@ -192,12 +192,12 @@ SLATE_SCOPE:
 #endif
 
 private:
-	FSlateUser(int32 InUserIndex, TSharedPtr<ICursor> InCursor);
+	SLATE_API FSlateUser(int32 InUserIndex, TSharedPtr<ICursor> InCursor);
 	
-	FSlateUser(FPlatformUserId InPlatformUser, TSharedPtr<ICursor> InCursor);
-	void UpdatePointerPosition(uint32 PointerIndex, const FVector2f& Position);
-	void LockCursorInternal(const FWidgetPath& WidgetPath);
-	TSharedRef<SWindow> GetOrCreateTooltipWindow();
+	SLATE_API FSlateUser(FPlatformUserId InPlatformUser, TSharedPtr<ICursor> InCursor);
+	SLATE_API void UpdatePointerPosition(uint32 PointerIndex, const FVector2f& Position);
+	SLATE_API void LockCursorInternal(const FWidgetPath& WidgetPath);
+	SLATE_API TSharedRef<SWindow> GetOrCreateTooltipWindow();
 
 	/** The index the user was assigned. */
 	int32 UserIndex = INDEX_NONE;
@@ -328,4 +328,4 @@ private:
 	FActiveTooltipInfo ActiveTooltipInfo;
 
 	FGestureDetector GestureDetector;
-};
\ No newline at end of file
+};
diff --git a/Engine/Source/Runtime/Slate/Public/Framework/Commands/Contexts/UIContentContext.h b/Engine/Source/Runtime/Slate/Public/Framework/Commands/Contexts/UIContentContext.h
--- a/Engine/Source/Runtime/Slate/Public/Framework/Commands/Contexts/UIContentContext.h
+++ b/Engine/Source/Runtime/Slate/Public/Framework/Commands/Contexts/UIContentContext.h
@@ -11,7 +11,7 @@
 /**
  * Context for generating content from a FUICommandInfo
  */
-struct SLATE_API FUIContentContext : public IUIActionContextBase
+struct FUIContentContext : public IUIActionContextBase
 {
 	virtual ~FUIContentContext() {}
 
@@ -26,4 +26,4 @@ struct SLATE_API FUIContentContext : public IUIActionContextBase
 	//~ End IUIActionContextBase Interface
 
 	inline static const FName ContextName = FName("ContentContext");
-};
\ No newline at end of file
+};
diff --git a/Engine/Source/Runtime/Slate/Public/Framework/Commands/Contexts/UIIdentifierContext.h b/Engine/Source/Runtime/Slate/Public/Framework/Commands/Contexts/UIIdentifierContext.h
--- a/Engine/Source/Runtime/Slate/Public/Framework/Commands/Contexts/UIIdentifierContext.h
+++ b/Engine/Source/Runtime/Slate/Public/Framework/Commands/Contexts/UIIdentifierContext.h
@@ -18,7 +18,7 @@ DECLARE_DELEGATE_RetVal(const FText, FOnGetContextText)
 /**
  * Context for generating identifiers such as icons, labels, and descriptions from a FUICommandInfo
  */
-struct SLATE_API FUIIdentifierContext : public IUIActionContextBase
+struct FUIIdentifierContext : public IUIActionContextBase
 {
 	virtual ~FUIIdentifierContext() {}
 
@@ -39,4 +39,4 @@ struct SLATE_API FUIIdentifierContext : public IUIActionContextBase
 	//~ End IUIActionContextBase Interface
 
 	inline static const FName ContextName = FName("IdentifierContext");
-};
\ No newline at end of file
+};
diff --git a/Engine/Source/Runtime/Slate/Public/Framework/Commands/InputBindingManager.h b/Engine/Source/Runtime/Slate/Public/Framework/Commands/InputBindingManager.h
--- a/Engine/Source/Runtime/Slate/Public/Framework/Commands/InputBindingManager.h
+++ b/Engine/Source/Runtime/Slate/Public/Framework/Commands/InputBindingManager.h
@@ -22,14 +22,14 @@ DECLARE_MULTICAST_DELEGATE_TwoParams(FOnRegisterCommandList, const FName, TShare
 /**
  * Manager responsible for creating and processing input bindings.                  
  */
-class SLATE_API FInputBindingManager
+class FInputBindingManager
 {
 public:
 
 	/**
 	 * @return The instance of this manager                   
 	 */
-	static FInputBindingManager& Get();
+	static SLATE_API FInputBindingManager& Get();
 
 	/**
 	 * Virtual destructor
@@ -42,17 +42,17 @@ public:
 	 * @param OutInputContexts	The generated list of contexts
 	 * @return A list of all known input contexts                   
 	 */
-	void GetKnownInputContexts( TArray< TSharedPtr<FBindingContext> >& OutInputContexts ) const;
+	SLATE_API void GetKnownInputContexts( TArray< TSharedPtr<FBindingContext> >& OutInputContexts ) const;
 
 	/**
 	 * Look up a binding context by name.
 	 */
-	TSharedPtr<FBindingContext> GetContextByName( const FName& InContextName );
+	SLATE_API TSharedPtr<FBindingContext> GetContextByName( const FName& InContextName );
 
 	/**
 	 * Remove the context with this name
 	 */
-	void RemoveContextByName( const FName& InContextName );
+	SLATE_API void RemoveContextByName( const FName& InContextName );
 
 	/**
 	 * Creates an input command from the specified user interface action
@@ -62,7 +62,7 @@ public:
 	 *								Will automatically bind the user interface action to the command specified 
 	 *								by the action so when the command's chord is pressed, the action is executed
 	 */
-	void CreateInputCommand( const TSharedRef<FBindingContext>& InBindingContext, TSharedRef<FUICommandInfo> InUICommandInfo );
+	SLATE_API void CreateInputCommand( const TSharedRef<FBindingContext>& InBindingContext, TSharedRef<FUICommandInfo> InUICommandInfo );
 
 	/**
 	 * Removes an input command, allowing a new one to take its place
@@ -70,7 +70,7 @@ public:
 	 * @param InBindingContext		The context where the command is valid
 	 * @param InUICommandInfo		The user interface action to remove
 	 */
-	void RemoveInputCommand( const TSharedRef<FBindingContext>& InBindingContext, TSharedRef<FUICommandInfo> InUICommandInfo );
+	SLATE_API void RemoveInputCommand( const TSharedRef<FBindingContext>& InBindingContext, TSharedRef<FUICommandInfo> InUICommandInfo );
 
 	/**
 	 * Returns a command info that is has the same active chord as the provided chord and is in the same binding context or parent context
@@ -80,7 +80,7 @@ public:
 	 * @param	bCheckDefault			Whether or not to check the default chord.  Will check the active chord if false
 	 * @return	A pointer to the command info which has the InChord as its active chord or null if one cannot be found
 	 */
-	const TSharedPtr<FUICommandInfo> GetCommandInfoFromInputChord( const FName InBindingContext, const FInputChord& InChord, bool bCheckDefault ) const;
+	SLATE_API const TSharedPtr<FUICommandInfo> GetCommandInfoFromInputChord( const FName InBindingContext, const FInputChord& InChord, bool bCheckDefault ) const;
 
 	/** 
 	 * Finds the command in the provided context which uses the provided input chord
@@ -90,7 +90,7 @@ public:
 	 * @param bCheckDefault		Whether or not to check the default chord of commands instead of active chords
 	 * @param OutActiveChordIndex	The index into the commands active chord array where the passed in chord was matched.  INDEX_NONE if bCheckDefault is true or nothing was found.
 	 */
-	const TSharedPtr<FUICommandInfo> FindCommandInContext( const FName InBindingContext, const FInputChord& InChord, bool bCheckDefault ) const;
+	SLATE_API const TSharedPtr<FUICommandInfo> FindCommandInContext( const FName InBindingContext, const FInputChord& InChord, bool bCheckDefault ) const;
 
 	/** 
 	 * Finds the command in the provided context which has the provided name 
@@ -98,24 +98,24 @@ public:
 	 * @param InBindingContext	The binding context name
 	 * @param CommandName		The name of the command to find
 	 */
-	const TSharedPtr<FUICommandInfo> FindCommandInContext( const FName InBindingContext, const FName CommandName ) const;
+	SLATE_API const TSharedPtr<FUICommandInfo> FindCommandInContext( const FName InBindingContext, const FName CommandName ) const;
 
 	/**
 	 * Called when the active chord is changed on a command
 	 *
 	 * @param CommandInfo	The command that had the active chord changed. 
 	 */
-	virtual void NotifyActiveChordChanged( const FUICommandInfo& CommandInfo, const EMultipleKeyBindingIndex InChordIndex);
+	SLATE_API virtual void NotifyActiveChordChanged( const FUICommandInfo& CommandInfo, const EMultipleKeyBindingIndex InChordIndex);
 	
 	/**
 	 * Saves the user defined chords to a json file
 	 */
-	void SaveInputBindings();
+	SLATE_API void SaveInputBindings();
 
 	/**
 	 * Removes any user defined chords
 	 */
-	void RemoveUserDefinedChords();
+	SLATE_API void RemoveUserDefinedChords();
 
 	/**
 	 * Returns all known command infos for a given binding context
@@ -123,7 +123,7 @@ public:
 	 * @param InBindingContext	The binding context to get command infos from
 	 * @param OutCommandInfos	The list of command infos for the binding context
 	 */
-	void GetCommandInfosFromContext( const FName InBindingContext, TArray< TSharedPtr<FUICommandInfo> >& OutCommandInfos ) const;
+	SLATE_API void GetCommandInfosFromContext( const FName InBindingContext, TArray< TSharedPtr<FUICommandInfo> >& OutCommandInfos ) const;
 
 	/** Registers a delegate to be called when a user-defined chord is edited */
 	FDelegateHandle RegisterUserDefinedChordChanged(const FOnUserDefinedChordChanged::FDelegate& Delegate)
@@ -144,7 +144,7 @@ public:
 	* @param InCommandName		The name of the command to get the chord from
 	* @param ChordIndex		The index of the key binding (in the multiple key bindings array)
 	*/
-	bool GetUserDefinedChord(const FName InBindingContext, const FName InCommandName, const EMultipleKeyBindingIndex InChordIndex, FInputChord& OutUserDefinedChord);
+	SLATE_API bool GetUserDefinedChord(const FName InBindingContext, const FName InCommandName, const EMultipleKeyBindingIndex InChordIndex, FInputChord& OutUserDefinedChord);
 
 	// Type of filter
 	enum class ECommandFilterType
@@ -163,10 +163,10 @@ public:
 	* @param InCommandName		The name of the command
 	* @param FilterType			Type of filter to apply
 	*/
-	void AddCommandFilter(const FName InOwnerName, const FName InBindingContext, const FName InCommandName, const ECommandFilterType FilterType);
+	SLATE_API void AddCommandFilter(const FName InOwnerName, const FName InBindingContext, const FName InCommandName, const ECommandFilterType FilterType);
 
 	/** Removes all filters associated with an owner */
-	void UnregisterCommandFilterOwner(const FName InOwnerName);
+	SLATE_API void UnregisterCommandFilterOwner(const FName InOwnerName);
 
 	/**
 	* Command passes allow/deny list filters
@@ -174,7 +174,7 @@ public:
 	* @param InBindingContext	The context in which the command is active
 	* @param InCommandName		The name of the command
 	*/
-	bool CommandPassesFilter(const FName InBindingContext, const FName InCommandName) const;
+	SLATE_API bool CommandPassesFilter(const FName InBindingContext, const FName InCommandName) const;
 
 	/**
 	 * Exposes the given command list and its context to all OnRegisterCommandList subscribers
@@ -182,7 +182,7 @@ public:
 	 * @param InBindingContext	The context in which the command list is active
 	 * @param CommandList		The command list to expose to subscribers
 	 */
-	bool RegisterCommandList(const ::FName InBindingContext, TSharedRef<FUICommandList> CommandList) const;
+	SLATE_API bool RegisterCommandList(const ::FName InBindingContext, TSharedRef<FUICommandList> CommandList) const;
 
 	/**
 	 * Exposes a new given command list from the given context to all OnRegisterCommandList subscribers
@@ -190,7 +190,7 @@ public:
 	 * @param InBindingContext	The context in which the command list is active
 	 * @return	The new command list exposed to subscribers
 	 */
-	TSharedPtr<FUICommandList> RegisterNewCommandList(const FName InBindingContext) const;
+	SLATE_API TSharedPtr<FUICommandList> RegisterNewCommandList(const FName InBindingContext) const;
 
 	/**
 	 * Exposes the given command list and its context to all OnUnregisterCommandList subscribers
@@ -198,7 +198,7 @@ public:
 	 * @param InBindingContext	The context in which the command list is active
 	 * @param CommandList		The command list to expose to subscribers
 	 */
-	bool UnregisterCommandList(const ::FName InBindingContext, TSharedRef<FUICommandList> CommandList) const;
+	SLATE_API bool UnregisterCommandList(const ::FName InBindingContext, TSharedRef<FUICommandList> CommandList) const;
 
 	/** A delegate for systems to subscribe to any other system exposing their command lists */
 	FOnRegisterCommandList OnRegisterCommandList;
@@ -211,12 +211,12 @@ private:
 	/**
 	 * Hidden default constructor.
 	 */
-	FInputBindingManager();
+	SLATE_API FInputBindingManager();
 
 	/**
 	 *	Checks a binding context for duplicate chords 
 	 */
-	void CheckForDuplicateDefaultChords( const FBindingContext& InBindingContext, TSharedPtr<FUICommandInfo> InCommandInfo ) const;
+	SLATE_API void CheckForDuplicateDefaultChords( const FBindingContext& InBindingContext, TSharedPtr<FUICommandInfo> InCommandInfo ) const;
 
 	/**
 	 * Recursively finds all child contexts of the provided binding context.  
@@ -224,12 +224,12 @@ private:
 	 * @param InBindingContext	The binding context to search
 	 * @param AllChildren		All the children of InBindingContext. InBindingContext is the first element in AllChildren
 	 */
-	void GetAllChildContexts( const FName InBindingContext, TArray<FName>& AllChildren ) const;
+	SLATE_API void GetAllChildContexts( const FName InBindingContext, TArray<FName>& AllChildren ) const;
 
 	/**
 	 * Displays list of commands.
 	 */
-	void PrintAllInputCommands(bool bBoundOnly);
+	SLATE_API void PrintAllInputCommands(bool bBoundOnly);
 
 private:
 
diff --git a/Engine/Source/Runtime/Slate/Public/Framework/Commands/InputChord.h b/Engine/Source/Runtime/Slate/Public/Framework/Commands/InputChord.h
--- a/Engine/Source/Runtime/Slate/Public/Framework/Commands/InputChord.h
+++ b/Engine/Source/Runtime/Slate/Public/Framework/Commands/InputChord.h
@@ -20,7 +20,7 @@ typedef struct FInputChord FInputGesture;
 
 /** An Input Chord is a key and the modifier keys that are to be held with it. */
 USTRUCT(BlueprintType)
-struct SLATE_API FInputChord
+struct FInputChord
 {
 	GENERATED_USTRUCT_BODY()
 
@@ -57,16 +57,9 @@ struct SLATE_API FInputChord
 	};
 
 	/** Returns the relationship between this chord and another. */
-	ERelationshipType GetRelationship(const FInputChord& OtherChord) const;
+	SLATE_API ERelationshipType GetRelationship(const FInputChord& OtherChord) const;
 
-	FInputChord()
-		: Key(EKeys::Invalid)
-		, bShift(false)
-		, bCtrl(false)
-		, bAlt(false)
-		, bCmd(false)
-	{
-	}
+	SLATE_API FInputChord();
 
 	FInputChord(const FKey InKey)
 		: Key(InKey)
@@ -185,21 +178,21 @@ struct SLATE_API FInputChord
 	 *
 	 * @return A localized string.
 	 */
-	FText GetInputText(const bool bLongDisplayName = true) const;
+	SLATE_API FText GetInputText(const bool bLongDisplayName = true) const;
 	
 	/**
 	 * Gets the key represented as a localized string.
 	 *
 	 * @return A localized string.
 	 */
-	FText GetKeyText(const bool bLongDisplayName = true) const;
+	SLATE_API FText GetKeyText(const bool bLongDisplayName = true) const;
 
 	/**
 	 * Gets the localized string for the modifier portion of the coord. 
 	 * 
 	 * @param ModifierAppender. Text to append between each modifier and the key. If this is not specified a '+' character will be used
 	 */
-	FText GetModifierText(TOptional<FText> ModifierAppender = TOptional<FText>()) const;
+	SLATE_API FText GetModifierText(TOptional<FText> ModifierAppender = TOptional<FText>()) const;
 
 	/**
 	 * Checks whether this chord requires an modifier keys to be pressed.
diff --git a/Engine/Source/Runtime/Slate/Public/Framework/Commands/UIAction.h b/Engine/Source/Runtime/Slate/Public/Framework/Commands/UIAction.h
--- a/Engine/Source/Runtime/Slate/Public/Framework/Commands/UIAction.h
+++ b/Engine/Source/Runtime/Slate/Public/Framework/Commands/UIAction.h
@@ -33,7 +33,7 @@ enum class EUIActionRepeatMode
 /**
  * Implements an UI action.
  */
-struct SLATE_API FUIAction
+struct FUIAction
 {
 	/** Holds a delegate that is executed when this action is activated. */
 	FExecuteAction ExecuteAction;
@@ -53,7 +53,7 @@ struct SLATE_API FUIAction
 public:
 
 	/** Default constructor. */
-	FUIAction();
+	SLATE_API FUIAction();
 
 	/**
 	 * Constructor that takes delegates to initialize the action with
@@ -61,7 +61,7 @@ public:
 	 * @param	ExecuteAction		The delegate to call when the action should be executed
 	 * @param	RepeatMode			Can this action can be repeated if the chord used to call it is held down?
 	 */
-	FUIAction(FExecuteAction InitExecuteAction, EUIActionRepeatMode InitRepeatMode = EUIActionRepeatMode::RepeatDisabled);
+	SLATE_API FUIAction(FExecuteAction InitExecuteAction, EUIActionRepeatMode InitRepeatMode = EUIActionRepeatMode::RepeatDisabled);
 
 	/**
 	 * Constructor that takes delegates to initialize the action with
@@ -70,7 +70,7 @@ public:
 	 * @param	CanExecuteAction	The delegate to call to see if the action can be executed
 	 * @param	RepeatMode			Can this action can be repeated if the chord used to call it is held down?
 	 */
-	FUIAction(FExecuteAction InitExecuteAction, FCanExecuteAction InitCanExecuteAction, EUIActionRepeatMode InitRepeatMode = EUIActionRepeatMode::RepeatDisabled);
+	SLATE_API FUIAction(FExecuteAction InitExecuteAction, FCanExecuteAction InitCanExecuteAction, EUIActionRepeatMode InitRepeatMode = EUIActionRepeatMode::RepeatDisabled);
 
 	/**
 	 * Constructor that takes delegates to initialize the action with
@@ -80,7 +80,7 @@ public:
 	 * @param	IsCheckedDelegate	The delegate to call to see if the action should appear checked when visualized
 	 * @param	RepeatMode			Can this action can be repeated if the chord used to call it is held down?
 	 */
-	FUIAction(FExecuteAction InitExecuteAction, FCanExecuteAction InitCanExecuteAction, FIsActionChecked InitIsCheckedDelegate, EUIActionRepeatMode InitRepeatMode = EUIActionRepeatMode::RepeatDisabled);
+	SLATE_API FUIAction(FExecuteAction InitExecuteAction, FCanExecuteAction InitCanExecuteAction, FIsActionChecked InitIsCheckedDelegate, EUIActionRepeatMode InitRepeatMode = EUIActionRepeatMode::RepeatDisabled);
 
 	/**
 	 * Constructor that takes delegates to initialize the action with
@@ -90,7 +90,7 @@ public:
 	 * @param	GetActionCheckState	The delegate to call to see what the check state of the action should be
 	 * @param	RepeatMode			Can this action can be repeated if the chord used to call it is held down?
 	 */
-	FUIAction(FExecuteAction InitExecuteAction, FCanExecuteAction InitCanExecuteAction, FGetActionCheckState InitGetActionCheckStateDelegate, EUIActionRepeatMode InitRepeatMode = EUIActionRepeatMode::RepeatDisabled);
+	SLATE_API FUIAction(FExecuteAction InitExecuteAction, FCanExecuteAction InitCanExecuteAction, FGetActionCheckState InitGetActionCheckStateDelegate, EUIActionRepeatMode InitRepeatMode = EUIActionRepeatMode::RepeatDisabled);
 
 	/**
 	 * Constructor that takes delegates to initialize the action with
@@ -101,7 +101,7 @@ public:
 	 * @param	IsActionVisible		The delegate to call to see if the action should be visible
 	 * @param	RepeatMode			Can this action can be repeated if the chord used to call it is held down?
 	 */
-	FUIAction(FExecuteAction InitExecuteAction, FCanExecuteAction InitCanExecuteAction, FIsActionChecked InitIsCheckedDelegate, FIsActionButtonVisible InitIsActionVisibleDelegate, EUIActionRepeatMode InitRepeatMode = EUIActionRepeatMode::RepeatDisabled);
+	SLATE_API FUIAction(FExecuteAction InitExecuteAction, FCanExecuteAction InitCanExecuteAction, FIsActionChecked InitIsCheckedDelegate, FIsActionButtonVisible InitIsActionVisibleDelegate, EUIActionRepeatMode InitRepeatMode = EUIActionRepeatMode::RepeatDisabled);
 
 	/**
 	 * Constructor that takes delegates to initialize the action with
@@ -112,7 +112,7 @@ public:
 	 * @param	IsActionVisible		The delegate to call to see if the action should be visible
 	 * @param	RepeatMode			Can this action can be repeated if the chord used to call it is held down?
 	 */
-	FUIAction(FExecuteAction InitExecuteAction, FCanExecuteAction InitCanExecuteAction, FGetActionCheckState InitGetActionCheckStateDelegate, FIsActionButtonVisible InitIsActionVisibleDelegate, EUIActionRepeatMode InitRepeatMode = EUIActionRepeatMode::RepeatDisabled);
+	SLATE_API FUIAction(FExecuteAction InitExecuteAction, FCanExecuteAction InitCanExecuteAction, FGetActionCheckState InitGetActionCheckStateDelegate, FIsActionButtonVisible InitIsActionVisibleDelegate, EUIActionRepeatMode InitRepeatMode = EUIActionRepeatMode::RepeatDisabled);
 
 public:
 
@@ -213,7 +213,7 @@ public:
  * 
  * Note: not using UObjects to avoid overhead.
  */
-struct SLATE_API IUIActionContextBase
+struct IUIActionContextBase
 {
 
 public:
@@ -226,14 +226,14 @@ public:
  * 
  * Ex: An FOnGetContent delegate that can be called to generate a widget.
  */
-struct SLATE_API FUIActionContext
+struct FUIActionContext
 {
 public:
 
-	TSharedPtr<IUIActionContextBase> FindContext(const FName InName) const;
+	SLATE_API TSharedPtr<IUIActionContextBase> FindContext(const FName InName) const;
 
-	void AddContext(const TSharedPtr<IUIActionContextBase>& InContext);
+	SLATE_API void AddContext(const TSharedPtr<IUIActionContextBase>& InContext);
 
 private:
 	TArray<TSharedPtr<IUIActionContextBase>> Contexts;
-};
\ No newline at end of file
+};
diff --git a/Engine/Source/Runtime/Slate/Public/Framework/Commands/UICommandDragDropOp.h b/Engine/Source/Runtime/Slate/Public/Framework/Commands/UICommandDragDropOp.h
--- a/Engine/Source/Runtime/Slate/Public/Framework/Commands/UICommandDragDropOp.h
+++ b/Engine/Source/Runtime/Slate/Public/Framework/Commands/UICommandDragDropOp.h
@@ -12,14 +12,14 @@ class SWidget;
 /**
  * A drag drop operation for UI Commands
  */
-class SLATE_API FUICommandDragDropOp
+class FUICommandDragDropOp
 	: public FDragDropOperation
 {
 public:
 
 	DRAG_DROP_OPERATOR_TYPE(FUICommandDragDropOp, FDragDropOperation)
 
-	static TSharedRef<FUICommandDragDropOp> New( FName InItemName, EMultiBlockType InBlockType, bool bInIsDraggingSection, FName InOriginMultiBox, TSharedPtr<SWidget> CustomDectorator, FVector2D DecoratorOffset );
+	static SLATE_API TSharedRef<FUICommandDragDropOp> New( FName InItemName, EMultiBlockType InBlockType, bool bInIsDraggingSection, FName InOriginMultiBox, TSharedPtr<SWidget> CustomDectorator, FVector2D DecoratorOffset );
 
 	FUICommandDragDropOp( FName InItemName, EMultiBlockType InBlockType, bool bInIsDraggingSection, FName InOriginMultiBox, TSharedPtr<SWidget> InCustomDecorator, FVector2D DecoratorOffset )
 		: ItemName( InItemName )
@@ -36,9 +36,9 @@ public:
 	void SetOnDropNotification( FSimpleDelegate InOnDropNotification ) { OnDropNotification = InOnDropNotification; }
 
 	/** FDragDropOperation interface */
-	virtual TSharedPtr<SWidget> GetDefaultDecorator() const override;
-	virtual void OnDragged( const class FDragDropEvent& DragDropEvent ) override;
-	virtual void OnDrop( bool bDropWasHandled, const FPointerEvent& MouseEvent ) override;
+	SLATE_API virtual TSharedPtr<SWidget> GetDefaultDecorator() const override;
+	SLATE_API virtual void OnDragged( const class FDragDropEvent& DragDropEvent ) override;
+	SLATE_API virtual void OnDrop( bool bDropWasHandled, const FPointerEvent& MouseEvent ) override;
 
 public:
 
diff --git a/Engine/Source/Runtime/Slate/Public/Framework/Commands/UICommandInfo.h b/Engine/Source/Runtime/Slate/Public/Framework/Commands/UICommandInfo.h
--- a/Engine/Source/Runtime/Slate/Public/Framework/Commands/UICommandInfo.h
+++ b/Engine/Source/Runtime/Slate/Public/Framework/Commands/UICommandInfo.h
@@ -51,23 +51,23 @@ class FUICommandInfo;
 /**
  *
  */
-class SLATE_API FUICommandInfoDecl
+class FUICommandInfoDecl
 {
 	friend class FBindingContext;
 
 public:
 
-	FUICommandInfoDecl& DefaultChord( const FInputChord& InDefaultChord, const EMultipleKeyBindingIndex InChordIndex = EMultipleKeyBindingIndex::Primary);
-	FUICommandInfoDecl& UserInterfaceType( EUserInterfaceActionType InType );
-	FUICommandInfoDecl& Icon( const FSlateIcon& InIcon );
-	FUICommandInfoDecl& Description( const FText& InDesc );
+	SLATE_API FUICommandInfoDecl& DefaultChord( const FInputChord& InDefaultChord, const EMultipleKeyBindingIndex InChordIndex = EMultipleKeyBindingIndex::Primary);
+	SLATE_API FUICommandInfoDecl& UserInterfaceType( EUserInterfaceActionType InType );
+	SLATE_API FUICommandInfoDecl& Icon( const FSlateIcon& InIcon );
+	SLATE_API FUICommandInfoDecl& Description( const FText& InDesc );
 
-	operator TSharedPtr<FUICommandInfo>() const;
-	operator TSharedRef<FUICommandInfo>() const;
+	SLATE_API operator TSharedPtr<FUICommandInfo>() const;
+	SLATE_API operator TSharedRef<FUICommandInfo>() const;
 
 public:
 
-	FUICommandInfoDecl( const TSharedRef<class FBindingContext>& InContext, const FName InCommandName, const FText& InLabel, const FText& InDesc, const FName InBundle = NAME_None);
+	SLATE_API FUICommandInfoDecl( const TSharedRef<class FBindingContext>& InContext, const FName InCommandName, const FText& InLabel, const FText& InDesc, const FName InBundle = NAME_None);
 
 private:
 
@@ -80,7 +80,7 @@ DECLARE_MULTICAST_DELEGATE_OneParam(FOnBindingContextChanged, const FBindingCont
 /**
  * Represents a context in which input bindings are valid
  */
-class SLATE_API FBindingContext
+class FBindingContext
 	: public TSharedFromThis<FBindingContext>
 {
 public:
@@ -109,7 +109,7 @@ public:
 	/**
 	 * Creates a new command declaration used to populate commands with data
 	 */
-	FUICommandInfoDecl NewCommand( const FName InCommandName, const FText& InCommandLabel, const FText& InCommandDesc );
+	SLATE_API FUICommandInfoDecl NewCommand( const FName InCommandName, const FText& InCommandLabel, const FText& InCommandDesc );
 
 	/**
 	 * @return The name of the context
@@ -138,7 +138,7 @@ public:
 	 * @param Name A unique identifier for the bundle in this context
 	 * @param Desc A localized description of the bundle
 	 */
-	void AddBundle(const FName Name, const FText& Desc);
+	SLATE_API void AddBundle(const FName Name, const FText& Desc);
 
 	/**
 	 * Gets a localized label of a command bundle
@@ -146,7 +146,7 @@ public:
 	 * @param Name The name of the bundle to get a label for
 	 * @return The localized label of the bundle
 	 */
-	const FText& GetBundleLabel(const FName Name);
+	SLATE_API const FText& GetBundleLabel(const FName Name);
 
 	friend uint32 GetTypeHash( const FBindingContext& Context )
 	{
@@ -159,7 +159,7 @@ public:
 	}
 
 	/** A delegate that is called when commands are registered or unregistered with a binding context */
-	static FOnBindingContextChanged CommandsChanged;
+	static SLATE_API FOnBindingContextChanged CommandsChanged;
 
 private:
 
@@ -180,7 +180,7 @@ private:
 };
 
 
-class SLATE_API FUICommandInfo
+class FUICommandInfo
 {
 	friend class FInputBindingManager;
 	friend class FUICommandInfoDecl;
@@ -192,26 +192,14 @@ public:
 	 *
 	 * @param InBindingContext The name of the binding context to use.
 	 */
-	FUICommandInfo( const FName InBindingContext )
-		: BindingContext( InBindingContext )
-		, UserInterfaceType( EUserInterfaceActionType::Button )
-		, bUseLongDisplayName( true )
-	{
-		LLM_SCOPE_BYTAG(UI_Slate);
-
-		ActiveChords.Empty(2);
-		ActiveChords.Add(TSharedRef<FInputChord>(new FInputChord));
-		ActiveChords.Add(TSharedRef<FInputChord>(new FInputChord));
-
-		DefaultChords.Init(FInputChord(EKeys::Invalid, EModifierKey::None), 2);
-	}
+	SLATE_API FUICommandInfo(const FName InBindingContext);
 
 	/**
 	 * Returns the friendly, localized string name of the first valid chord in the key bindings list that is required to perform the command
 	 *
 	 * @return	Localized friendly text for the chord
 	 */
-	const FText GetInputText() const;
+	SLATE_API const FText GetInputText() const;
 
 	/**
 	 * @return	Returns the active chord at the specified index for this command
@@ -244,10 +232,10 @@ public:
 	const FInputChord& GetDefaultChord(const EMultipleKeyBindingIndex InChordIndex) const { return DefaultChords[static_cast<uint8>(InChordIndex)]; }
 
 	/** Utility function to make an FUICommandInfo */
-	static void MakeCommandInfo( const TSharedRef<class FBindingContext>& InContext, TSharedPtr< FUICommandInfo >& OutCommand, const FName InCommandName, const FText& InCommandLabel, const FText& InCommandDesc, const FSlateIcon& InIcon, const EUserInterfaceActionType InUserInterfaceType, const FInputChord& InDefaultChord, const FInputChord& InAlternateDefaultChord = FInputChord(), const FName InBundle = NAME_None);
+	static SLATE_API void MakeCommandInfo( const TSharedRef<class FBindingContext>& InContext, TSharedPtr< FUICommandInfo >& OutCommand, const FName InCommandName, const FText& InCommandLabel, const FText& InCommandDesc, const FSlateIcon& InIcon, const EUserInterfaceActionType InUserInterfaceType, const FInputChord& InDefaultChord, const FInputChord& InAlternateDefaultChord = FInputChord(), const FName InBundle = NAME_None);
 
 	/** Utility function to unregister an FUICommandInfo */
-	static void UnregisterCommandInfo(const TSharedRef<class FBindingContext>& InContext, const TSharedRef<FUICommandInfo>& InCommand);
+	static SLATE_API void UnregisterCommandInfo(const TSharedRef<class FBindingContext>& InContext, const TSharedRef<FUICommandInfo>& InCommand);
 
 	/** @return The display label for this command */
 	const FText& GetLabel() const { return Label; }
@@ -277,17 +265,17 @@ public:
 	void SetUseLongDisplayName(const bool bInUseLongDisplayName) { bUseLongDisplayName = bInUseLongDisplayName; }
 
 	/** Sets the new active chord for this command */
-	void SetActiveChord( const FInputChord& NewChord, const EMultipleKeyBindingIndex InChordIndex );
+	SLATE_API void SetActiveChord( const FInputChord& NewChord, const EMultipleKeyBindingIndex InChordIndex );
 
 	/** Removes the active chord from this command */
-	void RemoveActiveChord(const EMultipleKeyBindingIndex InChordIndex);
+	SLATE_API void RemoveActiveChord(const EMultipleKeyBindingIndex InChordIndex);
 
 	/** 
 	 * Makes a tooltip for this command.
 	 * @param	InText	Optional dynamic text to be displayed in the tooltip.
 	 * @return	The tooltip widget
 	 */
-	TSharedRef<class SToolTip> MakeTooltip( const TAttribute<FText>& InText = TAttribute<FText>() , const TAttribute< EVisibility >& InToolTipVisibility = TAttribute<EVisibility>()) const;
+	SLATE_API TSharedRef<class SToolTip> MakeTooltip( const TAttribute<FText>& InText = TAttribute<FText>() , const TAttribute< EVisibility >& InToolTipVisibility = TAttribute<EVisibility>()) const;
 
 private:
 
diff --git a/Engine/Source/Runtime/Slate/Public/Framework/Commands/UICommandList.h b/Engine/Source/Runtime/Slate/Public/Framework/Commands/UICommandList.h
--- a/Engine/Source/Runtime/Slate/Public/Framework/Commands/UICommandList.h
+++ b/Engine/Source/Runtime/Slate/Public/Framework/Commands/UICommandList.h
@@ -11,7 +11,7 @@
 struct FKeyEvent;
 struct FPointerEvent;
 
-class SLATE_API FUICommandList
+class FUICommandList
 	: public TSharedFromThis<FUICommandList>
 {
 public:
@@ -29,7 +29,7 @@ public:
 	 * @param ExecuteAction		The delegate to call when the command should be executed
 	 * @param RepeatMode		Can this action can be repeated if the chord used to call it is held down?
 	 */
-	void MapAction( const TSharedPtr< const FUICommandInfo > InUICommandInfo, FExecuteAction ExecuteAction, EUIActionRepeatMode RepeatMode = EUIActionRepeatMode::RepeatDisabled );
+	SLATE_API void MapAction( const TSharedPtr< const FUICommandInfo > InUICommandInfo, FExecuteAction ExecuteAction, EUIActionRepeatMode RepeatMode = EUIActionRepeatMode::RepeatDisabled );
 
 	/**
 	 * Maps a command info to a series of delegates that are executed by a multibox or mouse/keyboard input
@@ -39,7 +39,7 @@ public:
 	 * @param CanExecuteAction	The delegate to call to see if the command can be executed
 	 * @param RepeatMode		Can this action can be repeated if the chord used to call it is held down?
 	 */
-	void MapAction( const TSharedPtr< const FUICommandInfo > InUICommandInfo, FExecuteAction ExecuteAction, FCanExecuteAction CanExecuteAction, EUIActionRepeatMode RepeatMode = EUIActionRepeatMode::RepeatDisabled );
+	SLATE_API void MapAction( const TSharedPtr< const FUICommandInfo > InUICommandInfo, FExecuteAction ExecuteAction, FCanExecuteAction CanExecuteAction, EUIActionRepeatMode RepeatMode = EUIActionRepeatMode::RepeatDisabled );
 
 	/**
 	 * Maps a command info to a series of delegates that are executed by a multibox or mouse/keyboard input
@@ -50,7 +50,7 @@ public:
 	 * @param IsCheckedDelegate	The delegate to call to see if the command should appear checked when visualized in a multibox
 	 * @param RepeatMode		Can this action can be repeated if the chord used to call it is held down?
 	 */
-	void MapAction( const TSharedPtr< const FUICommandInfo > InUICommandInfo, FExecuteAction ExecuteAction, FCanExecuteAction CanExecuteAction, FIsActionChecked IsCheckedDelegate, EUIActionRepeatMode RepeatMode = EUIActionRepeatMode::RepeatDisabled );
+	SLATE_API void MapAction( const TSharedPtr< const FUICommandInfo > InUICommandInfo, FExecuteAction ExecuteAction, FCanExecuteAction CanExecuteAction, FIsActionChecked IsCheckedDelegate, EUIActionRepeatMode RepeatMode = EUIActionRepeatMode::RepeatDisabled );
 
 	/**
 	 * Maps a command info to a series of delegates that are executed by a multibox or mouse/keyboard input
@@ -61,7 +61,7 @@ public:
 	 * @param GetActionCheckState	The delegate to call to see what the check state of the action should be when visualized in a multibox
 	 * @param RepeatMode			Can this action can be repeated if the chord used to call it is held down?
 	 */
-	void MapAction( const TSharedPtr< const FUICommandInfo > InUICommandInfo, FExecuteAction ExecuteAction, FCanExecuteAction CanExecuteAction, FGetActionCheckState GetActionCheckState, EUIActionRepeatMode RepeatMode = EUIActionRepeatMode::RepeatDisabled );
+	SLATE_API void MapAction( const TSharedPtr< const FUICommandInfo > InUICommandInfo, FExecuteAction ExecuteAction, FCanExecuteAction CanExecuteAction, FGetActionCheckState GetActionCheckState, EUIActionRepeatMode RepeatMode = EUIActionRepeatMode::RepeatDisabled );
 
 	/**
 	 * Maps a command info to a series of delegates that are executed by a multibox or mouse/keyboard input
@@ -73,7 +73,7 @@ public:
 	 * @param IsVisibleDelegate	The delegate to call to see if the command should appear or be hidden when visualized in a multibox
 	 * @param RepeatMode		Can this action can be repeated if the chord used to call it is held down?
 	 */
-	void MapAction( const TSharedPtr< const FUICommandInfo > InUICommandInfo, FExecuteAction ExecuteAction, FCanExecuteAction CanExecuteAction, FIsActionChecked IsCheckedDelegate, FIsActionButtonVisible IsVisibleDelegate, EUIActionRepeatMode RepeatMode = EUIActionRepeatMode::RepeatDisabled );
+	SLATE_API void MapAction( const TSharedPtr< const FUICommandInfo > InUICommandInfo, FExecuteAction ExecuteAction, FCanExecuteAction CanExecuteAction, FIsActionChecked IsCheckedDelegate, FIsActionButtonVisible IsVisibleDelegate, EUIActionRepeatMode RepeatMode = EUIActionRepeatMode::RepeatDisabled );
 
 	/**
 	 * Maps a command info to a series of delegates that are executed by a multibox or mouse/keyboard input
@@ -85,7 +85,7 @@ public:
 	 * @param IsVisibleDelegate		The delegate to call to see if the command should appear or be hidden when visualized in a multibox
 	 * @param RepeatMode			Can this action can be repeated if the chord used to call it is held down?
 	 */
-	void MapAction( const TSharedPtr< const FUICommandInfo > InUICommandInfo, FExecuteAction ExecuteAction, FCanExecuteAction CanExecuteAction, FGetActionCheckState GetActionCheckState, FIsActionButtonVisible IsVisibleDelegate, EUIActionRepeatMode RepeatMode = EUIActionRepeatMode::RepeatDisabled );
+	SLATE_API void MapAction( const TSharedPtr< const FUICommandInfo > InUICommandInfo, FExecuteAction ExecuteAction, FCanExecuteAction CanExecuteAction, FGetActionCheckState GetActionCheckState, FIsActionButtonVisible IsVisibleDelegate, EUIActionRepeatMode RepeatMode = EUIActionRepeatMode::RepeatDisabled );
 
 	/**
 	 * Maps a command info to a series of delegates that are executed by a multibox or mouse/keyboard input
@@ -93,7 +93,7 @@ public:
 	 * @param InUICommandInfo	The command info to map
 	 * @param InUIAction		Action to map to this command
 	 */
-	virtual void MapAction( const TSharedPtr< const FUICommandInfo > InUICommandInfo, const FUIAction& InUIAction );
+	SLATE_API virtual void MapAction( const TSharedPtr< const FUICommandInfo > InUICommandInfo, const FUIAction& InUIAction );
 
 	/**
 	 * Maps a command info to a series of delegates that are executed by a multibox or mouse/keyboard input
@@ -103,26 +103,26 @@ public:
 	 * @param InUIAction		Action to map to this command
 	 * @param InUIActionContext	Action context to map to this command
 	 */
-	virtual void MapAction(const TSharedPtr< const FUICommandInfo > InUICommandInfo, const FUIAction& InUIAction, const FUIActionContext& InUIActionContext);
+	SLATE_API virtual void MapAction(const TSharedPtr< const FUICommandInfo > InUICommandInfo, const FUIAction& InUIAction, const FUIActionContext& InUIActionContext);
 
 	/**
 	 * Append commands in InCommandsToAppend to this command list.
 	 */
-	void Append( const TSharedRef<FUICommandList>& InCommandsToAppend );
+	SLATE_API void Append( const TSharedRef<FUICommandList>& InCommandsToAppend );
 
 	/**
 	 * Unmaps a command info that was previously mapped
 	 *
 	 * @param InUICommandInfo		The command info to unmap
 	 */
-	void UnmapAction( const TSharedPtr< const FUICommandInfo > InUICommandInfo );
+	SLATE_API void UnmapAction( const TSharedPtr< const FUICommandInfo > InUICommandInfo );
 
 	/**
 	 * Checks if an action is already mapped
 	 * 
 	 * @return True if an action is already mapped or false otherwise
 	 */
-	bool IsActionMapped( const TSharedPtr< const FUICommandInfo > InUICommandInfo ) const;
+	SLATE_API bool IsActionMapped( const TSharedPtr< const FUICommandInfo > InUICommandInfo ) const;
 
 	/**
 	 * Executes the action associated with the provided command info
@@ -130,14 +130,14 @@ public:
 	 *
 	 * @param InUICommandInfo	The command info execute
 	 */
-	virtual bool ExecuteAction( const TSharedRef< const FUICommandInfo > InUICommandInfo ) const;
+	SLATE_API virtual bool ExecuteAction( const TSharedRef< const FUICommandInfo > InUICommandInfo ) const;
 
 	/**
 	 * Calls the CanExecuteAction associated with the provided command info to see if ExecuteAction can be called
 	 *
 	 * @param InUICommandInfo	The command info execute
 	 */
-	bool CanExecuteAction( const TSharedRef< const FUICommandInfo > InUICommandInfo ) const;
+	SLATE_API bool CanExecuteAction( const TSharedRef< const FUICommandInfo > InUICommandInfo ) const;
 
 	/**
 	 * Attempts to execute the action associated with the provided command info
@@ -145,21 +145,21 @@ public:
 	 *
 	 * @param InUICommandInfo	The command info execute
 	 */
-	bool TryExecuteAction( const TSharedRef< const FUICommandInfo > InUICommandInfo ) const;
+	SLATE_API bool TryExecuteAction( const TSharedRef< const FUICommandInfo > InUICommandInfo ) const;
 
 	/**
 	 * Calls the IsVisible delegate associated with the provided command info to see if the command should be visible in a toolbar
 	 *
 	 * @param InUICommandInfo	The command info execute
 	 */
-	EVisibility GetVisibility( const TSharedRef< const FUICommandInfo > InUICommandInfo ) const;
+	SLATE_API EVisibility GetVisibility( const TSharedRef< const FUICommandInfo > InUICommandInfo ) const;
 
 	/**
 	 * Calls the GetCheckState delegate to see if the visualization of this command in a multibox should appear checked
 	 *
 	 * @param InUICommandInfo	The command info execute
 	 */
-	ECheckBoxState GetCheckState( const TSharedRef< const FUICommandInfo > InUICommandInfo ) const;
+	SLATE_API ECheckBoxState GetCheckState( const TSharedRef< const FUICommandInfo > InUICommandInfo ) const;
 
 	/**
 	 * Processes any UI commands which are activated by the specified key event
@@ -167,7 +167,7 @@ public:
 	 * @param InKeyEvent	The key event to check
 	 * @return true if an action was processed
 	 */
-	bool ProcessCommandBindings( const FKeyEvent& InKeyEvent ) const;
+	SLATE_API bool ProcessCommandBindings( const FKeyEvent& InKeyEvent ) const;
 
 	/**
 	 * Processes any UI commands which are activated by the specified mouse event
@@ -175,7 +175,7 @@ public:
 	 * @param InKeyEvent	The mouse event to check
 	 * @return true if an action was processed
 	 */
-	bool ProcessCommandBindings( const FPointerEvent& InMouseEvent ) const;
+	SLATE_API bool ProcessCommandBindings( const FPointerEvent& InMouseEvent ) const;
 
 	/**
 	 * Processes any UI commands which are activated by the specified key, modifier keys state and input event
@@ -185,7 +185,7 @@ public:
 	 * @param bRepeat			True if input is repeating (held)
 	 * @return true if an action was processed
 	 */
-	bool ProcessCommandBindings( const FKey Key, const FModifierKeysState& ModifierKeysState, const bool bRepeat ) const;
+	SLATE_API bool ProcessCommandBindings( const FKey Key, const FModifierKeysState& ModifierKeysState, const bool bRepeat ) const;
 
 	/** Sets the delegate that determines if this UICommandList is capable of producing an action for the supplied command */
 	void SetCanProduceActionForCommand( const FCanProduceActionForCommand& NewCanProduceActionForCommand ) { CanProduceActionForCommand = NewCanProduceActionForCommand; }
@@ -195,14 +195,14 @@ public:
 	  *
 	  * @param Command				The UI command for which you are discovering an action
 	  */
-	const FUIAction* GetActionForCommand(TSharedPtr<const FUICommandInfo> Command) const;
+	SLATE_API const FUIAction* GetActionForCommand(TSharedPtr<const FUICommandInfo> Command) const;
 
 	/** 
 	  * Attempts to find an context for the specified command in the current UICommandList. This is a wrapper for GetContextForCommandRecursively.
 	  *
 	  * @param Command				The UI command for which you are discovering an action
 	  */
-	const FUIActionContext* GetContextForCommand(TSharedPtr<const FUICommandInfo> Command) const;
+	SLATE_API const FUIActionContext* GetContextForCommand(TSharedPtr<const FUICommandInfo> Command) const;
 
 protected:
 
@@ -216,7 +216,7 @@ protected:
 	 * @param bRepeat	True if command is repeating (held)
 	 * @return True if a command was executed, False otherwise
 	 */
-	bool ConditionalProcessCommandBindings( const FKey Key, bool bCtrl, bool bAlt, bool bShift, bool bCmd, bool bRepeat ) const;
+	SLATE_API bool ConditionalProcessCommandBindings( const FKey Key, bool bCtrl, bool bAlt, bool bShift, bool bCmd, bool bRepeat ) const;
 
 	/** 
 	  * Attempts to find an action for the specified command in the current UICommandList. If it is not found, the action for the
@@ -227,7 +227,7 @@ protected:
 	  * @param bIncludeParents		If true, parents of this command list will be searched in the event that the action is not found
 	  * @param OutVisitedLists		The set of visited lists during recursion. This is used to prevent cycles.
 	  */
-	const FUIAction* GetActionForCommandRecursively(const TSharedRef<const FUICommandInfo>& Command, bool bIncludeChildren, bool bIncludeParents, TSet<TSharedRef<const FUICommandList>>& InOutVisitedLists) const;
+	SLATE_API const FUIAction* GetActionForCommandRecursively(const TSharedRef<const FUICommandInfo>& Command, bool bIncludeChildren, bool bIncludeParents, TSet<TSharedRef<const FUICommandList>>& InOutVisitedLists) const;
 
 	/**
 	  * Attempts to find an context for the specified command in the current UICommandList. If it is not found, the context for the
@@ -238,13 +238,13 @@ protected:
 	  * @param bIncludeParents		If true, parents of this command list will be searched in the event that the action is not found
 	  * @param OutVisitedLists		The set of visited lists during recursion. This is used to prevent cycles.
 	  */
-	const FUIActionContext* GetContextForCommandRecursively(const TSharedRef<const FUICommandInfo>& Command, bool bIncludeChildren, bool bIncludeParents, TSet<TSharedRef<const FUICommandList>>& InOutVisitedLists) const;
+	SLATE_API const FUIActionContext* GetContextForCommandRecursively(const TSharedRef<const FUICommandInfo>& Command, bool bIncludeChildren, bool bIncludeParents, TSet<TSharedRef<const FUICommandList>>& InOutVisitedLists) const;
 
 	/** Returns all contexts associated with this list. This is a wrapper for GatherContextsForListRecursively */
-	void GatherContextsForList(TSet<FName>& OutAllContexts) const;
+	SLATE_API void GatherContextsForList(TSet<FName>& OutAllContexts) const;
 
 	/** Returns all contexts associated with this list. */
-	void GatherContextsForListRecursively(TSet<FName>& OutAllContexts, TSet<TSharedRef<const FUICommandList>>& InOutVisitedLists) const;
+	SLATE_API void GatherContextsForListRecursively(TSet<FName>& OutAllContexts, TSet<TSharedRef<const FUICommandList>>& InOutVisitedLists) const;
 
 private:
 
diff --git a/Engine/Source/Runtime/Slate/Public/Framework/Docking/LayoutService.h b/Engine/Source/Runtime/Slate/Public/Framework/Docking/LayoutService.h
--- a/Engine/Source/Runtime/Slate/Public/Framework/Docking/LayoutService.h
+++ b/Engine/Source/Runtime/Slate/Public/Framework/Docking/LayoutService.h
@@ -5,10 +5,10 @@
 #include "CoreMinimal.h"
 #include "Framework/Docking/TabManager.h"
 
-struct SLATE_API FLayoutSaveRestore
+struct FLayoutSaveRestore
 {
 	/** Gets the ini section label for the additional layout configs */
-	static const FString& GetAdditionalLayoutConfigIni();
+	static SLATE_API const FString& GetAdditionalLayoutConfigIni();
 
 	/**
 	 * Write the layout out into a named config file.
@@ -16,7 +16,7 @@ struct SLATE_API FLayoutSaveRestore
 	 * @param InConfigFileName file to be saved to.
 	 * @param InLayoutToSave the layout to save.
 	 */
-	static void SaveToConfig(const FString& InConfigFileName, const TSharedRef<FTabManager::FLayout>& InLayoutToSave );
+	static SLATE_API void SaveToConfig(const FString& InConfigFileName, const TSharedRef<FTabManager::FLayout>& InLayoutToSave );
 
 	/**
 	 * Given a named DefaultLayout, return any saved version of it from the given ini file, otherwise return the default, also default to open tabs based on bool.
@@ -32,9 +32,9 @@ struct SLATE_API FLayoutSaveRestore
 	 *
 	 * @return Loaded layout or the default.
 	 */
-	static TSharedRef<FTabManager::FLayout> LoadFromConfig(const FString& InConfigFileName, const TSharedRef<FTabManager::FLayout>& InDefaultLayout,
+	static SLATE_API TSharedRef<FTabManager::FLayout> LoadFromConfig(const FString& InConfigFileName, const TSharedRef<FTabManager::FLayout>& InDefaultLayout,
 		const EOutputCanBeNullptr InPrimaryAreaOutputCanBeNullptr = EOutputCanBeNullptr::Never);
-	static TSharedRef<FTabManager::FLayout> LoadFromConfig(const FString& InConfigFileName, const TSharedRef<FTabManager::FLayout>& InDefaultLayout,
+	static SLATE_API TSharedRef<FTabManager::FLayout> LoadFromConfig(const FString& InConfigFileName, const TSharedRef<FTabManager::FLayout>& InDefaultLayout,
 		const EOutputCanBeNullptr InPrimaryAreaOutputCanBeNullptr, TArray<FString>& OutRemovedOlderLayoutVersions);
 
 	/**
@@ -46,7 +46,7 @@ struct SLATE_API FLayoutSaveRestore
 	 * @param InSectionName the section name where to save the value.
 	 * @param InSectionValue the value to save.
 	 */
-	static void SaveSectionToConfig(const FString& InConfigFileName, const FString& InSectionName, const FText& InSectionValue);
+	static SLATE_API void SaveSectionToConfig(const FString& InConfigFileName, const FString& InSectionName, const FText& InSectionValue);
 
 	/**
 	 * Read the desired FText value from the desired section of a named config file.
@@ -58,7 +58,7 @@ struct SLATE_API FLayoutSaveRestore
 	 *
 	 * @return Loaded FText associated for that section.
 	 */
-	static FText LoadSectionFromConfig(const FString& InConfigFileName, const FString& InSectionName);
+	static SLATE_API FText LoadSectionFromConfig(const FString& InConfigFileName, const FString& InSectionName);
 
 	/**
 	 * Migrates the layout configuration from one config file to another.
@@ -66,14 +66,14 @@ struct SLATE_API FLayoutSaveRestore
 	 * @param OldConfigFileName The name of the old configuration file.
 	 * @param NewConfigFileName The name of the new configuration file.
 	 */
-	static void MigrateConfig(const FString& OldConfigFileName, const FString& NewConfigFileName);
+	static SLATE_API void MigrateConfig(const FString& OldConfigFileName, const FString& NewConfigFileName);
 
 	/** 
 	 * Duplicate the layout config from one file to another.
 	 * @param SourceConfigFileName The name of the source configuration file.
 	 * @param TargetConfigFileName The name of the target configuration file.
 	 */
-	static bool DuplicateConfig(const FString& SourceConfigFileName, const FString& TargetConfigFileName);
+	static SLATE_API bool DuplicateConfig(const FString& SourceConfigFileName, const FString& TargetConfigFileName);
 
 	/**
 	 * It checks whether a file is a valid layout config file.
@@ -81,7 +81,7 @@ struct SLATE_API FLayoutSaveRestore
 	 * @param bAllowFallback Whether a config should be considered valid if EditorLayoutsSectionName is missing but DefaultEditorLayoutsSectionName is found
 	 * @return Whether the file is a valid layout config file.
 	 */
-	static bool IsValidConfig(const FString& InConfigFileName, bool bAllowFallback = true);
+	static SLATE_API bool IsValidConfig(const FString& InConfigFileName, bool bAllowFallback = true);
 
 private:
 
diff --git a/Engine/Source/Runtime/Slate/Public/Framework/Docking/TabCommands.h b/Engine/Source/Runtime/Slate/Public/Framework/Docking/TabCommands.h
--- a/Engine/Source/Runtime/Slate/Public/Framework/Docking/TabCommands.h
+++ b/Engine/Source/Runtime/Slate/Public/Framework/Docking/TabCommands.h
@@ -6,7 +6,7 @@
 #include "Styling/CoreStyle.h"
 #include "Framework/Commands/Commands.h"
 
-class SLATE_API FTabCommands : public TCommands < FTabCommands >
+class FTabCommands : public TCommands < FTabCommands >
 {
 public:
 
@@ -19,7 +19,7 @@ public:
 	{
 	}
 
-	virtual void RegisterCommands() override;
+	SLATE_API virtual void RegisterCommands() override;
 
 	TSharedPtr<FUICommandInfo> CloseMajorTab;
 	TSharedPtr<FUICommandInfo> CloseMinorTab;
