[Slate]

diff --git a/Engine/Source/Runtime/Slate/Public/Framework/Docking/TabManager.h b/Engine/Source/Runtime/Slate/Public/Framework/Docking/TabManager.h
--- a/Engine/Source/Runtime/Slate/Public/Framework/Docking/TabManager.h
+++ b/Engine/Source/Runtime/Slate/Public/Framework/Docking/TabManager.h
@@ -162,7 +162,7 @@ DECLARE_DELEGATE_RetVal_OneParam(bool, FCanSpawnTab, const FSpawnTabArgs&);
  */
 DECLARE_DELEGATE_RetVal_OneParam( TSharedPtr<SDockTab>, FOnFindTabToReuse, const FTabId& )
 
-struct SLATE_API FMinorTabConfig
+struct FMinorTabConfig
 {
 public:
 	FMinorTabConfig()
@@ -361,7 +361,7 @@ enum class EOutputCanBeNullptr
 };
 
 
-class SLATE_API FTabManager : public TSharedFromThis<FTabManager>
+class FTabManager : public TSharedFromThis<FTabManager>
 {
 	friend class FGlobalTabmanager;
 	public:
@@ -373,7 +373,7 @@ class SLATE_API FTabManager : public TSharedFromThis<FTabManager>
 
 		DECLARE_DELEGATE_OneParam( FOnPersistLayout, const TSharedRef<FLayout>& );
 
-		class SLATE_API FLayoutNode : public TSharedFromThis<FLayoutNode>
+		class FLayoutNode : public TSharedFromThis<FLayoutNode>
 		{
 			friend class FTabManager;
 
@@ -381,11 +381,11 @@ class SLATE_API FTabManager : public TSharedFromThis<FTabManager>
 
 			virtual ~FLayoutNode() { }
 			
-			virtual TSharedPtr<FStack> AsStack();
+			SLATE_API virtual TSharedPtr<FStack> AsStack();
 
-			virtual TSharedPtr<FSplitter> AsSplitter();
+			SLATE_API virtual TSharedPtr<FSplitter> AsSplitter();
 
-			virtual TSharedPtr<FArea> AsArea();
+			SLATE_API virtual TSharedPtr<FArea> AsArea();
 
 			float GetSizeCoefficient() const { return SizeCoefficient; }
 
@@ -432,7 +432,7 @@ class SLATE_API FTabManager : public TSharedFromThis<FTabManager>
 			bool bPinnedInSidebar;
 		};
 
-		class SLATE_API FStack : public FLayoutNode
+		class FStack : public FLayoutNode
 		{
 				friend class FTabManager;
 				friend class FLayout;
@@ -531,7 +531,7 @@ class SLATE_API FTabManager : public TSharedFromThis<FTabManager>
 		};
 
 
-		class SLATE_API FSplitter : public FLayoutNode
+		class FSplitter : public FLayoutNode
 		{
 				friend class FTabManager;
 				friend class FLayoutExtender;
@@ -595,7 +595,7 @@ class SLATE_API FTabManager : public TSharedFromThis<FTabManager>
 		};
 
 
-		class SLATE_API FArea : public FSplitter
+		class FArea : public FSplitter
 		{
 				friend class FTabManager;
 		
@@ -671,13 +671,13 @@ class SLATE_API FTabManager : public TSharedFromThis<FTabManager>
 		};
 
 
-		class SLATE_API FLayout : public TSharedFromThis<FLayout>
+		class FLayout : public TSharedFromThis<FLayout>
 		{
 				friend class FTabManager;
 				
 			public:
 
-				static const TSharedRef<FTabManager::FLayout> NullLayout; /** A dummy layout meant to spawn nothing during (e.g., asset editor) initialization */
+				static SLATE_API const TSharedRef<FTabManager::FLayout> NullLayout; /** A dummy layout meant to spawn nothing during (e.g., asset editor) initialization */
 
 				TSharedRef<FLayout> AddArea( const TSharedRef<FArea>& InArea )
 				{
@@ -691,17 +691,17 @@ class SLATE_API FTabManager : public TSharedFromThis<FTabManager>
 				}
 				
 			public:
-				static TSharedPtr<FTabManager::FLayout> NewFromString( const FString& LayoutAsText );
-				static TSharedPtr<FTabManager::FLayout> NewFromJson( const TSharedPtr<FJsonObject>& LayoutAsJson );
-				FName GetLayoutName() const;
-				TSharedRef<FJsonObject> ToJson() const;
-				FString ToString() const;
+				static SLATE_API TSharedPtr<FTabManager::FLayout> NewFromString( const FString& LayoutAsText );
+				static SLATE_API TSharedPtr<FTabManager::FLayout> NewFromJson( const TSharedPtr<FJsonObject>& LayoutAsJson );
+				SLATE_API FName GetLayoutName() const;
+				SLATE_API TSharedRef<FJsonObject> ToJson() const;
+				SLATE_API FString ToString() const;
 
-				void ProcessExtensions(const FLayoutExtender& Extender);
+				SLATE_API void ProcessExtensions(const FLayoutExtender& Extender);
 
 			protected:
-				static TSharedRef<class FJsonObject> PersistToString_Helper(const TSharedRef<FLayoutNode>& NodeToPersist);
-				static TSharedRef<FLayoutNode> NewFromString_Helper( TSharedPtr<FJsonObject> JsonObject );
+				static SLATE_API TSharedRef<class FJsonObject> PersistToString_Helper(const TSharedRef<FLayoutNode>& NodeToPersist);
+				static SLATE_API TSharedRef<FLayoutNode> NewFromString_Helper( TSharedPtr<FJsonObject> JsonObject );
 
 				FLayout(const FName& InLayoutName)
 					: LayoutName(InLayoutName)
@@ -746,39 +746,39 @@ class SLATE_API FTabManager : public TSharedFromThis<FTabManager>
 				
 		};
 
-		FTabManager::FPrivateApi& GetPrivateApi();
+		SLATE_API FTabManager::FPrivateApi& GetPrivateApi();
 
 	public:
-		class SLATE_API FSearchPreference
+		class FSearchPreference
 		{
 		public:
 			virtual TSharedPtr<SDockTab> Search(const FTabManager& Manager, FName PlaceholderId, const TSharedRef<SDockTab>& UnmanagedTab) const = 0;
 			virtual ~FSearchPreference() {}
 		};
 
-		class SLATE_API FRequireClosedTab : public FSearchPreference
+		class FRequireClosedTab : public FSearchPreference
 		{
 		public:
-			virtual TSharedPtr<SDockTab> Search(const FTabManager& Manager, FName PlaceholderId, const TSharedRef<SDockTab>& UnmanagedTab) const override;
+			SLATE_API virtual TSharedPtr<SDockTab> Search(const FTabManager& Manager, FName PlaceholderId, const TSharedRef<SDockTab>& UnmanagedTab) const override;
 		};
 
-		class SLATE_API FLiveTabSearch : public FSearchPreference
+		class FLiveTabSearch : public FSearchPreference
 		{
 		public:
-			FLiveTabSearch(FName InSearchForTabId = NAME_None);
+			SLATE_API FLiveTabSearch(FName InSearchForTabId = NAME_None);
 
-			virtual TSharedPtr<SDockTab> Search(const FTabManager& Manager, FName PlaceholderId, const TSharedRef<SDockTab>& UnmanagedTab) const override;
+			SLATE_API virtual TSharedPtr<SDockTab> Search(const FTabManager& Manager, FName PlaceholderId, const TSharedRef<SDockTab>& UnmanagedTab) const override;
 
 		private:
 			FName SearchForTabId;
 		};
 
-		class SLATE_API FLastMajorOrNomadTab : public FSearchPreference
+		class FLastMajorOrNomadTab : public FSearchPreference
 		{
 		public:
-			FLastMajorOrNomadTab(FName InFallbackTabId);
+			SLATE_API FLastMajorOrNomadTab(FName InFallbackTabId);
 
-			virtual TSharedPtr<SDockTab> Search(const FTabManager& Manager, FName PlaceholderId, const TSharedRef<SDockTab>& UnmanagedTab) const override;
+			SLATE_API virtual TSharedPtr<SDockTab> Search(const FTabManager& Manager, FName PlaceholderId, const TSharedRef<SDockTab>& UnmanagedTab) const override;
 		private:
 			FName FallbackTabId;
 		};
@@ -826,22 +826,22 @@ class SLATE_API FTabManager : public TSharedFromThis<FTabManager>
 			DefaultTabWindowSizeMap.Remove(TabName);
 		}
 
-		void SetOnPersistLayout( const FOnPersistLayout& InHandler );
+		SLATE_API void SetOnPersistLayout( const FOnPersistLayout& InHandler );
 
 		/** Close all live areas and wipe all the persisted areas. */
-		void CloseAllAreas();
+		SLATE_API void CloseAllAreas();
 
 		/** Gather the persistent layout */
-		TSharedRef<FTabManager::FLayout> PersistLayout() const;
+		SLATE_API TSharedRef<FTabManager::FLayout> PersistLayout() const;
 
 		/** Gather the persistent layout and execute the custom delegate for saving it to persistent storage (e.g. into config files) */
-		void SavePersistentLayout();
+		SLATE_API void SavePersistentLayout();
 
 		/** Request a deferred save of the layout. */
-		void RequestSavePersistentLayout();
+		SLATE_API void RequestSavePersistentLayout();
 
 		/** Clears the request for a deferred save of the layout */
-		void ClearPendingLayoutSave();
+		SLATE_API void ClearPendingLayoutSave();
 
 		/**
 		 * Register a new tab spawner with the tab manager.  The spawner will be called when anyone calls
@@ -851,32 +851,32 @@ class SLATE_API FTabManager : public TSharedFromThis<FTabManager>
 		 * @param CanSpawnTab The callback that will be used to ask if spawning the tab is allowed
 		 * @return The registration entry for the spawner.
 		 */
-		FTabSpawnerEntry& RegisterTabSpawner(const FName TabId, const FOnSpawnTab& OnSpawnTab, const FCanSpawnTab& CanSpawnTab = FCanSpawnTab());
+		SLATE_API FTabSpawnerEntry& RegisterTabSpawner(const FName TabId, const FOnSpawnTab& OnSpawnTab, const FCanSpawnTab& CanSpawnTab = FCanSpawnTab());
 
 		/**
 		 * Unregisters the tab spawner matching the provided TabId.
 		 * @param TabId The TabId to remove the spawner for.
 		 * @return true if a spawner was found for this TabId, otherwise false.
 		 */
-		bool UnregisterTabSpawner( const FName TabId );
+		SLATE_API bool UnregisterTabSpawner( const FName TabId );
 
 		/**
 		 * Unregisters all tab spawners.
 		 */
-		void UnregisterAllTabSpawners();
+		SLATE_API void UnregisterAllTabSpawners();
 
-		TSharedPtr<SWidget> RestoreFrom(const TSharedRef<FLayout>& Layout, const TSharedPtr<SWindow>& ParentWindow, const bool bEmbedTitleAreaContent = false,
+		SLATE_API TSharedPtr<SWidget> RestoreFrom(const TSharedRef<FLayout>& Layout, const TSharedPtr<SWindow>& ParentWindow, const bool bEmbedTitleAreaContent = false,
 			const EOutputCanBeNullptr RestoreAreaOutputCanBeNullptr = EOutputCanBeNullptr::Never);
 
-		void PopulateLocalTabSpawnerMenu( FMenuBuilder& PopulateMe );
+		SLATE_API void PopulateLocalTabSpawnerMenu( FMenuBuilder& PopulateMe );
 
-		void PopulateTabSpawnerMenu(FMenuBuilder& PopulateMe, TSharedRef<FWorkspaceItem> MenuStructure);
+		SLATE_API void PopulateTabSpawnerMenu(FMenuBuilder& PopulateMe, TSharedRef<FWorkspaceItem> MenuStructure);
 
-		void PopulateTabSpawnerMenu( FMenuBuilder& PopulateMe, TSharedRef<FWorkspaceItem> MenuStructure, bool bIncludeOrphanedMenus);
+		SLATE_API void PopulateTabSpawnerMenu( FMenuBuilder& PopulateMe, TSharedRef<FWorkspaceItem> MenuStructure, bool bIncludeOrphanedMenus);
 
-		void PopulateTabSpawnerMenu( FMenuBuilder &PopulateMe, const FName& TabType );
+		SLATE_API void PopulateTabSpawnerMenu( FMenuBuilder &PopulateMe, const FName& TabType );
 
-		virtual void DrawAttention( const TSharedRef<SDockTab>& TabToHighlight );
+		SLATE_API virtual void DrawAttention( const TSharedRef<SDockTab>& TabToHighlight );
 
 		struct ESearchPreference
 		{
@@ -888,19 +888,19 @@ class SLATE_API FTabManager : public TSharedFromThis<FTabManager>
 		};
 
 		/** Insert a new UnmanagedTab document tab next to an existing tab (closed or open) that has the PlaceholdId. Give the New tab NewTabId */
-		void InsertNewDocumentTab( FName PlaceholderId, FName NewTabId, const FSearchPreference& SearchPreference, const TSharedRef<SDockTab>& UnmanagedTab );
+		SLATE_API void InsertNewDocumentTab( FName PlaceholderId, FName NewTabId, const FSearchPreference& SearchPreference, const TSharedRef<SDockTab>& UnmanagedTab );
 
 		/** Insert a new UnmanagedTab document tab next to an existing tab (closed or open) that has the PlaceholdId. */
-		void InsertNewDocumentTab( FName PlaceholderId, const FSearchPreference& SearchPreference, const TSharedRef<SDockTab>& UnmanagedTab );
+		SLATE_API void InsertNewDocumentTab( FName PlaceholderId, const FSearchPreference& SearchPreference, const TSharedRef<SDockTab>& UnmanagedTab );
 		
 		/** Insert a new UnmanagedTab document tab next to an existing tab (closed or open) that has the PlaceholdId. */
-		void InsertNewDocumentTab(FName PlaceholderId, ESearchPreference::Type SearchPreference, const TSharedRef<SDockTab>& UnmanagedTab);
+		SLATE_API void InsertNewDocumentTab(FName PlaceholderId, ESearchPreference::Type SearchPreference, const TSharedRef<SDockTab>& UnmanagedTab);
 
 		/**
 		 * Much like InsertNewDocumentTab, but the UnmanagedTab is not seen by the user as newly-created.
 		 * e.g. Opening an restores multiple previously opened documents; these are not seen as new tabs.
 		 */
-		void RestoreDocumentTab(FName PlaceholderId, ESearchPreference::Type SearchPreference, const TSharedRef<SDockTab>& UnmanagedTab);
+		SLATE_API void RestoreDocumentTab(FName PlaceholderId, ESearchPreference::Type SearchPreference, const TSharedRef<SDockTab>& UnmanagedTab);
 
 		/**
 		 * Try to open tab if it is closed at the last known location.  If it already exists, it will draw attention to the tab.
@@ -909,7 +909,7 @@ class SLATE_API FTabManager : public TSharedFromThis<FTabManager>
 		 * @param bInvokeAsInactive	Leave the tab inactive instead of drawing attention to it
 		 * @return The existing or newly spawned tab instance if successful.
 		 */
-		virtual TSharedPtr<SDockTab> TryInvokeTab(const FTabId& TabId, bool bInvokeAsInactive = false);
+		SLATE_API virtual TSharedPtr<SDockTab> TryInvokeTab(const FTabId& TabId, bool bInvokeAsInactive = false);
 
 		/**
 		 * Finds the first instance of an existing tab with the given tab id.
@@ -917,12 +917,12 @@ class SLATE_API FTabManager : public TSharedFromThis<FTabManager>
 		 * @param TabId The tab identifier.
 		 * @return The existing tab instance if found, otherwise null.
 		 */
-		TSharedPtr<SDockTab> FindExistingLiveTab(const FTabId& TabId) const;
+		SLATE_API TSharedPtr<SDockTab> FindExistingLiveTab(const FTabId& TabId) const;
 
-		virtual ~FTabManager();
+		SLATE_API virtual ~FTabManager();
 
 		/** Sets whether or not this tab manager supports a custom menu bar for the active major tab that will be shown on top of the major tab area in the window this tab manager resides in. */
-		void SetAllowWindowMenuBar(bool bInAllowWindowMenuBar);
+		SLATE_API void SetAllowWindowMenuBar(bool bInAllowWindowMenuBar);
 
 		/** Whether or not this tab manager supports a custom menu bar for the active major tab that will be shown on top of the major tab area in the window this tab manager resides in. */
 		bool AllowsWindowMenuBar() const { return bAllowPerWindowMenu; }
@@ -932,27 +932,27 @@ class SLATE_API FTabManager : public TSharedFromThis<FTabManager>
 		 * On Mac the menu bar appears globally at the top of the desktop in all cases regardless of whether or not SetAllowWindowMenuBar is called.  On other desktop platforms the menu appears at the top of the window this tab manager is a part of only if SetAllowWindowMenuBar(true) is called.
 		 * @param NewMenuMutliBox The multi-box to generate the global menu bar from.
 		 */
-		void SetMenuMultiBox(const TSharedPtr<FMultiBox> NewMenuMutliBox, const TSharedPtr<SWidget> MenuWidget);
+		SLATE_API void SetMenuMultiBox(const TSharedPtr<FMultiBox> NewMenuMutliBox, const TSharedPtr<SWidget> MenuWidget);
 
 		/**
 		 * Update the native, global menu bar if it is being used.
 		 * @param bForce Used to force an update even if the parent window doesn't contain the widget with keyboard focus.
 		 */
-		void UpdateMainMenu(TSharedPtr<SDockTab> ForTab, const bool bForce);
+		SLATE_API void UpdateMainMenu(TSharedPtr<SDockTab> ForTab, const bool bForce);
 
 		/** Provide a tab that will be the main tab and cannot be closed.
 		 *  Prefer using the version that takes in a TabID so you can set the MainTab before it is actually created
 		 */
-		void SetMainTab(const TSharedRef<const SDockTab>& InTab);
+		SLATE_API void SetMainTab(const TSharedRef<const SDockTab>& InTab);
 
 		/** Provide a tab that will be the main tab and cannot be closed. */
-		void SetMainTab(const FTabId& InMainTabID);
+		SLATE_API void SetMainTab(const FTabId& InMainTabID);
 
 		/** If true, this tab manager will always have one main tab that cannot be closed
 		 *  Automatically set to the first tab registered if there are no tabs registered already (set on creation)
 		 *  or the first tab found otherwise
 		 */
-		void SetEnforceMainTab(bool bInEnforceMainTab);
+		SLATE_API void SetEnforceMainTab(bool bInEnforceMainTab);
 
 		/* Prevent or allow all tabs to be drag */
 		void SetCanDoDragOperation(bool CanDoDragOperation) { bCanDoDragOperation = CanDoDragOperation; }
@@ -961,68 +961,68 @@ class SLATE_API FTabManager : public TSharedFromThis<FTabManager>
 		bool GetCanDoDragOperation() { return bCanDoDragOperation; }
 
 		/** @return if the provided tab can be closed. */
-		bool IsTabCloseable(const TSharedRef<const SDockTab>& InTab) const;
+		SLATE_API bool IsTabCloseable(const TSharedRef<const SDockTab>& InTab) const;
 
 		/** @return true if a tab is ever allowed in a sidebar */
-		bool IsTabAllowedInSidebar(const FTabId TabId) const;
+		SLATE_API bool IsTabAllowedInSidebar(const FTabId TabId) const;
 
 		/**
 		 * Temporarily moves all open tabs in this tab manager to a sidebar or restores them from a temporary state
 		 */
-		void ToggleSidebarOpenTabs();
+		SLATE_API void ToggleSidebarOpenTabs();
 
 		/** @return The local workspace menu root */
-		const TSharedRef<FWorkspaceItem> GetLocalWorkspaceMenuRoot() const;
+		SLATE_API const TSharedRef<FWorkspaceItem> GetLocalWorkspaceMenuRoot() const;
 
 		/** Adds a category to the local workspace menu by name */
-		TSharedRef<FWorkspaceItem> AddLocalWorkspaceMenuCategory( const FText& CategoryTitle );
+		SLATE_API TSharedRef<FWorkspaceItem> AddLocalWorkspaceMenuCategory( const FText& CategoryTitle );
 
 		/** Adds an existing workspace item to the local workspace menu */
-		void AddLocalWorkspaceMenuItem( const TSharedRef<FWorkspaceItem>& CategoryItem );
+		SLATE_API void AddLocalWorkspaceMenuItem( const TSharedRef<FWorkspaceItem>& CategoryItem );
 
 		/** Clears all categories in the local workspace menu */
-		void ClearLocalWorkspaceMenuCategories();
+		SLATE_API void ClearLocalWorkspaceMenuCategories();
 
 		/** @return true if the tab has a factory registered for it that allows it to be spawned. */
-		bool HasTabSpawner(FName TabId) const;
+		SLATE_API bool HasTabSpawner(FName TabId) const;
 
 		/** Returns the owner tab (if it exists) */
 		TSharedPtr<SDockTab> GetOwnerTab() { return OwnerTabPtr.Pin(); }
 
 		/** Returns filter for additional control over available tabs */
-		TSharedRef<FNamePermissionList>& GetTabPermissionList();
+		SLATE_API TSharedRef<FNamePermissionList>& GetTabPermissionList();
 
-		FUIAction GetUIActionForTabSpawnerMenuEntry(TSharedPtr<FTabSpawnerEntry> InTabMenuEntry);
+		SLATE_API FUIAction GetUIActionForTabSpawnerMenuEntry(TSharedPtr<FTabSpawnerEntry> InTabMenuEntry);
 
 		/** Return an array of all tab spawners that should have menu items in the main menu. */
-		TArray< TWeakPtr<FTabSpawnerEntry> > CollectSpawners();
+		SLATE_API TArray< TWeakPtr<FTabSpawnerEntry> > CollectSpawners();
 
 		/** Find a tab spawner by name. */
-		TSharedPtr<FTabSpawnerEntry> FindTabSpawnerFor(FName TabId);
+		SLATE_API TSharedPtr<FTabSpawnerEntry> FindTabSpawnerFor(FName TabId);
 
 		/** Find a tab spawner by name. */
-		const TSharedPtr<const FTabSpawnerEntry> FindTabSpawnerFor(FName TabId) const;
+		SLATE_API const TSharedPtr<const FTabSpawnerEntry> FindTabSpawnerFor(FName TabId) const;
 
 	protected:
-		void InvokeTabForMenu( FName TabId );
+		SLATE_API void InvokeTabForMenu( FName TabId );
 
 	protected:
 		
-		void InsertDocumentTab( FName PlaceholderId, const FSearchPreference& SearchPreference, const TSharedRef<SDockTab>& UnmanagedTab, bool bPlaySpawnAnim );
-		void InsertDocumentTab( FName PlaceholderId, FName NewTabId, const FSearchPreference& SearchPreference, const TSharedRef<SDockTab>& UnmanagedTab, bool bPlaySpawnAnim );
+		SLATE_API void InsertDocumentTab( FName PlaceholderId, const FSearchPreference& SearchPreference, const TSharedRef<SDockTab>& UnmanagedTab, bool bPlaySpawnAnim );
+		SLATE_API void InsertDocumentTab( FName PlaceholderId, FName NewTabId, const FSearchPreference& SearchPreference, const TSharedRef<SDockTab>& UnmanagedTab, bool bPlaySpawnAnim );
 
-		virtual void OpenUnmanagedTab(FName PlaceholderId, const FSearchPreference& SearchPreference, const TSharedRef<SDockTab>& UnmanagedTab);
+		SLATE_API virtual void OpenUnmanagedTab(FName PlaceholderId, const FSearchPreference& SearchPreference, const TSharedRef<SDockTab>& UnmanagedTab);
 			
-		void PopulateTabSpawnerMenu_Helper( FMenuBuilder& PopulateMe, struct FPopulateTabSpawnerMenu_Args Args );
+		SLATE_API void PopulateTabSpawnerMenu_Helper( FMenuBuilder& PopulateMe, struct FPopulateTabSpawnerMenu_Args Args );
 
-		void MakeSpawnerMenuEntry( FMenuBuilder &PopulateMe, const TSharedPtr<FTabSpawnerEntry> &InSpawnerNode );
+		SLATE_API void MakeSpawnerMenuEntry( FMenuBuilder &PopulateMe, const TSharedPtr<FTabSpawnerEntry> &InSpawnerNode );
 
-		TSharedPtr<SDockTab> InvokeTab_Internal(const FTabId& TabId, bool bInvokeAsInactive = false);
+		SLATE_API TSharedPtr<SDockTab> InvokeTab_Internal(const FTabId& TabId, bool bInvokeAsInactive = false);
 
 		/** Finds the last major or nomad tab in a particular window. */
-		TSharedPtr<SDockTab> FindLastTabInWindow(TSharedPtr<SWindow> Window) const;
+		SLATE_API TSharedPtr<SDockTab> FindLastTabInWindow(TSharedPtr<SWindow> Window) const;
 
-		TSharedPtr<SDockingTabStack> FindPotentiallyClosedTab( const FTabId& ClosedTabId );
+		SLATE_API TSharedPtr<SDockingTabStack> FindPotentiallyClosedTab( const FTabId& ClosedTabId );
 
 		typedef TMap< FName, TSharedRef<FTabSpawnerEntry> > FTabSpawner;
 		
@@ -1031,12 +1031,12 @@ class SLATE_API FTabManager : public TSharedFromThis<FTabManager>
 			return MakeShareable( new FTabManager(InOwnerTab, InNomadTabSpawner) );
 		}
 
-		FTabManager( const TSharedPtr<SDockTab>& InOwnerTab, const TSharedRef<FTabManager::FTabSpawner> & InNomadTabSpawner );
+		SLATE_API FTabManager( const TSharedPtr<SDockTab>& InOwnerTab, const TSharedRef<FTabManager::FTabSpawner> & InNomadTabSpawner );
 
-		TSharedPtr<SDockingArea> RestoreArea(
+		SLATE_API TSharedPtr<SDockingArea> RestoreArea(
 			const TSharedRef<FArea>& AreaToRestore, const TSharedPtr<SWindow>& InParentWindow, const bool bEmbedTitleAreaContent = false, const EOutputCanBeNullptr OutputCanBeNullptr = EOutputCanBeNullptr::Never);
 
-		TSharedPtr<class SDockingNode> RestoreArea_Helper(const TSharedRef<FLayoutNode>& LayoutNode, const TSharedPtr<SWindow>& ParentWindow, const bool bEmbedTitleAreaContent, FSidebarTabLists& OutSidebarTabs, const EOutputCanBeNullptr OutputCanBeNullptr = EOutputCanBeNullptr::Never);
+		SLATE_API TSharedPtr<class SDockingNode> RestoreArea_Helper(const TSharedRef<FLayoutNode>& LayoutNode, const TSharedPtr<SWindow>& ParentWindow, const bool bEmbedTitleAreaContent, FSidebarTabLists& OutSidebarTabs, const EOutputCanBeNullptr OutputCanBeNullptr = EOutputCanBeNullptr::Never);
 
 		/**
 		 * Use CanRestoreSplitterContent + RestoreSplitterContent when the output of its internal RestoreArea_Helper can be a nullptr.
@@ -1050,22 +1050,22 @@ class SLATE_API FTabManager : public TSharedFromThis<FTabManager>
 		 *			RestoreSplitterContent(DockingNodes, SplitterWidget);
 		 *		}
 		 */
-		bool CanRestoreSplitterContent(TArray<TSharedRef<class SDockingNode>>& DockingNodes, const TSharedRef<FSplitter>& SplitterNode, const TSharedPtr<SWindow>& ParentWindow, FSidebarTabLists& OutSidebarTabs, const EOutputCanBeNullptr OutputCanBeNullptr);
-		void RestoreSplitterContent(const TArray<TSharedRef<class SDockingNode>>& DockingNodes, const TSharedRef<class SDockingSplitter>& SplitterWidget);
+		SLATE_API bool CanRestoreSplitterContent(TArray<TSharedRef<class SDockingNode>>& DockingNodes, const TSharedRef<FSplitter>& SplitterNode, const TSharedPtr<SWindow>& ParentWindow, FSidebarTabLists& OutSidebarTabs, const EOutputCanBeNullptr OutputCanBeNullptr);
+		SLATE_API void RestoreSplitterContent(const TArray<TSharedRef<class SDockingNode>>& DockingNodes, const TSharedRef<class SDockingSplitter>& SplitterWidget);
 
 		/**
 		 * Use this standalone RestoreSplitterContent when the output of its internal RestoreArea_Helper cannot be a nullptr.
 		 */
-		void RestoreSplitterContent(const TSharedRef<FSplitter>& SplitterNode, const TSharedRef<class SDockingSplitter>& SplitterWidget, const TSharedPtr<SWindow>& ParentWindow, FSidebarTabLists& OutSidebarTabs);
+		SLATE_API void RestoreSplitterContent(const TSharedRef<FSplitter>& SplitterNode, const TSharedRef<class SDockingSplitter>& SplitterWidget, const TSharedPtr<SWindow>& ParentWindow, FSidebarTabLists& OutSidebarTabs);
 		
-		bool IsValidTabForSpawning( const FTab& SomeTab ) const;
-		bool IsAllowedTab(const FTabId& TabId) const;
-		bool IsAllowedTabType(const FName TabType) const;
+		SLATE_API bool IsValidTabForSpawning( const FTab& SomeTab ) const;
+		SLATE_API bool IsAllowedTab(const FTabId& TabId) const;
+		SLATE_API bool IsAllowedTabType(const FName TabType) const;
 
-		TSharedPtr<SDockTab> SpawnTab(const FTabId& TabId, const TSharedPtr<SWindow>& ParentWindow, const bool bCanOutputBeNullptr = false);
+		SLATE_API TSharedPtr<SDockTab> SpawnTab(const FTabId& TabId, const TSharedPtr<SWindow>& ParentWindow, const bool bCanOutputBeNullptr = false);
 
-		TSharedPtr<class SDockingTabStack> FindTabInLiveAreas( const FTabMatcher& TabMatcher ) const;
-		static TSharedPtr<class SDockingTabStack> FindTabInLiveArea( const FTabMatcher& TabMatcher, const TSharedRef<SDockingArea>& InArea );
+		SLATE_API TSharedPtr<class SDockingTabStack> FindTabInLiveAreas( const FTabMatcher& TabMatcher ) const;
+		static SLATE_API TSharedPtr<class SDockingTabStack> FindTabInLiveArea( const FTabMatcher& TabMatcher, const TSharedRef<SDockingArea>& InArea );
 
 		template<typename MatchFunctorType> static bool HasAnyMatchingTabs( const TSharedRef<FTabManager::FLayoutNode>& SomeNode, const MatchFunctorType& Matcher );
 
@@ -1074,10 +1074,10 @@ class SLATE_API FTabManager : public TSharedFromThis<FTabManager>
 		 * It searches for valid and open tabs on SomeNode.
 		 * @return It returns true if there is at least a valid open tab in the input SomeNode.
 		 */
-		bool HasValidOpenTabs( const TSharedRef<FTabManager::FLayoutNode>& SomeNode ) const;
+		SLATE_API bool HasValidOpenTabs( const TSharedRef<FTabManager::FLayoutNode>& SomeNode ) const;
 
 	protected:
-		bool HasValidTabs( const TSharedRef<FTabManager::FLayoutNode>& SomeNode ) const;
+		SLATE_API bool HasValidTabs( const TSharedRef<FTabManager::FLayoutNode>& SomeNode ) const;
 
 		/**
 		 * It sets the desired (or all) tabs in the FTabManager::FLayoutNode to the desired value.
@@ -1085,42 +1085,42 @@ class SLATE_API FTabManager : public TSharedFromThis<FTabManager>
 		 * @param NewTabState The new TabState value.
 		 * @param OriginalTabState Only the tabs with this value will be modified. Use ETabState::AnyTab to modify them all.
 		 */
-		void SetTabsTo(const TSharedRef<FTabManager::FLayoutNode>& SomeNode, const ETabState::Type NewTabState, const ETabState::Type OriginalTabState) const;
+		SLATE_API void SetTabsTo(const TSharedRef<FTabManager::FLayoutNode>& SomeNode, const ETabState::Type NewTabState, const ETabState::Type OriginalTabState) const;
 
 		/**
 		 * Notify the tab manager that the NewForegroundTab was brought to front and the BackgroundedTab was send to the background as a result.
 		 */
-		virtual void OnTabForegrounded( const TSharedPtr<SDockTab>& NewForegroundTab, const TSharedPtr<SDockTab>& BackgroundedTab );
-		virtual void OnTabRelocated( const TSharedRef<SDockTab>& RelocatedTab, const TSharedPtr<SWindow>& NewOwnerWindow );
-		virtual void OnTabOpening( const TSharedRef<SDockTab>& TabBeingOpened );
-		virtual void OnTabClosing( const TSharedRef<SDockTab>& TabBeingClosed );
+		SLATE_API virtual void OnTabForegrounded( const TSharedPtr<SDockTab>& NewForegroundTab, const TSharedPtr<SDockTab>& BackgroundedTab );
+		SLATE_API virtual void OnTabRelocated( const TSharedRef<SDockTab>& RelocatedTab, const TSharedPtr<SWindow>& NewOwnerWindow );
+		SLATE_API virtual void OnTabOpening( const TSharedRef<SDockTab>& TabBeingOpened );
+		SLATE_API virtual void OnTabClosing( const TSharedRef<SDockTab>& TabBeingClosed );
 		/** Invoked when a tab manager is closing down. */
-		virtual void OnTabManagerClosing();
+		SLATE_API virtual void OnTabManagerClosing();
 		/** Check these all tabs to see if it is OK to close them. Ignore the TabsToIgnore */
-		virtual bool CanCloseManager( const TSet< TSharedRef<SDockTab> >& TabsToIgnore = TSet< TSharedRef<SDockTab> >() );
+		SLATE_API virtual bool CanCloseManager( const TSet< TSharedRef<SDockTab> >& TabsToIgnore = TSet< TSharedRef<SDockTab> >() );
 
-		static void GetAllStacks( const TSharedRef<SDockingArea>& InDockArea, TArray< TSharedRef<SDockingTabStack> >& OutTabStacks );
+		static SLATE_API void GetAllStacks( const TSharedRef<SDockingArea>& InDockArea, TArray< TSharedRef<SDockingTabStack> >& OutTabStacks );
 
 		/** @return the stack that is under NodeToSearchUnder and contains TabIdToFind; Invalid pointer if not found. */
-		static TSharedPtr<FTabManager::FStack> FindTabUnderNode( const FTabMatcher& Matcher, const TSharedRef<FTabManager::FLayoutNode>& NodeToSearchUnder );
-		int32 FindTabInCollapsedAreas( const FTabMatcher& Matcher );
-		void RemoveTabFromCollapsedAreas( const FTabMatcher& Matcher );
+		static SLATE_API TSharedPtr<FTabManager::FStack> FindTabUnderNode( const FTabMatcher& Matcher, const TSharedRef<FTabManager::FLayoutNode>& NodeToSearchUnder );
+		SLATE_API int32 FindTabInCollapsedAreas( const FTabMatcher& Matcher );
+		SLATE_API void RemoveTabFromCollapsedAreas( const FTabMatcher& Matcher );
 
 		/** Called when tab(s) have been added or windows created */
-		virtual void UpdateStats();
+		SLATE_API virtual void UpdateStats();
 		
 		/** Called at the end of RestoreFrom for tab managers to complete any work after all tabs have been restored */
 		virtual void FinishRestore() {};
 
 	private:
 		/** Checks all dock areas and adds up the number of open tabs and unique parent windows in the manager */
-		void GetRecordableStats( int32& OutTabCount, TArray<TSharedPtr<SWindow>>& OutUniqueParentWindows ) const;
+		SLATE_API void GetRecordableStats( int32& OutTabCount, TArray<TSharedPtr<SWindow>>& OutUniqueParentWindows ) const;
 
 	protected:
 		FTabSpawner TabSpawner;
 		TSharedRef<FTabSpawner> NomadTabSpawner;
 
-		bool HasTabSpawnerFor(FName TabId) const;
+		SLATE_API bool HasTabSpawnerFor(FName TabId) const;
 
 		TArray< TWeakPtr<SDockingArea> > DockAreas;
 		/**
@@ -1161,13 +1161,13 @@ class SLATE_API FTabManager : public TSharedFromThis<FTabManager>
 		int32 LastDocumentUID;
 
 		/** The fallback size for a window */
-		const static FVector2D FallbackWindowSize;
+		SLATE_API const static FVector2D FallbackWindowSize;
 
 		/** Default tab window sizes for newly-created tabs */
-		static TMap<FTabId, FVector2D> DefaultTabWindowSizeMap;
+		static SLATE_API TMap<FTabId, FVector2D> DefaultTabWindowSizeMap;
 
 		/** Returns the default window size for the TabId, or the fallback window size if it wasn't registered */
-		static FVector2D GetDefaultTabWindowSize(const FTabId& TabId);
+		static SLATE_API FVector2D GetDefaultTabWindowSize(const FTabId& TabId);
 
 		/** The ID of the main tab, which cannot be closed*/
 		FTabId MainNonCloseableTabID;
@@ -1202,32 +1202,32 @@ class SLATE_API FTabManager : public TSharedFromThis<FTabManager>
 class FProxyTabmanager;
 
 
-class SLATE_API FGlobalTabmanager : public FTabManager
+class FGlobalTabmanager : public FTabManager
 {
 public:	
 
-	static const TSharedRef<FGlobalTabmanager>& Get();
+	static SLATE_API const TSharedRef<FGlobalTabmanager>& Get();
 
 	/** Subscribe to notifications about the active tab changing */
-	FDelegateHandle OnActiveTabChanged_Subscribe( const FOnActiveTabChanged::FDelegate& InDelegate );
+	SLATE_API FDelegateHandle OnActiveTabChanged_Subscribe( const FOnActiveTabChanged::FDelegate& InDelegate );
 
 	/** Unsubscribe to notifications about the active tab changing */
-	void OnActiveTabChanged_Unsubscribe( FDelegateHandle Handle );
+	SLATE_API void OnActiveTabChanged_Unsubscribe( FDelegateHandle Handle );
 
 	/** Subscribe to notifications about a foreground tab changing */
-	FDelegateHandle OnTabForegrounded_Subscribe(const FOnActiveTabChanged::FDelegate& InDelegate);
+	SLATE_API FDelegateHandle OnTabForegrounded_Subscribe(const FOnActiveTabChanged::FDelegate& InDelegate);
 
 	/** Unsubscribe to notifications about a foreground tab changing */
-	void OnTabForegrounded_Unsubscribe(FDelegateHandle Handle);
+	SLATE_API void OnTabForegrounded_Unsubscribe(FDelegateHandle Handle);
 
 	/** @return the currently active tab; NULL pointer if there is no active tab */
-	TSharedPtr<SDockTab> GetActiveTab() const;
+	SLATE_API TSharedPtr<SDockTab> GetActiveTab() const;
 
 	/** Can the manager activate this Tab as the new active tab? */
-	bool CanSetAsActiveTab(const TSharedPtr<SDockTab>& Tab);
+	SLATE_API bool CanSetAsActiveTab(const TSharedPtr<SDockTab>& Tab);
 
 	/** Activate the NewActiveTab. If NewActiveTab is NULL, the active tab is cleared. */
-	void SetActiveTab( const TSharedPtr<SDockTab>& NewActiveTab );
+	SLATE_API void SetActiveTab( const TSharedPtr<SDockTab>& NewActiveTab );
 
 	/**
 	 * Register a new normad tab spawner with the tab manager.  The spawner will be called when anyone calls
@@ -1238,59 +1238,59 @@ public:
 	 * @param CanSpawnTab The callback that will be used to ask if spawning the tab is allowed
 	 * @return The registration entry for the spawner.
 	 */
-	FTabSpawnerEntry& RegisterNomadTabSpawner( const FName TabId, const FOnSpawnTab& OnSpawnTab, const FCanSpawnTab& CanSpawnTab = FCanSpawnTab());
+	SLATE_API FTabSpawnerEntry& RegisterNomadTabSpawner( const FName TabId, const FOnSpawnTab& OnSpawnTab, const FCanSpawnTab& CanSpawnTab = FCanSpawnTab());
 
-	void UnregisterNomadTabSpawner( const FName TabId );
+	SLATE_API void UnregisterNomadTabSpawner( const FName TabId );
 
-	void SetApplicationTitle( const FText& AppTitle );
+	SLATE_API void SetApplicationTitle( const FText& AppTitle );
 
-	const FText& GetApplicationTitle() const;
+	SLATE_API const FText& GetApplicationTitle() const;
 
 	static TSharedRef<FGlobalTabmanager> New()
 	{
 		return MakeShareable( new FGlobalTabmanager() );
 	}
 
-	virtual bool CanCloseManager( const TSet< TSharedRef<SDockTab> >& TabsToIgnore = TSet< TSharedRef<SDockTab> >()) override;
+	SLATE_API virtual bool CanCloseManager( const TSet< TSharedRef<SDockTab> >& TabsToIgnore = TSet< TSharedRef<SDockTab> >()) override;
 
 	/** Gets the major tab for the manager */
-	TSharedPtr<SDockTab> GetMajorTabForTabManager(const TSharedRef<FTabManager>& ChildManager);
+	SLATE_API TSharedPtr<SDockTab> GetMajorTabForTabManager(const TSharedRef<FTabManager>& ChildManager);
 
 	/** 
 	 * Gets the tab manager that a major tab owns. The returned tab manager is the tab manager that manages the minor tabs for a major tab.
 	 * Note: this is not the same as DockTab->GetTabManager(). That function returns the tab manager the major tab is in
 	 */
-	TSharedPtr<FTabManager> GetTabManagerForMajorTab(const TSharedPtr<SDockTab> DockTab) const;
+	SLATE_API TSharedPtr<FTabManager> GetTabManagerForMajorTab(const TSharedPtr<SDockTab> DockTab) const;
 
 	/** Draw the user's attention to a child tab manager */
-	void DrawAttentionToTabManager( const TSharedRef<FTabManager>& ChildManager );
+	SLATE_API void DrawAttentionToTabManager( const TSharedRef<FTabManager>& ChildManager );
 
-	TSharedRef<FTabManager> NewTabManager( const TSharedRef<SDockTab>& InOwnerTab );
+	SLATE_API TSharedRef<FTabManager> NewTabManager( const TSharedRef<SDockTab>& InOwnerTab );
 	
 	/**
 	 * Update the native, global menu bar if it is being used for a specific tab managed by the global tab manager.
 	 * @param ForTab The tab to update the main menu.
 	 * @param bForce Used to force an update even if the parent window doesn't contain the widget with keyboard focus.
 	 */
-	void UpdateMainMenu(const TSharedRef<SDockTab>& ForTab, bool const bForce);
+	SLATE_API void UpdateMainMenu(const TSharedRef<SDockTab>& ForTab, bool const bForce);
 
 	/** Persist and serialize the layout of every TabManager and the custom visual state of every Tab. */
-	void SaveAllVisualState();
+	SLATE_API void SaveAllVisualState();
 
 	/** Provide a window under which all other windows in this application should nest. */
-	void SetRootWindow( const TSharedRef<SWindow> InRootWindow );
+	SLATE_API void SetRootWindow( const TSharedRef<SWindow> InRootWindow );
 
 	/** The window under which all other windows in our app nest; might be null */
-	TSharedPtr<SWindow> GetRootWindow() const;
+	SLATE_API TSharedPtr<SWindow> GetRootWindow() const;
 
 	/** Adds a legacy tab type to the tab type redirection map so tabs loaded with this type will be automatically converted to the new type */
-	void AddLegacyTabType(FName InLegacyTabType, FName InNewTabType);
+	SLATE_API void AddLegacyTabType(FName InLegacyTabType, FName InNewTabType);
 
 	/** Returns true if the specified tab type is registered as a legacy tab */
-	bool IsLegacyTabType(FName InTabType) const;
+	SLATE_API bool IsLegacyTabType(FName InTabType) const;
 
 	/** If the specified TabType is deprecated, returns the new replacement tab type. Otherwise, returns InTabType */
-	FName GetTabTypeForPotentiallyLegacyTab(FName InTabType) const;
+	SLATE_API FName GetTabTypeForPotentiallyLegacyTab(FName InTabType) const;
 
 	/** Returns the highest number of tabs that were open simultaneously during this session */
 	int32 GetMaximumTabCount() const { return AllTabsMaxCount; }
@@ -1298,24 +1298,24 @@ public:
 	/** Returns the highest number of parent windows that were open simultaneously during this session */
 	int32 GetMaximumWindowCount() const { return AllAreasWindowMaxCount; }
 
-	void SetProxyTabManager(TSharedPtr<FProxyTabmanager> InProxyTabManager);
+	SLATE_API void SetProxyTabManager(TSharedPtr<FProxyTabmanager> InProxyTabManager);
 
 	DECLARE_DELEGATE(FOnOverrideDockableAreaRestore);
 	/** Used to override dockable area restoration behavior */
 	FOnOverrideDockableAreaRestore OnOverrideDockableAreaRestore_Handler;
 
 protected:
-	virtual void OnTabForegrounded( const TSharedPtr<SDockTab>& NewForegroundTab, const TSharedPtr<SDockTab>& BackgroundedTab ) override;
-	virtual void OnTabRelocated( const TSharedRef<SDockTab>& RelocatedTab, const TSharedPtr<SWindow>& NewOwnerWindow ) override;
-	virtual void OnTabClosing( const TSharedRef<SDockTab>& TabBeingClosed ) override;
-	virtual void UpdateStats() override;
+	SLATE_API virtual void OnTabForegrounded( const TSharedPtr<SDockTab>& NewForegroundTab, const TSharedPtr<SDockTab>& BackgroundedTab ) override;
+	SLATE_API virtual void OnTabRelocated( const TSharedRef<SDockTab>& RelocatedTab, const TSharedPtr<SWindow>& NewOwnerWindow ) override;
+	SLATE_API virtual void OnTabClosing( const TSharedRef<SDockTab>& TabBeingClosed ) override;
+	SLATE_API virtual void UpdateStats() override;
 
-	virtual void OpenUnmanagedTab(FName PlaceholderId, const FSearchPreference& SearchPreference, const TSharedRef<SDockTab>& UnmanagedTab) override;
+	SLATE_API virtual void OpenUnmanagedTab(FName PlaceholderId, const FSearchPreference& SearchPreference, const TSharedRef<SDockTab>& UnmanagedTab) override;
 
-	virtual void FinishRestore() override;
+	SLATE_API virtual void FinishRestore() override;
 
 public:
-	virtual void OnTabManagerClosing() override;
+	SLATE_API virtual void OnTabManagerClosing() override;
 
 
 private:
@@ -1403,7 +1403,7 @@ DECLARE_MULTICAST_DELEGATE_TwoParams(FIsTabSupportedEvent, FTabId /* TabId */, b
 DECLARE_MULTICAST_DELEGATE_OneParam(FOnTabEvent, TSharedPtr<SDockTab>);
 
 
-class SLATE_API FProxyTabmanager : public FTabManager
+class FProxyTabmanager : public FTabManager
 {
 public:
 
@@ -1414,11 +1414,11 @@ public:
 		bCanDoDragOperation = false;
 	}
 
-	virtual void OpenUnmanagedTab(FName PlaceholderId, const FSearchPreference& SearchPreference, const TSharedRef<SDockTab>& UnmanagedTab) override;
-	virtual void DrawAttention(const TSharedRef<SDockTab>& TabToHighlight) override;
+	SLATE_API virtual void OpenUnmanagedTab(FName PlaceholderId, const FSearchPreference& SearchPreference, const TSharedRef<SDockTab>& UnmanagedTab) override;
+	SLATE_API virtual void DrawAttention(const TSharedRef<SDockTab>& TabToHighlight) override;
 	
-	bool IsTabSupported( const FTabId TabId ) const;
-	void SetParentWindow(TSharedRef<SWindow> InParentWindow);
+	SLATE_API bool IsTabSupported( const FTabId TabId ) const;
+	SLATE_API void SetParentWindow(TSharedRef<SWindow> InParentWindow);
 
 public:
 	FIsTabSupportedEvent OnIsTabSupported;
diff --git a/Engine/Source/Runtime/Slate/Public/Framework/Layout/IScrollableWidget.h b/Engine/Source/Runtime/Slate/Public/Framework/Layout/IScrollableWidget.h
--- a/Engine/Source/Runtime/Slate/Public/Framework/Layout/IScrollableWidget.h
+++ b/Engine/Source/Runtime/Slate/Public/Framework/Layout/IScrollableWidget.h
@@ -7,7 +7,7 @@
 /**
 * An interface to allow Slate Widgets to expose scrollable functionality.
 */
-class SLATE_API IScrollableWidget
+class IScrollableWidget
 {
 public:
 	virtual ~IScrollableWidget() {}
diff --git a/Engine/Source/Runtime/Slate/Public/Framework/Layout/InertialScrollManager.h b/Engine/Source/Runtime/Slate/Public/Framework/Layout/InertialScrollManager.h
--- a/Engine/Source/Runtime/Slate/Public/Framework/Layout/InertialScrollManager.h
+++ b/Engine/Source/Runtime/Slate/Public/Framework/Layout/InertialScrollManager.h
@@ -11,17 +11,17 @@
  * when the velocity grows small enough, and the percentage of friction lost prevents large velocities
  * from scrolling forever.
  */
-class SLATE_API FInertialScrollManager
+class FInertialScrollManager
 {
 	/**
 	 * This is the percentage of velocity loss per second.
 	 */
-	static float FrictionCoefficient;
+	static SLATE_API float FrictionCoefficient;
 
 	/**
 	 * This is a constant amount of velocity lost per second.
 	 */
-	static float StaticVelocityDrag;
+	static SLATE_API float StaticVelocityDrag;
 
 public:
 	/**
@@ -29,16 +29,16 @@ public:
 	  * @param ScrollDecceleration	The acceleration against the velocity causing it to decay.
 	  * @param SampleTimeout		Samples older than this amount of time will be discarded.
 	  */
-	FInertialScrollManager(double SampleTimeout = 0.1f);
+	SLATE_API FInertialScrollManager(double SampleTimeout = 0.1f);
 
 	/** Adds a scroll velocity sample to help calculate a smooth velocity */
-	void AddScrollSample(float Delta, double CurrentTime);
+	SLATE_API void AddScrollSample(float Delta, double CurrentTime);
 
 	/** Updates the current scroll velocity. Call every frame. */
-	void UpdateScrollVelocity(const float InDeltaTime);
+	SLATE_API void UpdateScrollVelocity(const float InDeltaTime);
 
 	/** Instantly end the inertial scroll */
-	void ClearScrollVelocity();
+	SLATE_API void ClearScrollVelocity();
 
 	/** Gets the calculated velocity of the scroll. */
 	float GetScrollVelocity() const { return ScrollVelocity; }
diff --git a/Engine/Source/Runtime/Slate/Public/Framework/Layout/Overscroll.h b/Engine/Source/Runtime/Slate/Public/Framework/Layout/Overscroll.h
--- a/Engine/Source/Runtime/Slate/Public/Framework/Layout/Overscroll.h
+++ b/Engine/Source/Runtime/Slate/Public/Framework/Layout/Overscroll.h
@@ -17,27 +17,27 @@ enum class EAllowOverscroll : uint8
 /**
  * Handles overscroll management.
  */
-struct SLATE_API FOverscroll
+struct FOverscroll
 {
 public:
 
 	/** The amount to scale the logarithm by to make it more loose */
-	static float Looseness;
+	static SLATE_API float Looseness;
 	/** The "max" used to perform the interpolation snap back, and make it faster the further away it is. */
-	static float OvershootLooseMax;
+	static SLATE_API float OvershootLooseMax;
 	/** The bounce back rate when the overscroll stops. */
-	static float OvershootBounceRate;
+	static SLATE_API float OvershootBounceRate;
 
-	FOverscroll();
+	SLATE_API FOverscroll();
 
 	/** @return The Amount actually scrolled */
-	float ScrollBy(const FGeometry& AllottedGeometry, float LocalDeltaScroll);
+	SLATE_API float ScrollBy(const FGeometry& AllottedGeometry, float LocalDeltaScroll);
 
 	/** How far the user scrolled above/below the beginning/end of the list. */
-	float GetOverscroll(const FGeometry& AllottedGeometry) const;
+	SLATE_API float GetOverscroll(const FGeometry& AllottedGeometry) const;
 
 	/** Ticks the overscroll manager so it can animate. */
-	void UpdateOverscroll(float InDeltaTime);
+	SLATE_API void UpdateOverscroll(float InDeltaTime);
 
 	/**
 	 * Should ScrollDelta be applied to overscroll or to regular item scrolling.
@@ -48,10 +48,10 @@ public:
 	 *
 	 * @return true if the user's scrolling should be applied toward overscroll.
 	 */
-	bool ShouldApplyOverscroll(const bool bIsAtStartOfList, const bool bIsAtEndOfList, const float ScrollDelta) const;
+	SLATE_API bool ShouldApplyOverscroll(const bool bIsAtStartOfList, const bool bIsAtEndOfList, const float ScrollDelta) const;
 	
 	/** Resets the overscroll amout. */
-	void ResetOverscroll();
+	SLATE_API void ResetOverscroll();
 private:
 	/** How much we've over-scrolled above/below the beginning/end of the list, stored in log form */
 	float OverscrollAmount;
diff --git a/Engine/Source/Runtime/Slate/Public/Framework/Layout/ScrollyZoomy.h b/Engine/Source/Runtime/Slate/Public/Framework/Layout/ScrollyZoomy.h
--- a/Engine/Source/Runtime/Slate/Public/Framework/Layout/ScrollyZoomy.h
+++ b/Engine/Source/Runtime/Slate/Public/Framework/Layout/ScrollyZoomy.h
@@ -42,7 +42,7 @@ public:
  * instance of FScrollyZoomy as a widget member variable, and call this
  * class's event handlers from your own widget's event handler callbacks.
  */
-class SLATE_API FScrollyZoomy
+class FScrollyZoomy
 {
 public:
 
@@ -51,7 +51,7 @@ public:
 	 *
 	 * @param InUseInertialScrolling Whether inertial scrolling should be used (default = true).
 	 */
-	FScrollyZoomy(const bool InUseInertialScrolling = true);
+	SLATE_API FScrollyZoomy(const bool InUseInertialScrolling = true);
 
 	/** 
 	 * Should be called every frame to update simulation state.
@@ -59,7 +59,7 @@ public:
 	 * @param DeltaTime Time that's passed.
 	 * @param ScrollableZoomable Interface to the widget to scroll/zoom.
 	 */
-	void Tick(const float DeltaTime, IScrollableZoomable& ScrollableZoomable);
+	SLATE_API void Tick(const float DeltaTime, IScrollableZoomable& ScrollableZoomable);
 
 	/**
 	 * Should be called when a mouse button is pressed.
@@ -67,7 +67,7 @@ public:
 	 * @param MouseEvent The mouse event passed to a widget's OnMouseButtonDown function.
 	 * @return If FReply::Handled is returned, that should be passed as the result of the calling function.
 	 */
-	FReply OnMouseButtonDown(const FPointerEvent& MouseEvent);
+	SLATE_API FReply OnMouseButtonDown(const FPointerEvent& MouseEvent);
 
 	/**
 	 * Should be called when a mouse button is released.
@@ -77,7 +77,7 @@ public:
 	 * @param MouseEvent The mouse event passed to a widget's OnMouseButtonUp function.
 	 * @return If FReply::Handled is returned, that should be passed as the result of the calling function.
 	 */
-	FReply OnMouseButtonUp(const TSharedRef<SWidget> MyWidget, const FGeometry& MyGeometry, const FPointerEvent& MouseEvent);
+	SLATE_API FReply OnMouseButtonUp(const TSharedRef<SWidget> MyWidget, const FGeometry& MyGeometry, const FPointerEvent& MouseEvent);
 		
 	/**
 	 * Should be called when a mouse move event occurs.
@@ -88,7 +88,7 @@ public:
 	 * @param MouseEvent The mouse event passed to a widget's OnMouseMove function.
 	 * @return If FReply::Handled is returned, that should be passed as the result of the calling function.
 	 */
-	FReply OnMouseMove(const TSharedRef<SWidget> MyWidget, IScrollableZoomable& ScrollableZoomable, const FGeometry& MyGeometry, const FPointerEvent& MouseEvent);
+	SLATE_API FReply OnMouseMove(const TSharedRef<SWidget> MyWidget, IScrollableZoomable& ScrollableZoomable, const FGeometry& MyGeometry, const FPointerEvent& MouseEvent);
 
 	/**
 	 * Should be called from your widget's OnMouseLeave override.
@@ -96,7 +96,7 @@ public:
 	 * @param MyWidget Pointer to the widget that owns this object.
 	 * @param MouseEvent The mouse leave event that was passed to the widget's OnMouseLeave function.
 	 */
-	void OnMouseLeave(const TSharedRef<SWidget> MyWidget, const FPointerEvent& MouseEvent);
+	SLATE_API void OnMouseLeave(const TSharedRef<SWidget> MyWidget, const FPointerEvent& MouseEvent);
 		
 	/**
 	 * Should be called by your widget when the mouse wheel is used
@@ -105,21 +105,21 @@ public:
 	 * @param ScrollableZoomable	Interface to the widget to scroll/zoom.
 	 * @return If FReply::Handled is returned, that should be passed as the result of the calling function.
 	 */
-	FReply OnMouseWheel(const FPointerEvent& MouseEvent, IScrollableZoomable& ScrollableZoomable);
+	SLATE_API FReply OnMouseWheel(const FPointerEvent& MouseEvent, IScrollableZoomable& ScrollableZoomable);
 
 	/**
 	 * Call this from your widget's OnCursorQuery function.
 	 *
 	 * @return The cursor reply to pass back.
 	 */
-	FCursorReply OnCursorQuery() const;
+	SLATE_API FCursorReply OnCursorQuery() const;
 
 	/**
 	 * Whether the user is actively scrolling.
 	 *
 	 * @return true if the user is scrolling, false otherwise.
 	 */
-	bool IsRightClickScrolling() const;
+	SLATE_API bool IsRightClickScrolling() const;
 
 	/**
 	 * Whether a software cursor should be rendered.
@@ -129,7 +129,7 @@ public:
 	 * @return true if a software cursor should be rendered, false otherwise.
 	 * @see GetSoftwareCursorPosition
 	 */
-	bool NeedsSoftwareCursor() const;
+	SLATE_API bool NeedsSoftwareCursor() const;
 
 	/**
 	 * Get the position of the software cursor (when NeedsSoftwareCursor is true).
@@ -139,7 +139,7 @@ public:
 	 * @return Cursor position.
 	 * @see NeedsSoftwareCursor
 	 */
-	UE::Slate::FDeprecateVector2DResult GetSoftwareCursorPosition() const;
+	SLATE_API UE::Slate::FDeprecateVector2DResult GetSoftwareCursorPosition() const;
 
 	/**
 	 * Call this from your widget's OnPaint to paint a software cursor, if needed 
@@ -150,7 +150,7 @@ public:
 	 * @param LayerId Layer identifier.
 	 * @return New layer Identifier.
 	 */
-	int32 PaintSoftwareCursorIfNeeded(const FGeometry& AllottedGeometry, const FSlateRect& MyCullingRect, FSlateWindowElementList& OutDrawElements, int32 LayerId) const;
+	SLATE_API int32 PaintSoftwareCursorIfNeeded(const FGeometry& AllottedGeometry, const FSlateRect& MyCullingRect, FSlateWindowElementList& OutDrawElements, int32 LayerId) const;
 
 private:
 
diff --git a/Engine/Source/Runtime/Slate/Public/Framework/MetaData/DriverMetaData.h b/Engine/Source/Runtime/Slate/Public/Framework/MetaData/DriverMetaData.h
--- a/Engine/Source/Runtime/Slate/Public/Framework/MetaData/DriverMetaData.h
+++ b/Engine/Source/Runtime/Slate/Public/Framework/MetaData/DriverMetaData.h
@@ -10,12 +10,12 @@ class ISlateMetaData;
 /**
  * A static helper class which is used to easily construct various types of AutomationDriver specific SlateMetaData.
  */
-class SLATE_API FDriverMetaData
+class FDriverMetaData
 {
 public:
 
 	/**
 	 * @return the automation driver specific metadata type for specifying an Id for a widget
 	 */
-	static TSharedRef<ISlateMetaData> Id(FName Tag);
-};
\ No newline at end of file
+	static SLATE_API TSharedRef<ISlateMetaData> Id(FName Tag);
+};
diff --git a/Engine/Source/Runtime/Slate/Public/Framework/MultiBox/MultiBox.h b/Engine/Source/Runtime/Slate/Public/Framework/MultiBox/MultiBox.h
--- a/Engine/Source/Runtime/Slate/Public/Framework/MultiBox/MultiBox.h
+++ b/Engine/Source/Runtime/Slate/Public/Framework/MultiBox/MultiBox.h
@@ -54,7 +54,7 @@ namespace MultiBoxConstants
 /**
  * MultiBlock (abstract).  Wraps a "block" of useful UI functionality that can be added to a MultiBox.
  */
-class SLATE_API FMultiBlock
+class FMultiBlock
 	: public TSharedFromThis< FMultiBlock >		// Enables this->AsShared()
 {
 
@@ -197,7 +197,7 @@ public:
 	 *
 	 * @return  MultiBlock widget object
 	 */
-	TSharedRef<class IMultiBlockBaseWidget> MakeWidget(TSharedRef< class SMultiBoxWidget > InOwnerMultiBoxWidget, EMultiBlockLocation::Type InLocation, bool bSectionContainsIcons, TSharedPtr<SWidget> OptionsBlockWidget) const;
+	SLATE_API TSharedRef<class IMultiBlockBaseWidget> MakeWidget(TSharedRef< class SMultiBoxWidget > InOwnerMultiBoxWidget, EMultiBlockLocation::Type InLocation, bool bSectionContainsIcons, TSharedPtr<SWidget> OptionsBlockWidget) const;
 
 	/**
 	 * Gets the type of this MultiBox
@@ -234,14 +234,14 @@ public:
 	*
 	* @param	bSearchable		The searchable state to set
 	*/
-	void SetSearchable(bool bSearchable);
+	SLATE_API void SetSearchable(bool bSearchable);
 
 	/**
 	* Gets the searchable state of this block
 	*
 	* @return	Whether this block is searchable
 	*/
-	bool GetSearchable() const;
+	SLATE_API bool GetSearchable() const;
 
 	/** Gets the extension hook so users can see what hooks are where */
 	FName GetExtensionHook() const { return ExtensionHook; }
@@ -301,17 +301,17 @@ private:
 /**
  * MultiBox.  Contains a list of MultiBlocks that provide various functionality.
  */
-class SLATE_API FMultiBox
+class FMultiBox
 	: public TSharedFromThis< FMultiBox >		// Enables this->AsShared()
 {
 
 public:
-	virtual ~FMultiBox();
+	SLATE_API virtual ~FMultiBox();
 
 	/**
 	 * Creates a new multibox instance
 	 */
-	static TSharedRef<FMultiBox> Create( const EMultiBoxType InType,  FMultiBoxCustomization InCustomization, const bool bInShouldCloseWindowAfterMenuSelection );
+	static SLATE_API TSharedRef<FMultiBox> Create( const EMultiBoxType InType,  FMultiBoxCustomization InCustomization, const bool bInShouldCloseWindowAfterMenuSelection );
 
 	/**
 	 * Gets the type of this MultiBox
@@ -336,22 +336,22 @@ public:
 	/**
 	 * Adds a MultiBlock to this MultiBox, to the end of the list
 	 */
-	void AddMultiBlock( TSharedRef< const FMultiBlock > InBlock );
+	SLATE_API void AddMultiBlock( TSharedRef< const FMultiBlock > InBlock );
 
 	/**
 	 * Adds a MultiBlock to this MultiBox, to the front of the list
 	 */
-	void AddMultiBlockToFront(TSharedRef< const FMultiBlock > InBlock);
+	SLATE_API void AddMultiBlockToFront(TSharedRef< const FMultiBlock > InBlock);
 
 	/**
 	 * Removes a MultiBlock from the list for user customization
 	 */
-	void RemoveCustomMultiBlock( TSharedRef< const FMultiBlock> InBlock );
+	SLATE_API void RemoveCustomMultiBlock( TSharedRef< const FMultiBlock> InBlock );
 
 	/**
 	 * Inserts a MultiBlock to the list for user customization
 	 */
-	void InsertCustomMultiBlock( TSharedRef<const FMultiBlock> InBlock, int32 Index );
+	SLATE_API void InsertCustomMultiBlock( TSharedRef<const FMultiBlock> InBlock, int32 Index );
 
 	DECLARE_DELEGATE_TwoParams( FOnMakeMultiBoxBuilderOverride, const TSharedRef<FMultiBox>&, const TSharedRef<SMultiBoxWidget>& );
 
@@ -367,7 +367,7 @@ public:
 	 *
 	 * @return  MultiBox widget object
 	 */
-	TSharedRef< class SMultiBoxWidget > MakeWidget( bool bSearchable, FOnMakeMultiBoxBuilderOverride* InMakeMultiBoxBuilderOverride = nullptr, TAttribute<float> InMaxHeight = TAttribute<float>() );
+	SLATE_API TSharedRef< class SMultiBoxWidget > MakeWidget( bool bSearchable, FOnMakeMultiBoxBuilderOverride* InMakeMultiBoxBuilderOverride = nullptr, TAttribute<float> InMaxHeight = TAttribute<float>() );
 
 
 	/**
@@ -394,7 +394,7 @@ public:
 	}
 
 	/** @return The customization name for this box */
-	FName GetCustomizationName() const;
+	SLATE_API FName GetCustomizationName() const;
 
 	/**
 	 * Creates a block from the provided command that is compatible with this box 
@@ -402,7 +402,7 @@ public:
 	 * @param Command	The UI command to create the block from
 	 * @return The created multi block.  If null, this command could not be placed in this block
 	 */
-	TSharedPtr<FMultiBlock> MakeMultiBlockFromCommand( TSharedPtr<const FUICommandInfo> Command, bool bCommandMustBeBound ) const;
+	SLATE_API TSharedPtr<FMultiBlock> MakeMultiBlockFromCommand( TSharedPtr<const FUICommandInfo> Command, bool bCommandMustBeBound ) const;
 
 	/**
 	 * Finds an existing block by name and type
@@ -410,19 +410,19 @@ public:
 	 * @param InName The name to search for
 	  * @param InType The type to match during the search
 	 */
-	TSharedPtr<const FMultiBlock> FindBlockFromNameAndType(const FName InName, const EMultiBlockType InType) const;
+	SLATE_API TSharedPtr<const FMultiBlock> FindBlockFromNameAndType(const FName InName, const EMultiBlockType InType) const;
 
 	/** @return Is being editing */
-	bool IsInEditMode() const;
+	SLATE_API bool IsInEditMode() const;
 
 	/** @return The tool menu associated with this multi box */
-	UToolMenuBase* GetToolMenu() const;
+	SLATE_API UToolMenuBase* GetToolMenu() const;
 
 	/**
 	 * Only callable during edit mode
 	 * @return Index of section heading or separator
 	 */
-	int32 GetSectionEditBounds(const int32 Index, int32& OutSectionEndIndex) const;
+	SLATE_API int32 GetSectionEditBounds(const int32 Index, int32& OutSectionEndIndex) const;
 
 	DECLARE_DELEGATE_OneParam(FEditSelectionChangedDelegate, TSharedRef<const FMultiBlock>);
 
@@ -452,12 +452,12 @@ private:
 	 * @param	InType	Type of MultiBox
 	 * @param	bInShouldCloseWindowAfterMenuSelection	Sets whether or not the window that contains this multibox should be destroyed after the user clicks on a menu item in this box
 	 */
-	FMultiBox( const EMultiBoxType InType,  FMultiBoxCustomization InCustomization, const bool bInShouldCloseWindowAfterMenuSelection );
+	SLATE_API FMultiBox( const EMultiBoxType InType,  FMultiBoxCustomization InCustomization, const bool bInShouldCloseWindowAfterMenuSelection );
 
 	/**
 	 * @return true if this box can be customized by a user
 	 */
-	bool IsCustomizable() const;
+	SLATE_API bool IsCustomizable() const;
 
 private:
 
@@ -488,7 +488,7 @@ private:
 /**
  * MultiBlock Slate widget interface
  */
-class SLATE_API IMultiBlockBaseWidget
+class IMultiBlockBaseWidget
 {
 
 public:
@@ -561,7 +561,7 @@ public:
 /**
  * MultiBlock Slate base widget (pure virtual).  You'll derive your own MultiBlock class from this base class.
  */
-class SLATE_API SMultiBlockBaseWidget
+class SMultiBlockBaseWidget
 	: public IMultiBlockBaseWidget,
 	  public SCompoundWidget
 {
@@ -570,19 +570,19 @@ public:
 	TSharedPtr<const FMultiBlock> GetBlock() const { return MultiBlock; }
 
 	/** IMultiBlockBaseWidget interface */
-	virtual TSharedRef< SWidget > AsWidget() override;
-	virtual TSharedRef< const SWidget > AsWidget() const override;
-	virtual void SetOwnerMultiBoxWidget(TSharedRef<SMultiBoxWidget> InOwnerMultiBoxWidget) override;
-	virtual void SetMultiBlock(TSharedRef<const FMultiBlock> InMultiBlock) override;
-	virtual void SetOptionsBlockWidget(TSharedPtr<SWidget> InOptionsBlock) override;
-	virtual void SetMultiBlockLocation(EMultiBlockLocation::Type InLocation, bool bInSectionContainsIcons) override;
-	virtual EMultiBlockLocation::Type GetMultiBlockLocation() override;
-	virtual bool IsInEditMode() const override;
+	SLATE_API virtual TSharedRef< SWidget > AsWidget() override;
+	SLATE_API virtual TSharedRef< const SWidget > AsWidget() const override;
+	SLATE_API virtual void SetOwnerMultiBoxWidget(TSharedRef<SMultiBoxWidget> InOwnerMultiBoxWidget) override;
+	SLATE_API virtual void SetMultiBlock(TSharedRef<const FMultiBlock> InMultiBlock) override;
+	SLATE_API virtual void SetOptionsBlockWidget(TSharedPtr<SWidget> InOptionsBlock) override;
+	SLATE_API virtual void SetMultiBlockLocation(EMultiBlockLocation::Type InLocation, bool bInSectionContainsIcons) override;
+	SLATE_API virtual EMultiBlockLocation::Type GetMultiBlockLocation() override;
+	SLATE_API virtual bool IsInEditMode() const override;
 
 	/** SWidget Interface */
-	virtual void OnDragEnter(const FGeometry& MyGeometry, const FDragDropEvent& DragDropEvent) override;
-	virtual FReply OnDragOver(const FGeometry& MyGeometry, const FDragDropEvent& DragDropEvent) override;
-	virtual FReply OnDrop(const FGeometry& MyGeometry, const FDragDropEvent& DragDropEvent) override;
+	SLATE_API virtual void OnDragEnter(const FGeometry& MyGeometry, const FDragDropEvent& DragDropEvent) override;
+	SLATE_API virtual FReply OnDragOver(const FGeometry& MyGeometry, const FDragDropEvent& DragDropEvent) override;
+	SLATE_API virtual FReply OnDrop(const FGeometry& MyGeometry, const FDragDropEvent& DragDropEvent) override;
 protected:
 
 	/** Weak reference back to the MultiBox widget that owns us */
@@ -603,7 +603,7 @@ protected:
 /**
  * MultiBox Slate widget
  */
-class SLATE_API SMultiBoxWidget
+class SMultiBoxWidget
 	: public SMenuOwner
 {
 
@@ -623,7 +623,7 @@ public:
 	 *
 	 * @param	InArgs	The declaration data for this widget
 	 */
-	void Construct( const FArguments& InArgs );
+	SLATE_API void Construct( const FArguments& InArgs );
 
 	/**
 	 * Associates this widget with a MultiBox
@@ -661,13 +661,13 @@ public:
 	*
 	* @param	bSearchable		The searchable state to set
 	*/
-	void SetSearchable(bool bSearchable);
+	SLATE_API void SetSearchable(bool bSearchable);
 	/**
 	* Gets the searchable state of this multibox
 	*
 	* @return	Whether this block is searchable
 	*/
-	bool GetSearchable() const;
+	SLATE_API bool GetSearchable() const;
 
 	/**
 	 * @return the the LinkedBoxManager
@@ -682,17 +682,17 @@ public:
 	/**
 	 * Builds this MultiBox widget up from the MultiBox associated with it
 	 */
-	void BuildMultiBoxWidget();
+	SLATE_API void BuildMultiBoxWidget();
 
 	/** Generates the tiles for an STileView for button rows */
-	TSharedRef<ITableRow> GenerateTiles(TSharedPtr<SWidget> Item, const TSharedRef<STableViewBase>& OwnerTable);
+	SLATE_API TSharedRef<ITableRow> GenerateTiles(TSharedPtr<SWidget> Item, const TSharedRef<STableViewBase>& OwnerTable);
 
 	/** Gets the maximum item width and height of the consituent widgets */
-	float GetItemWidth() const;
-	float GetItemHeight() const;
+	SLATE_API float GetItemWidth() const;
+	SLATE_API float GetItemHeight() const;
 
 	/** Event handler for clicking the wrap button */
-	TSharedRef<SWidget> OnWrapButtonClicked();
+	SLATE_API TSharedRef<SWidget> OnWrapButtonClicked();
 
 	const ISlateStyle* GetStyleSet() const { return MultiBox->GetStyleSet(); }
 	const FName& GetStyleName() const { return MultiBox->GetStyleName(); }
@@ -700,72 +700,72 @@ public:
 	/**
 	 * Called when a user drags a UI command into a multiblock in this widget
 	 */
-	void OnCustomCommandDragEnter( TSharedRef<const FMultiBlock> MultiBlock, const FGeometry& MyGeometry, const FDragDropEvent& DragDropEvent );
+	SLATE_API void OnCustomCommandDragEnter( TSharedRef<const FMultiBlock> MultiBlock, const FGeometry& MyGeometry, const FDragDropEvent& DragDropEvent );
 
 	/**
 	 * Called when a user drags a UI command within multiblock in this widget
 	 */
-	void OnCustomCommandDragged( TSharedRef<const FMultiBlock> MultiBlock, const FGeometry& MyGeometry, const FDragDropEvent& DragDropEvent );
+	SLATE_API void OnCustomCommandDragged( TSharedRef<const FMultiBlock> MultiBlock, const FGeometry& MyGeometry, const FDragDropEvent& DragDropEvent );
 
 	/**
 	 * Called when a user drops a UI command in this widget
 	 */
-	void OnCustomCommandDropped();
+	SLATE_API void OnCustomCommandDropped();
 	
 	/**
 	 * Called after a drag was initiated from this box but was dropped elsewhere
 	 */
-	void OnDropExternal();
+	SLATE_API void OnDropExternal();
 
 	/** Helper function used to transfer focus to the next/previous widget */
-	static FReply FocusNextWidget(EUINavigation NavigationType);
+	static SLATE_API FReply FocusNextWidget(EUINavigation NavigationType);
 
 	/** SWidget interface */
-	virtual FReply OnDragOver( const FGeometry& MyGeometry, const FDragDropEvent& DragDropEvent ) override;
-	virtual FReply OnDrop( const FGeometry& MyGeometry, const FDragDropEvent& DragDropEvent ) override;
-	virtual bool SupportsKeyboardFocus() const override;
-	virtual FReply OnFocusReceived( const FGeometry& MyGeometry, const FFocusEvent& InFocusEvent ) override;
-	virtual FReply OnKeyDown( const FGeometry& MyGeometry, const FKeyEvent& KeyEvent ) override;
-	virtual FReply OnKeyChar(const FGeometry& MyGeometry, const FCharacterEvent& InCharacterEvent) override;
-	virtual bool OnVisualizeTooltip(const TSharedPtr<SWidget>& TooltipContent) override;
+	SLATE_API virtual FReply OnDragOver( const FGeometry& MyGeometry, const FDragDropEvent& DragDropEvent ) override;
+	SLATE_API virtual FReply OnDrop( const FGeometry& MyGeometry, const FDragDropEvent& DragDropEvent ) override;
+	SLATE_API virtual bool SupportsKeyboardFocus() const override;
+	SLATE_API virtual FReply OnFocusReceived( const FGeometry& MyGeometry, const FFocusEvent& InFocusEvent ) override;
+	SLATE_API virtual FReply OnKeyDown( const FGeometry& MyGeometry, const FKeyEvent& KeyEvent ) override;
+	SLATE_API virtual FReply OnKeyChar(const FGeometry& MyGeometry, const FCharacterEvent& InCharacterEvent) override;
+	SLATE_API virtual bool OnVisualizeTooltip(const TSharedPtr<SWidget>& TooltipContent) override;
 
 	/**
 	* Starts a search and makes the search box visible
 	*
 	* @param	InChar		The first character typed in
 	*/
-	void BeginSearch(const TCHAR InChar);
+	SLATE_API void BeginSearch(const TCHAR InChar);
 
 	/**
 	* Resets the search to be empty
 	*/
-	void ResetSearch();
+	SLATE_API void ResetSearch();
 
 	/**
 	* Changes visibility of widgets in the multibox
 	*/
-	void FilterMultiBoxEntries();
+	SLATE_API void FilterMultiBoxEntries();
 
 	/**
 	* Get the text to search by
 	*
 	* @return	The text currently being searched
 	*/
-	FText GetSearchText() const;
+	SLATE_API FText GetSearchText() const;
 
 	/**
 	* Get the SSearchText widget holding the search text
 	*
 	* @return	The widget to get
 	*/
-	TSharedPtr<SWidget> GetSearchTextWidget();
+	SLATE_API TSharedPtr<SWidget> GetSearchTextWidget();
 
 	/**
 	* Set the block widget holding the search text
 	*
 	* @param	BlockWidget		The widget to set
 	*/
-	void SetSearchBlockWidget(TSharedPtr<SWidget>);
+	SLATE_API void SetSearchBlockWidget(TSharedPtr<SWidget>);
 
 	/**
 	* Adds a widget to SearchElement map, making the widget searchable
@@ -774,7 +774,7 @@ public:
 	* @param	BlockDisplayText	The display text of the widget to search by
 	*/
 	UE_DEPRECATED(4.26, "AddSearchElement is deprecated as non-searchable elements also need to be stored, use AddElement instead")
-	void AddSearchElement( TSharedPtr<SWidget>, FText );
+	SLATE_API void AddSearchElement( TSharedPtr<SWidget>, FText );
 
 	/**
 	* Adds a widget to MultiBoxWidgets map, to access and modify its visibility based on search filters
@@ -783,47 +783,47 @@ public:
 	* @param	BlockDisplayText	The display text of the widget to search by
 	  @param    bSearchable			Whether the text is searchable
 	*/
-	void AddElement(TSharedPtr<SWidget> BlockWidget, FText BlockDisplayText, bool bInSearchable = true);
+	SLATE_API void AddElement(TSharedPtr<SWidget> BlockWidget, FText BlockDisplayText, bool bInSearchable = true);
 
 	/**
 	 * @return True if the passed in block is being dragged
 	 */
-	bool IsBlockBeingDragged( TSharedPtr<const FMultiBlock> Block ) const;
+	SLATE_API bool IsBlockBeingDragged( TSharedPtr<const FMultiBlock> Block ) const;
 
 	/**
 	 * @return The visibility of customization widgets for a block
 	 */
-	EVisibility GetCustomizationVisibility(TWeakPtr<const FMultiBlock> BlockWeakPtr, TWeakPtr<SWidget> BlockWidgetWeakPtr) const;
+	SLATE_API EVisibility GetCustomizationVisibility(TWeakPtr<const FMultiBlock> BlockWeakPtr, TWeakPtr<SWidget> BlockWidgetWeakPtr) const;
 
 	/**
 	 * @return The visibility of the drop location indicator of a drag and drop for a block
 	 */
-	EVisibility GetCustomizationBorderDragVisibility(const FName InBlockName, const EMultiBlockType InBlockType, bool& bOutInsertAfter) const;
+	SLATE_API EVisibility GetCustomizationBorderDragVisibility(const FName InBlockName, const EMultiBlockType InBlockType, bool& bOutInsertAfter) const;
 
 	/**
 	 * Records the time that the multibox last summoned a menu
 	 */
-	void SetSummonedMenuTime(double InSummonedMenuTime);
+	SLATE_API void SetSummonedMenuTime(double InSummonedMenuTime);
 
 	/**
 	 * @return The last recorded time that the multibox summoned a menu
 	 */
-	double GetSummonedMenuTime() const;
+	SLATE_API double GetSummonedMenuTime() const;
 
 private:
 	/** Adds a block Widget to this widget */
-	void AddBlockWidget(const FMultiBlock& Block, TSharedPtr<SHorizontalBox> HorizontalBox, TSharedPtr<SVerticalBox> VerticalBox, EMultiBlockLocation::Type InLocation, bool bSectionContainsIcons, TSharedPtr<const FToolBarComboButtonBlock> OptionsBlock);
+	SLATE_API void AddBlockWidget(const FMultiBlock& Block, TSharedPtr<SHorizontalBox> HorizontalBox, TSharedPtr<SVerticalBox> VerticalBox, EMultiBlockLocation::Type InLocation, bool bSectionContainsIcons, TSharedPtr<const FToolBarComboButtonBlock> OptionsBlock);
 
 	/**
 	 * Updates the preview block being dragged.  The drag area is where the users dragged block will be dropped
 	 */
-	void UpdateDropAreaPreviewBlock( TSharedRef<const FMultiBlock> MultiBlock, TSharedPtr<FUICommandDragDropOp> DragDropContent, const FGeometry& DragArea, const FVector2D& DragPos );
+	SLATE_API void UpdateDropAreaPreviewBlock( TSharedRef<const FMultiBlock> MultiBlock, TSharedPtr<FUICommandDragDropOp> DragDropContent, const FGeometry& DragArea, const FVector2D& DragPos );
 
 	/** Creates the SearchTextWidget if the MultiBox has requested one */
-	void CreateSearchTextWidget();
+	SLATE_API void CreateSearchTextWidget();
 
 	/** Called when the SearchText changes */
-	void OnFilterTextChanged(const FText& InFilterText);
+	SLATE_API void OnFilterTextChanged(const FText& InFilterText);
 
 	/**
 	 * Walks the sub-menus and adds new searchable blocks representing the flattened structure of any nested sub-menus.
@@ -831,7 +831,7 @@ private:
 	 * 
 	 * @param MaxRecursionLevels Maximum depth of nested sub-menus to walk when flattening
 	 */
-	void FlattenSubMenusRecursive(uint32 MaxRecursionLevels);
+	SLATE_API void FlattenSubMenusRecursive(uint32 MaxRecursionLevels);
 
 private:
 	/** A preview of a block being dragged */
diff --git a/Engine/Source/Runtime/Slate/Public/Framework/MultiBox/MultiBoxBuilder.h b/Engine/Source/Runtime/Slate/Public/Framework/MultiBox/MultiBoxBuilder.h
--- a/Engine/Source/Runtime/Slate/Public/Framework/MultiBox/MultiBoxBuilder.h
+++ b/Engine/Source/Runtime/Slate/Public/Framework/MultiBox/MultiBoxBuilder.h
@@ -30,7 +30,7 @@ DECLARE_DELEGATE_OneParam( FNewMenuDelegate, class FMenuBuilder& );
 /**
  * MultiBox builder
  */
-class SLATE_API FMultiBoxBuilder
+class FMultiBoxBuilder
 {
 
 public:
@@ -43,7 +43,7 @@ public:
 	 * @param	InCommandList	The action list that maps command infos to delegates that should be called for each command associated with a multiblock widget.  This can be modified after the MultiBox is created by calling the PushCommandList() and PopCommandList() methods.
 	 * @param	InTutorialHighlightName	Optional name to identify this widget and highlight during tutorials
 	 */
-	FMultiBoxBuilder( const EMultiBoxType InType, FMultiBoxCustomization InCustomization, const bool bInShouldCloseWindowAfterMenuSelection, const TSharedPtr< const FUICommandList >& InCommandList, TSharedPtr<FExtender> InExtender = TSharedPtr<FExtender>(), FName InTutorialHighlightName = NAME_None, FName InMenuName = NAME_None );
+	SLATE_API FMultiBoxBuilder( const EMultiBoxType InType, FMultiBoxCustomization InCustomization, const bool bInShouldCloseWindowAfterMenuSelection, const TSharedPtr< const FUICommandList >& InCommandList, TSharedPtr<FExtender> InExtender = TSharedPtr<FExtender>(), FName InTutorialHighlightName = NAME_None, FName InMenuName = NAME_None );
 
 	virtual ~FMultiBoxBuilder() {}
 
@@ -58,7 +58,7 @@ public:
 	 * @param	InOnTextChanged		Called when the text is changed interactively
 	 * @param	bInReadOnly			Whether or not the text block is read only
 	 */
-	void AddEditableText( const FText& InLabel, const FText& InToolTip, const FSlateIcon& InIcon, const TAttribute< FText >& InTextAttribute, const FOnTextCommitted& InOnTextCommitted = FOnTextCommitted(), const FOnTextChanged& InOnTextChanged = FOnTextChanged(), bool bInReadOnly = false );
+	SLATE_API void AddEditableText( const FText& InLabel, const FText& InToolTip, const FSlateIcon& InIcon, const TAttribute< FText >& InTextAttribute, const FOnTextCommitted& InOnTextCommitted = FOnTextCommitted(), const FOnTextChanged& InOnTextChanged = FOnTextChanged(), bool bInReadOnly = false );
 
 	/**
 	 * Adds an editable text entry with a VerifyTextChanged delegate
@@ -72,14 +72,14 @@ public:
 	 * @param	InOnTextChanged			Called when the text is changed interactively
 	 * @param	bInReadOnly				Whether or not the text block is read only
 	 */
-	void AddVerifiedEditableText(const FText& InLabel, const FText& InToolTip, const FSlateIcon& InIcon, const TAttribute< FText >& InTextAttribute, const FOnVerifyTextChanged& InOnVerifyTextChanged, const FOnTextCommitted& InOnTextCommitted = FOnTextCommitted(), const FOnTextChanged& InOnTextChanged = FOnTextChanged(), bool bInReadOnly = false);
+	SLATE_API void AddVerifiedEditableText(const FText& InLabel, const FText& InToolTip, const FSlateIcon& InIcon, const TAttribute< FText >& InTextAttribute, const FOnVerifyTextChanged& InOnVerifyTextChanged, const FOnTextCommitted& InOnTextCommitted = FOnTextCommitted(), const FOnTextChanged& InOnTextChanged = FOnTextChanged(), bool bInReadOnly = false);
 
 	/**
 	 * Creates a widget for this MultiBox
 	 *
 	 * @return  New widget object
 	 */
-	virtual TSharedRef< class SWidget > MakeWidget( FMultiBox::FOnMakeMultiBoxBuilderOverride* InMakeMultiBoxBuilderOverride = nullptr, uint32 MaxHeight = INT_MAX);
+	SLATE_API virtual TSharedRef< class SWidget > MakeWidget( FMultiBox::FOnMakeMultiBoxBuilderOverride* InMakeMultiBoxBuilderOverride = nullptr, uint32 MaxHeight = INT_MAX);
 	
 
 	/** 
@@ -87,7 +87,7 @@ public:
 	 *
 	 * @return The multi-box being built.
 	 */
-	TSharedRef< class FMultiBox > GetMultiBox();
+	SLATE_API TSharedRef< class FMultiBox > GetMultiBox();
 
 
 	/**
@@ -95,46 +95,46 @@ public:
 	 *
 	 * @param	CommandList		The new command list to use
 	 */
-	void PushCommandList( const TSharedRef< const FUICommandList > CommandList );
+	SLATE_API void PushCommandList( const TSharedRef< const FUICommandList > CommandList );
 
 
 	/**
 	 * Pops the current command list.
 	 */
-	void PopCommandList();
+	SLATE_API void PopCommandList();
 	
 	/**
 	 * @return The top command list
 	 */
-	TSharedPtr<const FUICommandList> GetTopCommandList();
+	SLATE_API TSharedPtr<const FUICommandList> GetTopCommandList();
 
 	/**
 	 * Pushes a new extender onto the stack. This extender will be used for all subsequently-added multiblocks, until the extender is popped.
 	 *
 	 * @param	InExtender	The new extender to use
 	 */
-	void PushExtender( TSharedRef< FExtender > InExtender );
+	SLATE_API void PushExtender( TSharedRef< FExtender > InExtender );
 
 
 	/**
 	 * Pops the current extender.
 	 */
-	void PopExtender();
+	SLATE_API void PopExtender();
 
 	/** @return The style set used by the multibox widgets */
-	const ISlateStyle* GetStyleSet() const;
+	SLATE_API const ISlateStyle* GetStyleSet() const;
 
 	/** @return The style name used by the multibox widgets */
-	const FName& GetStyleName() const;
+	SLATE_API const FName& GetStyleName() const;
 
 	/** the override for the checkbox style */
-	void SetCheckBoxStyle(FName InCheckBoxStyle);
+	SLATE_API void SetCheckBoxStyle(FName InCheckBoxStyle);
 
 	/** Sets the style to use on the multibox widgets */
-	void SetStyle( const ISlateStyle* InStyleSet, const FName& InStyleName );
+	SLATE_API void SetStyle( const ISlateStyle* InStyleSet, const FName& InStyleName );
 
 	/** @return  The customization settings for the box being built */
-	FMultiBoxCustomization GetCustomization() const;
+	SLATE_API FMultiBoxCustomization GetCustomization() const;
 
 	/** Sets extender support */
 	void SetExtendersEnabled(bool bEnabled) { bExtendersEnabled = bEnabled; }
@@ -227,7 +227,7 @@ struct FMenuEntryParams : public FMultiBlock::FMultiBlockParams
 /**
  * Base menu builder
  */
-class SLATE_API FBaseMenuBuilder : public FMultiBoxBuilder
+class FBaseMenuBuilder : public FMultiBoxBuilder
 {
 
 public:
@@ -241,7 +241,7 @@ public:
 	 * @param	bInCloseSelfOnly	True if clicking on a menu entry closes itself only and its children but not the entire stack 
 	 * @param	InTutorialHighlightName	Optional name to identify this widget and highlight during tutorials
 	 */
-	FBaseMenuBuilder( const EMultiBoxType InType, const bool bInShouldCloseWindowAfterMenuSelection, TSharedPtr< const FUICommandList > InCommandList, bool bInCloseSelfOnly, TSharedPtr<FExtender> InExtender = TSharedPtr<FExtender>(), const ISlateStyle* InStyleSet = &FCoreStyle::Get(), FName InTutorialHighlightName = NAME_None, FName InMenuName = NAME_None );
+	SLATE_API FBaseMenuBuilder( const EMultiBoxType InType, const bool bInShouldCloseWindowAfterMenuSelection, TSharedPtr< const FUICommandList > InCommandList, bool bInCloseSelfOnly, TSharedPtr<FExtender> InExtender = TSharedPtr<FExtender>(), const ISlateStyle* InStyleSet = &FCoreStyle::Get(), FName InTutorialHighlightName = NAME_None, FName InMenuName = NAME_None );
 
 	/**
 	 * Adds a menu entry
@@ -253,7 +253,7 @@ public:
 	 * @param	InIconOverride		Optional name of the slate brush to use for the tool bar image.  If omitted, then the command's icon will be used instead.
 	 * @param	InTutorialHighlightName	Optional name to identify this widget and highlight during tutorials
 	 */
-	void AddMenuEntry( const TSharedPtr< const FUICommandInfo > InCommand, FName InExtensionHook = NAME_None, const TAttribute<FText>& InLabelOverride = TAttribute<FText>(), const TAttribute<FText>& InToolTipOverride = TAttribute<FText>(), const FSlateIcon& InIconOverride = FSlateIcon(), FName InTutorialHighlightName = NAME_None );
+	SLATE_API void AddMenuEntry( const TSharedPtr< const FUICommandInfo > InCommand, FName InExtensionHook = NAME_None, const TAttribute<FText>& InLabelOverride = TAttribute<FText>(), const TAttribute<FText>& InToolTipOverride = TAttribute<FText>(), const FSlateIcon& InIconOverride = FSlateIcon(), FName InTutorialHighlightName = NAME_None );
 
 	/**
 	 * Adds a menu entry without the use of a command
@@ -266,7 +266,7 @@ public:
 	 * @param	UserInterfaceActionType	Type of interface action
 	 * @param	InTutorialHighlightName	Optional name to identify this widget and highlight during tutorials
 	 */
-	void AddMenuEntry( const TAttribute<FText>& InLabel, const TAttribute<FText>& InToolTip, const FSlateIcon& InIcon, const FUIAction& UIAction, FName InExtensionHook = NAME_None, const EUserInterfaceActionType UserInterfaceActionType = EUserInterfaceActionType::Button, FName InTutorialHighlightName = NAME_None );
+	SLATE_API void AddMenuEntry( const TAttribute<FText>& InLabel, const TAttribute<FText>& InToolTip, const FSlateIcon& InIcon, const FUIAction& UIAction, FName InExtensionHook = NAME_None, const EUserInterfaceActionType UserInterfaceActionType = EUserInterfaceActionType::Button, FName InTutorialHighlightName = NAME_None );
 	
 	/**
 	 * Adds a menu entry with a custom widget
@@ -278,10 +278,10 @@ public:
 	 * @param	UserInterfaceActionType	Type of interface action
 	 * @param	InTutorialHighlightName	Optional name to identify this widget and highlight during tutorials
 	 */
-	void AddMenuEntry( const FUIAction& UIAction, const TSharedRef< SWidget > Contents, const FName& InExtensionHook = NAME_None, const TAttribute<FText>& InToolTip = TAttribute<FText>(), const EUserInterfaceActionType UserInterfaceActionType = EUserInterfaceActionType::Button, FName InTutorialHighlightName = NAME_None );
+	SLATE_API void AddMenuEntry( const FUIAction& UIAction, const TSharedRef< SWidget > Contents, const FName& InExtensionHook = NAME_None, const TAttribute<FText>& InToolTip = TAttribute<FText>(), const EUserInterfaceActionType UserInterfaceActionType = EUserInterfaceActionType::Button, FName InTutorialHighlightName = NAME_None );
 
 	/** Adds a menu entry with given param struct */
-	void AddMenuEntry( const FMenuEntryParams& InMenuEntryParams );
+	SLATE_API void AddMenuEntry( const FMenuEntryParams& InMenuEntryParams );
 
 protected:
 	/** True if clicking on a menu entry closes itself only and its children and not the entire stack */
@@ -292,7 +292,7 @@ protected:
 /**
  * Vertical menu builder
  */
-class SLATE_API FMenuBuilder : public FBaseMenuBuilder
+class FMenuBuilder : public FBaseMenuBuilder
 {
 public:
 
@@ -326,13 +326,13 @@ public:
 	*
 	* @return  New widget object
 	*/
-	virtual TSharedRef< class SWidget > MakeWidget( FMultiBox::FOnMakeMultiBoxBuilderOverride* InMakeMultiBoxBuilderOverride = nullptr, uint32 MaxHeight = 1000) override;
+	SLATE_API virtual TSharedRef< class SWidget > MakeWidget( FMultiBox::FOnMakeMultiBoxBuilderOverride* InMakeMultiBoxBuilderOverride = nullptr, uint32 MaxHeight = 1000) override;
 
 	/**
 	 * Adds a separator
 	 */
-	void AddMenuSeparator(FName InExtensionHook = NAME_None);
-	void AddSeparator(FName InExtensionHook = NAME_None);
+	SLATE_API void AddMenuSeparator(FName InExtensionHook = NAME_None);
+	SLATE_API void AddSeparator(FName InExtensionHook = NAME_None);
 	
 	/**
 	 * Starts a section on to the extender section hook stack
@@ -340,12 +340,12 @@ public:
 	 * @param InExtensionHook	The section hook. Can be NAME_None
 	 * @param InHeadingText		The heading text to use. If none, only a separator is used
 	 */
-	void BeginSection( FName InExtensionHook, const TAttribute< FText >& InHeadingText = TAttribute<FText>() );
+	SLATE_API void BeginSection( FName InExtensionHook, const TAttribute< FText >& InHeadingText = TAttribute<FText>() );
 
 	/**
 	 * Ends the current section
 	 */
-	void EndSection();
+	SLATE_API void EndSection();
 
 
 	/**
@@ -361,13 +361,13 @@ public:
 	 * @param	InIcon				The icon to use
 	 * @param	bInShouldCloseWindowAfterMenuSelection	Whether the submenu should close after an item is selected
 	 */
-	void AddSubMenu( const TAttribute<FText>& InMenuLabel, const TAttribute<FText>& InToolTip, const FNewMenuDelegate& InSubMenu, const FUIAction& InUIAction, FName InExtensionHook, const EUserInterfaceActionType InUserInterfaceActionType, const bool bInOpenSubMenuOnClick = false, const FSlateIcon& InIcon = FSlateIcon(), const bool bInShouldCloseWindowAfterMenuSelection = true );
+	SLATE_API void AddSubMenu( const TAttribute<FText>& InMenuLabel, const TAttribute<FText>& InToolTip, const FNewMenuDelegate& InSubMenu, const FUIAction& InUIAction, FName InExtensionHook, const EUserInterfaceActionType InUserInterfaceActionType, const bool bInOpenSubMenuOnClick = false, const FSlateIcon& InIcon = FSlateIcon(), const bool bInShouldCloseWindowAfterMenuSelection = true );
 
-	void AddSubMenu( const TAttribute<FText>& InMenuLabel, const TAttribute<FText>& InToolTip, const FNewMenuDelegate& InSubMenu, const bool bInOpenSubMenuOnClick = false, const FSlateIcon& InIcon = FSlateIcon(), const bool bInShouldCloseWindowAfterMenuSelection = true, FName InExtensionHook = NAME_None );
+	SLATE_API void AddSubMenu( const TAttribute<FText>& InMenuLabel, const TAttribute<FText>& InToolTip, const FNewMenuDelegate& InSubMenu, const bool bInOpenSubMenuOnClick = false, const FSlateIcon& InIcon = FSlateIcon(), const bool bInShouldCloseWindowAfterMenuSelection = true, FName InExtensionHook = NAME_None );
 
-	void AddSubMenu( const TSharedRef< SWidget > Contents, const FNewMenuDelegate& InSubMenu, const bool bInOpenSubMenuOnClick = false, const bool bInShouldCloseWindowAfterMenuSelection = true );
+	SLATE_API void AddSubMenu( const TSharedRef< SWidget > Contents, const FNewMenuDelegate& InSubMenu, const bool bInOpenSubMenuOnClick = false, const bool bInShouldCloseWindowAfterMenuSelection = true );
 
-	void AddSubMenu( const FUIAction& UIAction, const TSharedRef< SWidget > Contents, const FNewMenuDelegate& InSubMenu, const bool bInShouldCloseWindowAfterMenuSelection = true );
+	SLATE_API void AddSubMenu( const FUIAction& UIAction, const TSharedRef< SWidget > Contents, const FNewMenuDelegate& InSubMenu, const bool bInShouldCloseWindowAfterMenuSelection = true );
 
 	/**
 	 * Adds any widget to the menu
@@ -378,19 +378,19 @@ public:
 	 * @param	bSearchable			If true, widget will be searchable (default == true)
 	 * @param	InToolTipText	Optional tooltip text to be added to the widget and label
 	 */
-	void AddWidget( TSharedRef<SWidget> InWidget, const FText& Label, bool bNoIndent = false, bool bSearchable = true, const TAttribute<FText>& InToolTipText = FText());
+	SLATE_API void AddWidget( TSharedRef<SWidget> InWidget, const FText& Label, bool bNoIndent = false, bool bSearchable = true, const TAttribute<FText>& InToolTipText = FText());
 
 	/**
 	* Adds the widget the multibox will use for searching
 	*/
-	void AddSearchWidget();
+	SLATE_API void AddSearchWidget();
 
 	void SetIsEditing(bool bInIsEditing) { bIsEditing = bInIsEditing; }
 
 protected:
 	/** FMultiBoxBuilder interface */
-	virtual void ApplyHook(FName InExtensionHook, EExtensionHook::Position HookPosition) override;
-	virtual void ApplySectionBeginning() override;
+	SLATE_API virtual void ApplyHook(FName InExtensionHook, EExtensionHook::Position HookPosition) override;
+	SLATE_API virtual void ApplySectionBeginning() override;
 
 public:
 	// These classes need access to the AddWrapperSubMenu() methods
@@ -404,9 +404,9 @@ public:
 	 * @param	InToolTip			The tooltip that should be shown when the menu is hovered over
 	 * @param	InSubMenu			Sub-Menu object which creates the sub-menu
 	 */
-	void AddWrapperSubMenu( const FText& InMenuLabel, const FText& InToolTip, const FOnGetContent& InSubMenu, const FSlateIcon& InIcon );
+	SLATE_API void AddWrapperSubMenu( const FText& InMenuLabel, const FText& InToolTip, const FOnGetContent& InSubMenu, const FSlateIcon& InIcon );
 
-	void AddWrapperSubMenu( const FText& InMenuLabel, const FText& InToolTip, const FOnGetContent& InSubMenu, const FSlateIcon& InIcon, const FUIAction& UIAction );
+	SLATE_API void AddWrapperSubMenu( const FText& InMenuLabel, const FText& InToolTip, const FOnGetContent& InSubMenu, const FSlateIcon& InIcon, const FUIAction& UIAction );
 
 	/**
 	 * Adds a sub-menu which is a menu within a menu
@@ -415,7 +415,7 @@ public:
 	 * @param	InToolTip			The tooltip that should be shown when the menu is hovered over
 	 * @param	InSubMenu			Sub-Menu object
 	 */
-	void AddWrapperSubMenu( const FText& InMenuLabel, const FText& InToolTip, const TSharedPtr<SWidget>& InSubMenu, const FSlateIcon& InIcon );
+	SLATE_API void AddWrapperSubMenu( const FText& InMenuLabel, const FText& InToolTip, const TSharedPtr<SWidget>& InSubMenu, const FSlateIcon& InIcon );
 
 private:
 	/** Current extension hook name for sections to determine where sections begin and end */
@@ -442,7 +442,7 @@ private:
 /**
  * Menu bar builder
  */
-class SLATE_API FMenuBarBuilder : public FBaseMenuBuilder
+class FMenuBarBuilder : public FBaseMenuBuilder
 {
 
 public:
@@ -466,7 +466,7 @@ public:
 	 * @param	InToolTip			The tooltip that should be shown when the menu is hovered over
 	 * @param	InPullDownMenu		Pull down menu object for the menu to add.
 	 */
-	void AddPullDownMenu( const TAttribute<FText>& InMenuLabel, const TAttribute<FText>& InToolTip, const FNewMenuDelegate& InPullDownMenu, FName InExtensionHook = NAME_None, FName InTutorialHighlightName = NAME_None);
+	SLATE_API void AddPullDownMenu( const TAttribute<FText>& InMenuLabel, const TAttribute<FText>& InToolTip, const FNewMenuDelegate& InPullDownMenu, FName InExtensionHook = NAME_None, FName InTutorialHighlightName = NAME_None);
 
 	/**
 	 * Adds a pull down menu
@@ -475,11 +475,11 @@ public:
 	 * @param	InToolTip				The tooltip that should be shown when the menu is hovered over
 	 * @param	InMenuContentGenerator	Delegate that generates a widget for this pulldown menu's content.  Called when the menu is summoned.
 	 */
-	void AddPullDownMenu(const TAttribute<FText>& InMenuLabel, const TAttribute<FText>& InToolTip, const FOnGetContent& InMenuContentGenerator, FName InExtensionHook = NAME_None, FName InTutorialHighlightName = NAME_None);
+	SLATE_API void AddPullDownMenu(const TAttribute<FText>& InMenuLabel, const TAttribute<FText>& InToolTip, const FOnGetContent& InMenuContentGenerator, FName InExtensionHook = NAME_None, FName InTutorialHighlightName = NAME_None);
 
 protected:
 	/** FMultiBoxBuilder interface */
-	virtual void ApplyHook(FName InExtensionHook, EExtensionHook::Position HookPosition) override;
+	SLATE_API virtual void ApplyHook(FName InExtensionHook, EExtensionHook::Position HookPosition) override;
 };
 
 
@@ -487,7 +487,7 @@ protected:
 /**
  * Tool bar builder
  */
-class SLATE_API FToolBarBuilder : public FMultiBoxBuilder
+class FToolBarBuilder : public FMultiBoxBuilder
 {
 	friend class UToolMenus;
 public:
@@ -517,14 +517,14 @@ public:
 
 	void SetLabelVisibility( EVisibility InLabelVisibility ) { LabelVisibility  = InLabelVisibility ; }
 
-	void SetIsFocusable(bool bInIsFocusable);
+	SLATE_API void SetIsFocusable(bool bInIsFocusable);
 
 	/**
 	 * Adds a tool bar button
 	 *
 	 * @param	ButtonArgs The Parameters object which will provide the data to initialize the button
 	 */
-	virtual void AddToolBarButton(FButtonArgs& ButtonArgs);
+	SLATE_API virtual void AddToolBarButton(FButtonArgs& ButtonArgs);
 
 
 	/**
@@ -537,7 +537,7 @@ public:
 	 * @param	InTutorialHighlightName	Name to identify this widget and highlight during tutorials
 	 * @param	CustomMenuDelegate  Optional custom menu delegate for cases where the toolbar is compressed into a menu
 	 */
-	void AddToolBarButton(const TSharedPtr< const FUICommandInfo > InCommand, FName InExtensionHook = NAME_None, const TAttribute<FText>& InLabelOverride = TAttribute<FText>(), const TAttribute<FText>& InToolTipOverride = TAttribute<FText>(), const TAttribute<FSlateIcon>& InIconOverride = TAttribute<FSlateIcon>(), FName InTutorialHighlightName = NAME_None, FNewMenuDelegate InCustomMenuDelegate = FNewMenuDelegate());
+	SLATE_API void AddToolBarButton(const TSharedPtr< const FUICommandInfo > InCommand, FName InExtensionHook = NAME_None, const TAttribute<FText>& InLabelOverride = TAttribute<FText>(), const TAttribute<FText>& InToolTipOverride = TAttribute<FText>(), const TAttribute<FSlateIcon>& InIconOverride = TAttribute<FSlateIcon>(), FName InTutorialHighlightName = NAME_None, FNewMenuDelegate InCustomMenuDelegate = FNewMenuDelegate());
 	
 	/**
 	 * Adds a tool bar button
@@ -549,7 +549,7 @@ public:
 	 * @param	UserInterfaceActionType	Type of interface action
 	 * @param	InTutorialHighlightName	Name to identify this widget and highlight during tutorials
 	 */
-	void AddToolBarButton(const FUIAction& InAction, FName InExtensionHook = NAME_None, const TAttribute<FText>& InLabelOverride = TAttribute<FText>(), const TAttribute<FText>& InToolTipOverride = TAttribute<FText>(), const TAttribute<FSlateIcon>& InIconOverride = TAttribute<FSlateIcon>(), const EUserInterfaceActionType UserInterfaceActionType = EUserInterfaceActionType::Button, FName InTutorialHighlightName = NAME_None );
+	SLATE_API void AddToolBarButton(const FUIAction& InAction, FName InExtensionHook = NAME_None, const TAttribute<FText>& InLabelOverride = TAttribute<FText>(), const TAttribute<FText>& InToolTipOverride = TAttribute<FText>(), const TAttribute<FSlateIcon>& InIconOverride = TAttribute<FSlateIcon>(), const EUserInterfaceActionType UserInterfaceActionType = EUserInterfaceActionType::Button, FName InTutorialHighlightName = NAME_None );
 
 	/**
 	 * Adds a combo button
@@ -562,7 +562,7 @@ public:
 	 * @param	bInSimpleComboBox			If true, the icon and label won't be displayed
 	 * @param	InTutorialHighlightName		Name to identify this widget and highlight during tutorials
 	 */
-	void AddComboButton( const FUIAction& InAction, const FOnGetContent& InMenuContentGenerator, const TAttribute<FText>& InLabelOverride = TAttribute<FText>(), const TAttribute<FText>& InToolTipOverride = TAttribute<FText>(), const TAttribute<FSlateIcon>& InIconOverride = TAttribute<FSlateIcon>(), bool bInSimpleComboBox = false, FName InTutorialHighlightName = NAME_None );
+	SLATE_API void AddComboButton( const FUIAction& InAction, const FOnGetContent& InMenuContentGenerator, const TAttribute<FText>& InLabelOverride = TAttribute<FText>(), const TAttribute<FText>& InToolTipOverride = TAttribute<FText>(), const TAttribute<FSlateIcon>& InIconOverride = TAttribute<FSlateIcon>(), bool bInSimpleComboBox = false, FName InTutorialHighlightName = NAME_None );
 
 	/**
 	 * Adds a tool bar stack button
@@ -570,7 +570,7 @@ public:
 	 * @param	InCommand				The command associated with this tool bar button
 	 * @param	InTutorialHighlightName	Name to identify this widget and highlight during tutorials
 	 */
-	void AddToolbarStackButton(const TSharedPtr< const FUICommandInfo > InCommand, FName InTutorialHighlightName = NAME_None);
+	SLATE_API void AddToolbarStackButton(const TSharedPtr< const FUICommandInfo > InCommand, FName InTutorialHighlightName = NAME_None);
 
 	/**
 	 * Adds any widget to the toolbar
@@ -580,7 +580,7 @@ public:
 	 * @param	InTutorialHighlightName	Name to identify this widget and highlight during tutorials
 	 * @param	bSearchable			If true, widget will be searchable (default == true)
 	 */
-	void AddToolBarWidget(TSharedRef<SWidget> InWidget, const TAttribute<FText>& InLabel = TAttribute<FText>(), FName InTutorialHighlightName = NAME_None, bool bSearchable = true);
+	SLATE_API void AddToolBarWidget(TSharedRef<SWidget> InWidget, const TAttribute<FText>& InLabel = TAttribute<FText>(), FName InTutorialHighlightName = NAME_None, bool bSearchable = true);
 
 
 	/**
@@ -592,41 +592,41 @@ public:
 	 * @param	Alignment			Horizontal alignment for the widget inside the toolbar
 	 * @param	CustomMenuDelegate  Optional custom menu delegate for cases where the toolbar is compressed into a menu
 	 */
-	void AddWidget(TSharedRef<SWidget> InWidget, FName InTutorialHighlightName = NAME_None, bool bSearchable = true, EHorizontalAlignment Alignment = HAlign_Fill, FNewMenuDelegate InCustomMenuDelegate = FNewMenuDelegate());
+	SLATE_API void AddWidget(TSharedRef<SWidget> InWidget, FName InTutorialHighlightName = NAME_None, bool bSearchable = true, EHorizontalAlignment Alignment = HAlign_Fill, FNewMenuDelegate InCustomMenuDelegate = FNewMenuDelegate());
 	
 	/**
 	 * Adds a toolbar separator
 	 */
-	void AddSeparator(FName InExtensionHook = NAME_None);
+	SLATE_API void AddSeparator(FName InExtensionHook = NAME_None);
 	
 	/**
 	 * Starts a section on to the extender section hook stack
 	 * 
 	 * @param InExtensionHook	The section hook. Can be NAME_None
 	 */
-	void BeginSection( FName InExtensionHook );
+	SLATE_API void BeginSection( FName InExtensionHook );
 
 	/**
 	 * Ends the current section
 	 */
-	void EndSection();
+	SLATE_API void EndSection();
 
 	/** 
 	 * Starts a new Group block, must be used in conjunction with EndBlockGroup
 	 */
-	void BeginBlockGroup();
+	SLATE_API void BeginBlockGroup();
 	
 	/** 
 	 * End a group block, must be used in conjunction with BeginBlockGroup.
 	 */
-	void EndBlockGroup();
+	SLATE_API void EndBlockGroup();
 
 	/**
 	 * Overrides the style being used by this toolbar with a different one for the
 	 * The override will be used for any added blocks until EndStyleOverride is called
 	 */
-	void BeginStyleOverride(FName StyleOverrideName);
-	void EndStyleOverride();
+	SLATE_API void BeginStyleOverride(FName StyleOverrideName);
+	SLATE_API void EndStyleOverride();
 protected:
 
 	FToolBarBuilder(EMultiBoxType InType, TSharedPtr<const FUICommandList> InCommandList, FMultiBoxCustomization InCustomization, TSharedPtr<FExtender> InExtender = TSharedPtr<FExtender>(), const bool InForceSmallIcons = false)
@@ -639,10 +639,10 @@ protected:
 
 
 	/** FMultiBoxBuilder interface */
-	virtual void ApplyHook(FName InExtensionHook, EExtensionHook::Position HookPosition) override;
-	virtual void ApplySectionBeginning() override;
+	SLATE_API virtual void ApplyHook(FName InExtensionHook, EExtensionHook::Position HookPosition) override;
+	SLATE_API virtual void ApplySectionBeginning() override;
 
-	void InitializeToolBarButtonBlock(TSharedPtr<FToolBarButtonBlock> ButtonRowBlock, FButtonArgs& ButtonArgs);
+	SLATE_API void InitializeToolBarButtonBlock(TSharedPtr<FToolBarButtonBlock> ButtonRowBlock, FButtonArgs& ButtonArgs);
 
 private:
 	/** Current extension hook name for sections to determine where sections begin and end */
@@ -663,7 +663,7 @@ private:
 };
 
 
-class SLATE_API FVerticalToolBarBuilder : public FToolBarBuilder
+class FVerticalToolBarBuilder : public FToolBarBuilder
 {
 public:
 	/**
@@ -678,7 +678,7 @@ public:
 	}
 };
 
-class SLATE_API FUniformToolBarBuilder : public FToolBarBuilder
+class FUniformToolBarBuilder : public FToolBarBuilder
 {
 public:
 	/**
@@ -692,7 +692,7 @@ public:
 	}
 };
 
-class SLATE_API FSlimHorizontalToolBarBuilder : public FToolBarBuilder
+class FSlimHorizontalToolBarBuilder : public FToolBarBuilder
 {
 public:
 	/**
@@ -709,7 +709,7 @@ public:
 /**
  * Button grid builder
  */
-class SLATE_API FButtonRowBuilder : public FMultiBoxBuilder
+class FButtonRowBuilder : public FMultiBoxBuilder
 {
 public:
 	/**
@@ -731,7 +731,7 @@ public:
 	 * @param	InToolTipOverride		Optional tool tip override.	 If omitted, then the action's label will be used instead.
 	 * @param	InIconOverride			Optional icon to use for the tool bar image.  If omitted, then the action's icon will be used instead.
 	 */
-	void AddButton(const TSharedPtr< const FUICommandInfo > InCommand, const TAttribute<FText>& InLabelOverride = TAttribute<FText>(), const TAttribute<FText>& InToolTipOverride = TAttribute<FText>(), const FSlateIcon& InIconOverride = FSlateIcon());
+	SLATE_API void AddButton(const TSharedPtr< const FUICommandInfo > InCommand, const TAttribute<FText>& InLabelOverride = TAttribute<FText>(), const TAttribute<FText>& InToolTipOverride = TAttribute<FText>(), const FSlateIcon& InIconOverride = FSlateIcon());
 
 	/**
 	 * Adds a button to a row
@@ -742,14 +742,14 @@ public:
 	 * @param	InIcon					The icon for the button
 	 * @param	UserInterfaceActionType	The style of button to display
 	 */
-	void AddButton(const FText& InLabel, const FText& InToolTip, const FUIAction& UIAction, const FSlateIcon& InIcon = FSlateIcon(), const EUserInterfaceActionType UserInterfaceActionType = EUserInterfaceActionType::Button);
+	SLATE_API void AddButton(const FText& InLabel, const FText& InToolTip, const FUIAction& UIAction, const FSlateIcon& InIcon = FSlateIcon(), const EUserInterfaceActionType UserInterfaceActionType = EUserInterfaceActionType::Button);
 
 protected:
 	/** FMultiBoxBuilder interface */
 	virtual void ApplyHook(FName InExtensionHook, EExtensionHook::Position HookPosition) override {}
 };
 
-class SLATE_API FSlimHorizontalUniformToolBarBuilder : public FToolBarBuilder
+class FSlimHorizontalUniformToolBarBuilder : public FToolBarBuilder
 {
 public:
 	/**
@@ -757,8 +757,8 @@ public:
 	 *
 	 * @param	InCommandList	The action list that maps command infos to delegates that should be called for each command associated with a multiblock widget
 	 */
-	FSlimHorizontalUniformToolBarBuilder(TSharedPtr<const FUICommandList> InCommandList, FMultiBoxCustomization InCustomization, TSharedPtr<FExtender> InExtender = nullptr, const bool InForceSmallIcons = false);
+	SLATE_API FSlimHorizontalUniformToolBarBuilder(TSharedPtr<const FUICommandList> InCommandList, FMultiBoxCustomization InCustomization, TSharedPtr<FExtender> InExtender = nullptr, const bool InForceSmallIcons = false);
 
-	void AddToolBarButton( FButtonArgs& ButtonArgs ) override;
+	SLATE_API void AddToolBarButton( FButtonArgs& ButtonArgs ) override;
 
-};
\ No newline at end of file
+};
