[Slate]

diff --git a/Engine/Source/Runtime/Slate/Public/Framework/MultiBox/MultiBoxDefs.h b/Engine/Source/Runtime/Slate/Public/Framework/MultiBox/MultiBoxDefs.h
--- a/Engine/Source/Runtime/Slate/Public/Framework/MultiBox/MultiBoxDefs.h
+++ b/Engine/Source/Runtime/Slate/Public/Framework/MultiBox/MultiBoxDefs.h
@@ -61,27 +61,27 @@ enum class EMultiBlockType : uint8
 };
 
 
-class SLATE_API FMultiBoxSettings
+class FMultiBoxSettings
 {
 public:
 
 	DECLARE_DELEGATE_RetVal_ThreeParams( TSharedRef< SToolTip >, FConstructToolTip, const TAttribute<FText>& /*ToolTipText*/, const TSharedPtr<SWidget>& /*OverrideContent*/, const TSharedPtr<const FUICommandInfo>& /*Action*/ );
 
 	/** Access to whether multiboxes use small icons or default sized icons */
-	static TAttribute<bool> UseSmallToolBarIcons;
-	static TAttribute<bool> DisplayMultiboxHooks;
-	static FConstructToolTip ToolTipConstructor;
+	static SLATE_API TAttribute<bool> UseSmallToolBarIcons;
+	static SLATE_API TAttribute<bool> DisplayMultiboxHooks;
+	static SLATE_API FConstructToolTip ToolTipConstructor;
 
-	FMultiBoxSettings();
+	SLATE_API FMultiBoxSettings();
 
-	static TSharedRef< SToolTip > ConstructDefaultToolTip( const TAttribute<FText>& ToolTipText, const TSharedPtr<SWidget>& OverrideContent, const TSharedPtr<const FUICommandInfo>& Action );
+	static SLATE_API TSharedRef< SToolTip > ConstructDefaultToolTip( const TAttribute<FText>& ToolTipText, const TSharedPtr<SWidget>& OverrideContent, const TSharedPtr<const FUICommandInfo>& Action );
 
-	static void ResetToolTipConstructor();
+	static SLATE_API void ResetToolTipConstructor();
 };
 
-struct SLATE_API FMultiBoxCustomization
+struct FMultiBoxCustomization
 {
-	static const FMultiBoxCustomization None;
+	static SLATE_API const FMultiBoxCustomization None;
 
 	static FMultiBoxCustomization AllowCustomization( FName InCustomizationName )
 	{
diff --git a/Engine/Source/Runtime/Slate/Public/Framework/MultiBox/MultiBoxExtender.h b/Engine/Source/Runtime/Slate/Public/Framework/MultiBox/MultiBoxExtender.h
--- a/Engine/Source/Runtime/Slate/Public/Framework/MultiBox/MultiBoxExtender.h
+++ b/Engine/Source/Runtime/Slate/Public/Framework/MultiBox/MultiBoxExtender.h
@@ -127,7 +127,7 @@ public:
 	static SLATE_API TSharedRef<FExtender> Combine(const TArray< TSharedPtr<FExtender> >& Extenders);
 
 	/** Returns the number of extensions in the extender */
-	SLATE_API int32 NumExtensions() const { return Extensions.Num(); }
+	int32 NumExtensions() const { return Extensions.Num(); }
 
 private:
 
diff --git a/Engine/Source/Runtime/Slate/Public/Framework/MultiBox/SToolBarButtonBlock.h b/Engine/Source/Runtime/Slate/Public/Framework/MultiBox/SToolBarButtonBlock.h
--- a/Engine/Source/Runtime/Slate/Public/Framework/MultiBox/SToolBarButtonBlock.h
+++ b/Engine/Source/Runtime/Slate/Public/Framework/MultiBox/SToolBarButtonBlock.h
@@ -18,7 +18,7 @@ class SBorder;
 /** Delegate used by multi-box to call a user function to populate a new menu.  Used for spawning sub-menus and pull-down menus. */
 DECLARE_DELEGATE_OneParam( FNewMenuDelegate, class FMenuBuilder& );
 
-struct SLATE_API FButtonArgs : public TSharedFromThis<FButtonArgs>
+struct FButtonArgs : public TSharedFromThis<FButtonArgs>
 {
 	TSharedPtr< const FUICommandInfo > Command;
 	TSharedPtr< const FUICommandList > CommandList;
@@ -41,14 +41,14 @@ struct SLATE_API FButtonArgs : public TSharedFromThis<FButtonArgs>
 /**
  * Tool bar button MultiBlock
  */
-class SLATE_API FToolBarButtonBlock
+class FToolBarButtonBlock
 	: public FMultiBlock
 {
 
 public:
 
 	
-	FToolBarButtonBlock( FButtonArgs ButtonArgs );
+	SLATE_API FToolBarButtonBlock( FButtonArgs ButtonArgs );
 	
 	/**
 	 * Constructor
@@ -59,7 +59,7 @@ public:
 	 * @param	InToolTipOverride	Optional tool tip override.	 If omitted, then the action's label will be used instead.
 	 * @param	InIconOverride		Optional icon to use for the tool bar image.  If omitted, then the action's icon will be used instead.
 	 */
-	FToolBarButtonBlock( const TSharedPtr< const FUICommandInfo > InCommand, TSharedPtr< const FUICommandList > InCommandList, const TAttribute<FText>& InLabelOverride = TAttribute<FText>(), const TAttribute<FText>& InToolTipOverride = TAttribute<FText>(), const TAttribute<FSlateIcon>& InIconOverride = TAttribute<FSlateIcon>() );
+	SLATE_API FToolBarButtonBlock( const TSharedPtr< const FUICommandInfo > InCommand, TSharedPtr< const FUICommandList > InCommandList, const TAttribute<FText>& InLabelOverride = TAttribute<FText>(), const TAttribute<FText>& InToolTipOverride = TAttribute<FText>(), const TAttribute<FSlateIcon>& InIconOverride = TAttribute<FSlateIcon>() );
 	
 	/**
 	 * Constructor
@@ -70,7 +70,7 @@ public:
 	 * @param	InUIAction			UI action to take when this menu item is clicked as well as to determine if the menu entry can be executed or appears "checked"
 	 * @param	InUserInterfaceActionType	Type of interface action
 	 */
-	FToolBarButtonBlock( const TAttribute<FText>& InLabel, const TAttribute<FText>& InToolTip, const TAttribute<FSlateIcon>& InIcon, const FUIAction& InUIAction, const EUserInterfaceActionType InUserInterfaceActionType );
+	SLATE_API FToolBarButtonBlock( const TAttribute<FText>& InLabel, const TAttribute<FText>& InToolTip, const TAttribute<FSlateIcon>& InIcon, const FUIAction& InUIAction, const EUserInterfaceActionType InUserInterfaceActionType );
 
 	void SetLabelVisibility( EVisibility InLabelVisibility ) { LabelVisibility = InLabelVisibility ; }
 
@@ -83,22 +83,22 @@ public:
 	void SetBorderBrushName( const FName InBorderBrushName ) { BorderBrushName = InBorderBrushName; }
 
 	/** FMultiBlock interface */
-	virtual void CreateMenuEntry(class FMenuBuilder& MenuBuilder) const override;
-	virtual bool HasIcon() const override;
+	SLATE_API virtual void CreateMenuEntry(class FMenuBuilder& MenuBuilder) const override;
+	SLATE_API virtual bool HasIcon() const override;
 
 	/** Set optional delegate to customize when a menu appears instead of the widget, such as in toolbars */
-	void SetCustomMenuDelegate( FNewMenuDelegate& InOnFillMenuDelegate);
+	SLATE_API void SetCustomMenuDelegate( FNewMenuDelegate& InOnFillMenuDelegate);
 
-	void SetOnGetMenuContent( FOnGetContent& OnGetMenuContent);
+	SLATE_API void SetOnGetMenuContent( FOnGetContent& OnGetMenuContent);
 
 protected:
 	
-	bool GetIsFocusable() const;
+	SLATE_API bool GetIsFocusable() const;
 
 private:
 
 	/** FMultiBlock private interface */
-	virtual TSharedRef< class IMultiBlockBaseWidget > ConstructWidget() const override;
+	SLATE_API virtual TSharedRef< class IMultiBlockBaseWidget > ConstructWidget() const override;
 
 private:
 
@@ -141,7 +141,7 @@ private:
 /**
  * Tool bar button MultiBlock widget
  */
-class SLATE_API SToolBarButtonBlock
+class SToolBarButtonBlock
 	: public SMultiBlockBaseWidget
 {
 
@@ -163,7 +163,7 @@ public:
 	/**
 	 * Builds this MultiBlock widget up from the MultiBlock associated with it
 	 */
-	virtual void BuildMultiBlockWidget(const ISlateStyle* StyleSet, const FName& StyleName) override;
+	SLATE_API virtual void BuildMultiBlockWidget(const ISlateStyle* StyleSet, const FName& StyleName) override;
 
 
 	/**
@@ -171,40 +171,40 @@ public:
 	 *
 	 * @param	InArgs	The declaration data for this widget
 	 */
-	void Construct( const FArguments& InArgs );
+	SLATE_API void Construct( const FArguments& InArgs );
 
 protected:
 
 	/**
 	 * Called by Slate when this tool bar button's button is clicked
 	 */
-	FReply OnClicked();
+	SLATE_API FReply OnClicked();
 	
 	/**
 	 * Called by Slate when this tool bar check box button is toggled
 	 */
-	void OnCheckStateChanged( const ECheckBoxState NewCheckedState );
+	SLATE_API void OnCheckStateChanged( const ECheckBoxState NewCheckedState );
 
 	/**
 	 * Called by slate to determine if this button should appear checked
 	 *
 	 * @return ECheckBoxState::Checked if it should be checked, ECheckBoxState::Unchecked if not.
 	 */
-	ECheckBoxState GetCheckState() const;
+	SLATE_API ECheckBoxState GetCheckState() const;
 
 	/**
 	 * Called by Slate to determine if this button is enabled
 	 * 
 	 * @return True if the menu entry is enabled, false otherwise
 	 */
-	bool IsEnabled() const;
+	SLATE_API bool IsEnabled() const;
 
 	/**
 	 * Called by Slate to determine if this button is visible
 	 *
 	 * @return EVisibility::Visible or EVisibility::Collapsed, depending on if the button should be displayed
 	 */
-	EVisibility GetBlockVisibility() const;
+	SLATE_API EVisibility GetBlockVisibility() const;
 
 private:
 	/** Called by Slate to determine whether icons/labels are visible */
diff --git a/Engine/Source/Runtime/Slate/Public/Framework/MultiBox/SToolBarComboButtonBlock.h b/Engine/Source/Runtime/Slate/Public/Framework/MultiBox/SToolBarComboButtonBlock.h
--- a/Engine/Source/Runtime/Slate/Public/Framework/MultiBox/SToolBarComboButtonBlock.h
+++ b/Engine/Source/Runtime/Slate/Public/Framework/MultiBox/SToolBarComboButtonBlock.h
@@ -90,7 +90,7 @@ private:
 /**
  * Tool bar button MultiBlock widget
  */
-class SLATE_API SToolBarComboButtonBlock
+class SToolBarComboButtonBlock
 	: public SMultiBlockBaseWidget
 {
 
@@ -115,7 +115,7 @@ public:
 	/**
 	 * Builds this MultiBlock widget up from the MultiBlock associated with it
 	 */
-	virtual void BuildMultiBlockWidget(const ISlateStyle* StyleSet, const FName& StyleName) override;
+	SLATE_API virtual void BuildMultiBlockWidget(const ISlateStyle* StyleSet, const FName& StyleName) override;
 
 
 	/**
@@ -123,7 +123,7 @@ public:
 	 *
 	 * @param	InArgs	The declaration data for this widget
 	 */
-	void Construct( const FArguments& InArgs );
+	SLATE_API void Construct( const FArguments& InArgs );
 
 protected:
 
@@ -132,7 +132,7 @@ protected:
 	 *
 	 * @return	The widget to use for the menu content
 	 */
-	TSharedRef<SWidget> OnGetMenuContent();
+	SLATE_API TSharedRef<SWidget> OnGetMenuContent();
 
 
 	/**
@@ -140,14 +140,14 @@ protected:
 	 * 
 	 * @return True if the menu entry is enabled, false otherwise
 	 */
-	bool IsEnabled() const;
+	SLATE_API bool IsEnabled() const;
 
 	/**
 	 * Called by Slate to determine if this button is visible
 	 *
 	 * @return EVisibility::Visible or EVisibility::Collapsed, depending on if the button should be displayed
 	 */
-	EVisibility GetVisibility() const;
+	SLATE_API EVisibility GetVisibility() const;
 private:
 	/** @return True if this toolbar button is using a dynamically set icon */
 	bool HasDynamicIcon() const;
diff --git a/Engine/Source/Runtime/Slate/Public/Framework/MultiBox/SToolBarStackButtonBlock.h b/Engine/Source/Runtime/Slate/Public/Framework/MultiBox/SToolBarStackButtonBlock.h
--- a/Engine/Source/Runtime/Slate/Public/Framework/MultiBox/SToolBarStackButtonBlock.h
+++ b/Engine/Source/Runtime/Slate/Public/Framework/MultiBox/SToolBarStackButtonBlock.h
@@ -89,7 +89,7 @@ protected:
 /**
  * Tool bar button MultiBlock widget
  */
-class SLATE_API SToolBarStackButtonBlock
+class SToolBarStackButtonBlock
 	: public SMultiBlockBaseWidget
 {
 
@@ -111,7 +111,7 @@ public:
 	/**
 	 * Builds this MultiBlock widget up from the MultiBlock associated with it
 	 */
-	virtual void BuildMultiBlockWidget(const ISlateStyle* StyleSet, const FName& StyleName) override;
+	SLATE_API virtual void BuildMultiBlockWidget(const ISlateStyle* StyleSet, const FName& StyleName) override;
 
 
 	/**
@@ -119,65 +119,65 @@ public:
 	 *
 	 * @param	InArgs	The declaration data for this widget
 	 */
-	void Construct( const FArguments& InArgs );
+	SLATE_API void Construct( const FArguments& InArgs );
 
 protected:
 
 	/**
 	 * Called by Slate when this tool bar check box button is toggled
 	 */
-	void OnCheckStateChanged(const ECheckBoxState NewCheckedState);
+	SLATE_API void OnCheckStateChanged(const ECheckBoxState NewCheckedState);
 
 	/**
 	 * Called by slate to determine if this button should appear checked
 	 *
 	 * @return ECheckBoxState::Checked if it should be checked, ECheckBoxState::Unchecked if not.
 	 */
-	ECheckBoxState GetCheckState() const;
+	SLATE_API ECheckBoxState GetCheckState() const;
 
 	/**
 	 * Called by Slate when content for this button's menu needs to be generated
 	 *
 	 * @return	The widget to use for the menu content
 	 */
-	TSharedRef<SWidget> OnGetMenuContent();
+	SLATE_API TSharedRef<SWidget> OnGetMenuContent();
 
 	/**
 	 * Called by Slate to determine if this button is enabled
 	 * 
 	 * @return True if the menu entry is enabled, false otherwise
 	 */
-	bool IsEnabled() const;
+	SLATE_API bool IsEnabled() const;
 
 	/**
 	 * Called by Slate to determine if this button is visible
 	 *
 	 * @return EVisibility::Visible or EVisibility::Collapsed, depending on if the button should be displayed
 	 */
-	EVisibility GetVisibility() const;
+	SLATE_API EVisibility GetVisibility() const;
 protected:
 
 	/** Gets the icon brush for the toolbar block widget */
-	const FSlateBrush* GetIconBrush() const;
+	SLATE_API const FSlateBrush* GetIconBrush() const;
 
 	/** Gets the label for the toolbar block widget */
-	FText GetLabel() const;
+	SLATE_API FText GetLabel() const;
 
 	/** Gets the description for the toolbar block widget */
-	FText GetDescription() const;
+	SLATE_API FText GetDescription() const;
 
 	/** @return The icon for the toolbar button; may be dynamic, so check HasDynamicIcon */
-	const FSlateBrush* GetNormalIconBrush() const;
+	SLATE_API const FSlateBrush* GetNormalIconBrush() const;
 
 	/** @return The small icon for the toolbar button; may be dynamic, so check HasDynamicIcon */
-	const FSlateBrush* GetSmallIconBrush() const;
+	SLATE_API const FSlateBrush* GetSmallIconBrush() const;
 
 	/** Called by Slate to determine whether icons/labels are visible */
-	EVisibility GetIconVisibility(bool bIsASmallIcon) const;
+	SLATE_API EVisibility GetIconVisibility(bool bIsASmallIcon) const;
 
-	FSlateColor GetIconForegroundColor() const;
+	SLATE_API FSlateColor GetIconForegroundColor() const;
 
-	const FSlateBrush* GetOverlayIconBrush() const;
+	SLATE_API const FSlateBrush* GetOverlayIconBrush() const;
 
 protected:
 	/** Overrides the visibility of the of label. This is used to set up the LabelVisibility attribute */
diff --git a/Engine/Source/Runtime/Slate/Public/Framework/MultiBox/SUniformToolbarButtonBlock.h b/Engine/Source/Runtime/Slate/Public/Framework/MultiBox/SUniformToolbarButtonBlock.h
--- a/Engine/Source/Runtime/Slate/Public/Framework/MultiBox/SUniformToolbarButtonBlock.h
+++ b/Engine/Source/Runtime/Slate/Public/Framework/MultiBox/SUniformToolbarButtonBlock.h
@@ -8,7 +8,7 @@ struct FButtonArgs;
 /**
  * Horizontal Button Row MultiBlock
  */
-class SLATE_API FUniformToolbarButtonBlock
+class FUniformToolbarButtonBlock
 	: public FToolBarButtonBlock
 {
 
@@ -19,7 +19,7 @@ public:
 	 *
 	 * @param	ButtonArgs The Parameters object which will provide the data to initialize the button
 	 */
-	FUniformToolbarButtonBlock( FButtonArgs ButtonArgs);
+	SLATE_API FUniformToolbarButtonBlock( FButtonArgs ButtonArgs);
 
 	
 	/**
@@ -27,11 +27,11 @@ public:
 	 *
 	 * @return  MultiBlock widget object
 	 */
-	 virtual TSharedRef< class IMultiBlockBaseWidget > ConstructWidget() const override;
+	 SLATE_API virtual TSharedRef< class IMultiBlockBaseWidget > ConstructWidget() const override;
 
 };
 
-class SLATE_API SUniformToolbarButtonBlock
+class SUniformToolbarButtonBlock
 	: public SToolBarButtonBlock
 {
 	
diff --git a/Engine/Source/Runtime/Slate/Public/Framework/MultiBox/ToolMenuBase.h b/Engine/Source/Runtime/Slate/Public/Framework/MultiBox/ToolMenuBase.h
--- a/Engine/Source/Runtime/Slate/Public/Framework/MultiBox/ToolMenuBase.h
+++ b/Engine/Source/Runtime/Slate/Public/Framework/MultiBox/ToolMenuBase.h
@@ -55,7 +55,7 @@ struct FCustomizedToolMenuNameArray
 };
 
 USTRUCT()
-struct SLATE_API FCustomizedToolMenu
+struct FCustomizedToolMenu
 {
 	GENERATED_BODY()
 
@@ -79,40 +79,40 @@ struct SLATE_API FCustomizedToolMenu
 
 	FNamePermissionList MenuPermissions;
 
-	FCustomizedToolMenuEntry* FindEntry(const FName InEntryName);
-	const FCustomizedToolMenuEntry* FindEntry(const FName InEntryName) const;
-	FCustomizedToolMenuEntry* AddEntry(const FName InEntryName);
-	ECustomizedToolMenuVisibility GetEntryVisiblity(const FName InSectionName) const;
-	bool IsEntryHidden(const FName InEntryName) const;
-	FName GetEntrySectionName(const FName InEntryName) const;
-
-	FCustomizedToolMenuSection* FindSection(const FName InSectionName);
-	const FCustomizedToolMenuSection* FindSection(const FName InSectionName) const;
-	FCustomizedToolMenuSection* AddSection(const FName InSectionName);
-	ECustomizedToolMenuVisibility GetSectionVisiblity(const FName InSectionName) const;
-	bool IsSectionHidden(const FName InSectionName) const;
-
-	void SetSuppressExtenders(const FName InOwnerName, const bool bInSuppress);
-	bool IsSuppressExtenders() const;
+	SLATE_API FCustomizedToolMenuEntry* FindEntry(const FName InEntryName);
+	SLATE_API const FCustomizedToolMenuEntry* FindEntry(const FName InEntryName) const;
+	SLATE_API FCustomizedToolMenuEntry* AddEntry(const FName InEntryName);
+	SLATE_API ECustomizedToolMenuVisibility GetEntryVisiblity(const FName InSectionName) const;
+	SLATE_API bool IsEntryHidden(const FName InEntryName) const;
+	SLATE_API FName GetEntrySectionName(const FName InEntryName) const;
+
+	SLATE_API FCustomizedToolMenuSection* FindSection(const FName InSectionName);
+	SLATE_API const FCustomizedToolMenuSection* FindSection(const FName InSectionName) const;
+	SLATE_API FCustomizedToolMenuSection* AddSection(const FName InSectionName);
+	SLATE_API ECustomizedToolMenuVisibility GetSectionVisiblity(const FName InSectionName) const;
+	SLATE_API bool IsSectionHidden(const FName InSectionName) const;
+
+	SLATE_API void SetSuppressExtenders(const FName InOwnerName, const bool bInSuppress);
+	SLATE_API bool IsSuppressExtenders() const;
 };
 
-struct SLATE_API FCustomizedToolMenuHierarchy
+struct FCustomizedToolMenuHierarchy
 {
-	FName GetEntrySectionName(const FName InEntryName) const;
-	bool IsEntryHidden(const FName InEntryName) const;
+	SLATE_API FName GetEntrySectionName(const FName InEntryName) const;
+	SLATE_API bool IsEntryHidden(const FName InEntryName) const;
 
-	bool IsSectionHidden(const FName InSectionName) const;
+	SLATE_API bool IsSectionHidden(const FName InSectionName) const;
 
-	bool IsSuppressExtenders() const;
+	SLATE_API bool IsSuppressExtenders() const;
 
-	FCustomizedToolMenu GenerateFlattened() const;
+	SLATE_API FCustomizedToolMenu GenerateFlattened() const;
 
 	TArray<const FCustomizedToolMenu*> Hierarchy;
 	TArray<const FCustomizedToolMenu*> RuntimeHierarchy;
 };
 
-UCLASS(Abstract)
-class SLATE_API UToolMenuBase : public UObject
+UCLASS(Abstract, MinimalAPI)
+class UToolMenuBase : public UObject
 {
 	GENERATED_BODY()
 
diff --git a/Engine/Source/Runtime/Slate/Public/Framework/Notifications/NotificationManager.h b/Engine/Source/Runtime/Slate/Public/Framework/Notifications/NotificationManager.h
--- a/Engine/Source/Runtime/Slate/Public/Framework/Notifications/NotificationManager.h
+++ b/Engine/Source/Runtime/Slate/Public/Framework/Notifications/NotificationManager.h
@@ -73,7 +73,7 @@ public:
 /**
  * A class which manages a group of notification windows                 
  */
-class SLATE_API FSlateNotificationManager
+class FSlateNotificationManager
 {
 	friend class SNotificationList;
 
@@ -81,25 +81,25 @@ public:
 	/**
 	 * Gets the instance of this manager                   
 	 */
-	static FSlateNotificationManager& Get();
+	static SLATE_API FSlateNotificationManager& Get();
 
 	/** Update the manager */
-	void Tick();
+	SLATE_API void Tick();
 
 	/** Provide a window under which all notifications should nest. */
-	void SetRootWindow( const TSharedRef<SWindow> InRootWindow );
+	SLATE_API void SetRootWindow( const TSharedRef<SWindow> InRootWindow );
 
 	/**
 	 * Adds a floating notification
 	 * @param Info 		Contains various settings used to initialize the notification
 	 */
-	TSharedPtr<SNotificationItem> AddNotification(const FNotificationInfo& Info);
+	SLATE_API TSharedPtr<SNotificationItem> AddNotification(const FNotificationInfo& Info);
 
 	/**
 	 * Thread safe method of queuing a notification for presentation on the next tick
 	 * @param Info 		Pointer to heap allocated notification info. Released by FSlateNotificationManager once the notification is displayed;
 	 */
-	void QueueNotification(FNotificationInfo* Info);
+	SLATE_API void QueueNotification(FNotificationInfo* Info);
 
 	/**
 	 * Begin a progress notification. These notifications should be used for background work not blocking work. Use SlowTask for blocking work
@@ -109,7 +109,7 @@ public:
  	 *
 	 * @return NotificationHandle	A handle that can be used to update the progress or cancel the notification
 	 */
-	FProgressNotificationHandle StartProgressNotification(FText DisplayText, int32 TotalWorkToDo);
+	SLATE_API FProgressNotificationHandle StartProgressNotification(FText DisplayText, int32 TotalWorkToDo);
 
 	/**
 	 * Updates a progress notification. 
@@ -119,18 +119,18 @@ public:
 	 * @param UpdatedDisplayText	(Optional) Updated display text of the notification
 	 * @return NotificationHandle	A handle that can be used to update the progress or cancel the notification
 	 */
-	void UpdateProgressNotification(FProgressNotificationHandle InHandle, int32 TotalWorkDone, int32 UpdatedTotalWorkToDo = 0, FText UpdatedDisplayText = FText::GetEmpty());
+	SLATE_API void UpdateProgressNotification(FProgressNotificationHandle InHandle, int32 TotalWorkDone, int32 UpdatedTotalWorkToDo = 0, FText UpdatedDisplayText = FText::GetEmpty());
 
 	/**
 	 * Cancels an active notification
 	 */
-	void CancelProgressNotification(FProgressNotificationHandle InHandle);
+	SLATE_API void CancelProgressNotification(FProgressNotificationHandle InHandle);
 
 	/**
 	 * Sets the progress notification handler for the current application. Only one handler is used at a time. 
 	 * This handler is not managed in any way. If your handler is being destroyed call SetProgressNotificationHandler(nullptr)
 	 */
-	void SetProgressNotificationHandler(IProgressNotificationHandler* NewHandler);
+	SLATE_API void SetProgressNotificationHandler(IProgressNotificationHandler* NewHandler);
 
 	/**
 	 * Called back from the SlateApplication when a window is activated/resized
@@ -138,13 +138,13 @@ public:
 	 * as there isn't a cross-platform OS-level way of making a 'topmost child'.
 	 * @param ActivateEvent 	Information about the activation event
 	 */
-	void ForceNotificationsInFront(const TSharedRef<SWindow> InWindow);
+	SLATE_API void ForceNotificationsInFront(const TSharedRef<SWindow> InWindow);
 
 	/**
 	 * Gets all the windows that represent notifications
 	 * @returns an array of windows that contain notifications.
 	 */
-	void GetWindows(TArray< TSharedRef<SWindow> >& OutWindows) const;
+	SLATE_API void GetWindows(TArray< TSharedRef<SWindow> >& OutWindows) const;
 
 	/**
 	 * Sets whether notifications should be displayed at all. Note, notifications can be
@@ -170,22 +170,22 @@ public:
 	/**
 	 * Register a Staged Async Notification, allowing the NotificationManager to keep a reference to it
 	 */
-	void RegisterStagedNotification(TSharedPtr<IAsyncTaskNotificationImpl> InNotification);
+	SLATE_API void RegisterStagedNotification(TSharedPtr<IAsyncTaskNotificationImpl> InNotification);
 
 	/**
 	 * Unregister a previously added Staged Notification
 	 */
-	void UnregisterStagedNotification(TSharedPtr<IAsyncTaskNotificationImpl> InNotification);
+	SLATE_API void UnregisterStagedNotification(TSharedPtr<IAsyncTaskNotificationImpl> InNotification);
 
 protected:
 	/** Protect constructor as this is a singleton */
-	FSlateNotificationManager();
+	SLATE_API FSlateNotificationManager();
 
 	/** Create a notification list for the specified screen rectangle */
-	TSharedRef<SNotificationList> CreateStackForArea(const FSlateRect& InRectangle, TSharedPtr<SWindow> Window);
+	SLATE_API TSharedRef<SNotificationList> CreateStackForArea(const FSlateRect& InRectangle, TSharedPtr<SWindow> Window);
 
 	/** FCoreDelegates::OnPreExit shutdown callback */
-	void ShutdownOnPreExit();
+	SLATE_API void ShutdownOnPreExit();
 
 private:
 
@@ -227,7 +227,7 @@ private:
 	FCriticalSection StagedNotificationCS;
 
 	/** Counter used to create progress handles */
-	static int32 ProgressHandleCounter;
+	static SLATE_API int32 ProgressHandleCounter;
 
 	/** Whether notifications should be displayed or not.  This can be used to globally suppress notification pop-ups */
 	bool bAllowNotifications = true;
diff --git a/Engine/Source/Runtime/Slate/Public/Framework/Text/BaseTextLayoutMarshaller.h b/Engine/Source/Runtime/Slate/Public/Framework/Text/BaseTextLayoutMarshaller.h
--- a/Engine/Source/Runtime/Slate/Public/Framework/Text/BaseTextLayoutMarshaller.h
+++ b/Engine/Source/Runtime/Slate/Public/Framework/Text/BaseTextLayoutMarshaller.h
@@ -7,7 +7,7 @@
 /**
  * Base class implementing some common functionality for all text marshallers
  */
-class SLATE_API FBaseTextLayoutMarshaller : public ITextLayoutMarshaller
+class FBaseTextLayoutMarshaller : public ITextLayoutMarshaller
 {
 public:
 
diff --git a/Engine/Source/Runtime/Slate/Public/Framework/Text/CharRangeList.h b/Engine/Source/Runtime/Slate/Public/Framework/Text/CharRangeList.h
--- a/Engine/Source/Runtime/Slate/Public/Framework/Text/CharRangeList.h
+++ b/Engine/Source/Runtime/Slate/Public/Framework/Text/CharRangeList.h
@@ -32,7 +32,7 @@ struct FCharRange
 
 /** Represents a list of character ranges. */
 USTRUCT()
-struct SLATE_API FCharRangeList
+struct FCharRangeList
 {
 	GENERATED_USTRUCT_BODY()
 
@@ -46,22 +46,22 @@ struct SLATE_API FCharRangeList
 	*       "a-zA-Z0-9._" <-- All lowercase and uppercase letters, all digits, dot and underscore.
 	*       "a-zA-Z0-9\-\\._" <-- All lowercase and uppercase letters, all digits, minus sign, backslash, dot and underscore.
 	*/
-	bool InitializeFromString(const FString& InDefinitionString);
+	SLATE_API bool InitializeFromString(const FString& InDefinitionString);
 	/** Returns a definition string that represents all the character ranges in this instance. */
-	FString ToString() const;
+	SLATE_API FString ToString() const;
 
 	/** Empties this instance. */
-	void Empty();
+	SLATE_API void Empty();
 
 	/** Returns true if this instance does not contain any characters, or false otherwise. */
-	bool IsEmpty() const;
+	SLATE_API bool IsEmpty() const;
 
 	/** Returns true if the given char is included in any of the char ranges, or false otherwise. */
-	bool IsCharIncluded(TCHAR InChar) const;
+	SLATE_API bool IsCharIncluded(TCHAR InChar) const;
 	/** Returns true if all the characters in the given string are included in any of the char ranges, or false otherwise. */
-	bool AreAllCharsIncluded(const FString& InString) const;
+	SLATE_API bool AreAllCharsIncluded(const FString& InString) const;
 	/** Finds all the characters in the given string that are not included in any of the char ranges. */
-	TSet<TCHAR> FindCharsNotIncluded(const FString& InString) const;
+	SLATE_API TSet<TCHAR> FindCharsNotIncluded(const FString& InString) const;
 
 	UPROPERTY(EditAnywhere, Category = "Char Range List")
 	TArray<FCharRange> Ranges;
diff --git a/Engine/Source/Runtime/Slate/Public/Framework/Text/DefaultLayoutBlock.h b/Engine/Source/Runtime/Slate/Public/Framework/Text/DefaultLayoutBlock.h
--- a/Engine/Source/Runtime/Slate/Public/Framework/Text/DefaultLayoutBlock.h
+++ b/Engine/Source/Runtime/Slate/Public/Framework/Text/DefaultLayoutBlock.h
@@ -5,7 +5,7 @@
 #include "Framework/Text/IRun.h"
 #include "Framework/Text/ILayoutBlock.h"
 
-class SLATE_API FDefaultLayoutBlock : public ILayoutBlock
+class FDefaultLayoutBlock : public ILayoutBlock
 {
 public:
 
diff --git a/Engine/Source/Runtime/Slate/Public/Framework/Text/ILayoutBlock.h b/Engine/Source/Runtime/Slate/Public/Framework/Text/ILayoutBlock.h
--- a/Engine/Source/Runtime/Slate/Public/Framework/Text/ILayoutBlock.h
+++ b/Engine/Source/Runtime/Slate/Public/Framework/Text/ILayoutBlock.h
@@ -5,7 +5,7 @@
 #include "Framework/Text/IRun.h"
 #include "Framework/Text/IRunRenderer.h"
 
-class SLATE_API ILayoutBlock
+class ILayoutBlock
 {
 public:
 
diff --git a/Engine/Source/Runtime/Slate/Public/Framework/Text/ILineHighlighter.h b/Engine/Source/Runtime/Slate/Public/Framework/Text/ILineHighlighter.h
--- a/Engine/Source/Runtime/Slate/Public/Framework/Text/ILineHighlighter.h
+++ b/Engine/Source/Runtime/Slate/Public/Framework/Text/ILineHighlighter.h
@@ -3,7 +3,7 @@
 
 #include "CoreMinimal.h"
 
-class SLATE_API ILineHighlighter
+class ILineHighlighter
 {
 public:
 
diff --git a/Engine/Source/Runtime/Slate/Public/Framework/Text/IRichTextMarkupParser.h b/Engine/Source/Runtime/Slate/Public/Framework/Text/IRichTextMarkupParser.h
--- a/Engine/Source/Runtime/Slate/Public/Framework/Text/IRichTextMarkupParser.h
+++ b/Engine/Source/Runtime/Slate/Public/Framework/Text/IRichTextMarkupParser.h
@@ -9,7 +9,7 @@
 
 /**
  */
-class SLATE_API IRichTextMarkupParser
+class IRichTextMarkupParser
 {
 public:
 
diff --git a/Engine/Source/Runtime/Slate/Public/Framework/Text/IRichTextMarkupWriter.h b/Engine/Source/Runtime/Slate/Public/Framework/Text/IRichTextMarkupWriter.h
--- a/Engine/Source/Runtime/Slate/Public/Framework/Text/IRichTextMarkupWriter.h
+++ b/Engine/Source/Runtime/Slate/Public/Framework/Text/IRichTextMarkupWriter.h
@@ -9,7 +9,7 @@
 
 /**
  */
-class SLATE_API IRichTextMarkupWriter
+class IRichTextMarkupWriter
 {
 public:
 
diff --git a/Engine/Source/Runtime/Slate/Public/Framework/Text/IRun.h b/Engine/Source/Runtime/Slate/Public/Framework/Text/IRun.h
--- a/Engine/Source/Runtime/Slate/Public/Framework/Text/IRun.h
+++ b/Engine/Source/Runtime/Slate/Public/Framework/Text/IRun.h
@@ -124,7 +124,7 @@ struct FLayoutBlockTextContext : public FRunTextContext
 	TextBiDi::ETextDirection TextDirection;
 };
 
-class SLATE_API IRun
+class IRun
 {
 public:
 
diff --git a/Engine/Source/Runtime/Slate/Public/Framework/Text/IRunRenderer.h b/Engine/Source/Runtime/Slate/Public/Framework/Text/IRunRenderer.h
--- a/Engine/Source/Runtime/Slate/Public/Framework/Text/IRunRenderer.h
+++ b/Engine/Source/Runtime/Slate/Public/Framework/Text/IRunRenderer.h
@@ -3,7 +3,7 @@
 
 #include "CoreMinimal.h"
 
-class SLATE_API IRunRenderer
+class IRunRenderer
 {
 public:
 
diff --git a/Engine/Source/Runtime/Slate/Public/Framework/Text/ISlateLineHighlighter.h b/Engine/Source/Runtime/Slate/Public/Framework/Text/ISlateLineHighlighter.h
--- a/Engine/Source/Runtime/Slate/Public/Framework/Text/ISlateLineHighlighter.h
+++ b/Engine/Source/Runtime/Slate/Public/Framework/Text/ISlateLineHighlighter.h
@@ -12,7 +12,7 @@ class FWidgetStyle;
 struct FGeometry;
 struct FTextBlockStyle;
 
-class SLATE_API ISlateLineHighlighter : public ILineHighlighter
+class ISlateLineHighlighter : public ILineHighlighter
 {
 public:
 	virtual ~ISlateLineHighlighter() {}
diff --git a/Engine/Source/Runtime/Slate/Public/Framework/Text/ISlateRun.h b/Engine/Source/Runtime/Slate/Public/Framework/Text/ISlateRun.h
--- a/Engine/Source/Runtime/Slate/Public/Framework/Text/ISlateRun.h
+++ b/Engine/Source/Runtime/Slate/Public/Framework/Text/ISlateRun.h
@@ -34,7 +34,7 @@ struct FTextArgs
 
 };
 
-class SLATE_API ISlateRun : public IRun
+class ISlateRun : public IRun
 {
 public:
 	virtual int32 OnPaint(const FPaintArgs& PaintArgs, const FTextArgs& TextArgs, const FGeometry& AllottedGeometry, const FSlateRect& MyCullingRect, FSlateWindowElementList& OutDrawElements, int32 LayerId, const FWidgetStyle& InWidgetStyle, bool bParentEnabled) const = 0;
diff --git a/Engine/Source/Runtime/Slate/Public/Framework/Text/ISlateRunRenderer.h b/Engine/Source/Runtime/Slate/Public/Framework/Text/ISlateRunRenderer.h
--- a/Engine/Source/Runtime/Slate/Public/Framework/Text/ISlateRunRenderer.h
+++ b/Engine/Source/Runtime/Slate/Public/Framework/Text/ISlateRunRenderer.h
@@ -11,7 +11,7 @@ class FPaintArgs;
 class FSlateWindowElementList;
 struct FTextBlockStyle;
 
-class SLATE_API ISlateRunRenderer : public IRunRenderer
+class ISlateRunRenderer : public IRunRenderer
 {
 public:
 	virtual ~ISlateRunRenderer() {}
diff --git a/Engine/Source/Runtime/Slate/Public/Framework/Text/ITextDecorator.h b/Engine/Source/Runtime/Slate/Public/Framework/Text/ITextDecorator.h
--- a/Engine/Source/Runtime/Slate/Public/Framework/Text/ITextDecorator.h
+++ b/Engine/Source/Runtime/Slate/Public/Framework/Text/ITextDecorator.h
@@ -10,7 +10,7 @@ class ISlateStyle;
 
 #if WITH_FANCY_TEXT
 
-struct SLATE_API FTextRunParseResults
+struct FTextRunParseResults
 {
 	FTextRunParseResults( FString InName, const FTextRange& InOriginalRange)
 		: Name( InName )
@@ -35,7 +35,7 @@ struct SLATE_API FTextRunParseResults
 	TMap< FString, FTextRange > MetaData; 
 };
 
-struct SLATE_API FTextLineParseResults
+struct FTextLineParseResults
 {
 public:
 
@@ -57,7 +57,7 @@ public:
 	TArray< FTextRunParseResults > Runs;
 };
 
-struct SLATE_API FTextRunInfo : FRunInfo
+struct FTextRunInfo : FRunInfo
 {
 	FTextRunInfo( FString InName, const FText& InContent )
 		: FRunInfo( MoveTemp(InName) )
@@ -69,7 +69,7 @@ struct SLATE_API FTextRunInfo : FRunInfo
 	FText Content;
 };
 
-class SLATE_API ITextDecorator
+class ITextDecorator
 {
 public:
 
diff --git a/Engine/Source/Runtime/Slate/Public/Framework/Text/ITextLayoutMarshaller.h b/Engine/Source/Runtime/Slate/Public/Framework/Text/ITextLayoutMarshaller.h
--- a/Engine/Source/Runtime/Slate/Public/Framework/Text/ITextLayoutMarshaller.h
+++ b/Engine/Source/Runtime/Slate/Public/Framework/Text/ITextLayoutMarshaller.h
@@ -8,7 +8,7 @@ class FTextLayout;
 /**
  * Interface used to get/set the raw text to/from a text layout
  */
-class SLATE_API ITextLayoutMarshaller
+class ITextLayoutMarshaller
 {
 public:
 
diff --git a/Engine/Source/Runtime/Slate/Public/Framework/Text/PlainTextLayoutMarshaller.h b/Engine/Source/Runtime/Slate/Public/Framework/Text/PlainTextLayoutMarshaller.h
--- a/Engine/Source/Runtime/Slate/Public/Framework/Text/PlainTextLayoutMarshaller.h
+++ b/Engine/Source/Runtime/Slate/Public/Framework/Text/PlainTextLayoutMarshaller.h
@@ -10,23 +10,23 @@ class FTextLayout;
 /**
  * Get/set the raw text to/from a text layout as plain text
  */
-class SLATE_API FPlainTextLayoutMarshaller : public FBaseTextLayoutMarshaller
+class FPlainTextLayoutMarshaller : public FBaseTextLayoutMarshaller
 {
 public:
 
-	static TSharedRef< FPlainTextLayoutMarshaller > Create();
+	static SLATE_API TSharedRef< FPlainTextLayoutMarshaller > Create();
 
-	virtual ~FPlainTextLayoutMarshaller();
+	SLATE_API virtual ~FPlainTextLayoutMarshaller();
 	
-	void SetIsPassword(const TAttribute<bool>& InIsPassword);
+	SLATE_API void SetIsPassword(const TAttribute<bool>& InIsPassword);
 
 	// ITextLayoutMarshaller
-	virtual void SetText(const FString& SourceString, FTextLayout& TargetTextLayout) override;
-	virtual void GetText(FString& TargetString, const FTextLayout& SourceTextLayout) override;
+	SLATE_API virtual void SetText(const FString& SourceString, FTextLayout& TargetTextLayout) override;
+	SLATE_API virtual void GetText(FString& TargetString, const FTextLayout& SourceTextLayout) override;
 
 protected:
 
-	FPlainTextLayoutMarshaller();
+	SLATE_API FPlainTextLayoutMarshaller();
 
 	/** This this marshaller displaying a password? */
 	TAttribute<bool> bIsPassword;
diff --git a/Engine/Source/Runtime/Slate/Public/Framework/Text/RichTextLayoutMarshaller.h b/Engine/Source/Runtime/Slate/Public/Framework/Text/RichTextLayoutMarshaller.h
--- a/Engine/Source/Runtime/Slate/Public/Framework/Text/RichTextLayoutMarshaller.h
+++ b/Engine/Source/Runtime/Slate/Public/Framework/Text/RichTextLayoutMarshaller.h
@@ -22,21 +22,21 @@ struct FTextBlockStyle;
 /**
  * Get/set the raw text to/from a text layout as rich text
  */
-class SLATE_API FRichTextLayoutMarshaller : public FBaseTextLayoutMarshaller
+class FRichTextLayoutMarshaller : public FBaseTextLayoutMarshaller
 {
 public:
 
-	static TSharedRef< FRichTextLayoutMarshaller > Create(TArray< TSharedRef< ITextDecorator > > InDecorators, const ISlateStyle* const InDecoratorStyleSet);
-	static TSharedRef< FRichTextLayoutMarshaller > Create(TSharedPtr< IRichTextMarkupParser > InParser, TSharedPtr< IRichTextMarkupWriter > InWriter, TArray< TSharedRef< ITextDecorator > > InDecorators, const ISlateStyle* const InDecoratorStyleSet);
+	static SLATE_API TSharedRef< FRichTextLayoutMarshaller > Create(TArray< TSharedRef< ITextDecorator > > InDecorators, const ISlateStyle* const InDecoratorStyleSet);
+	static SLATE_API TSharedRef< FRichTextLayoutMarshaller > Create(TSharedPtr< IRichTextMarkupParser > InParser, TSharedPtr< IRichTextMarkupWriter > InWriter, TArray< TSharedRef< ITextDecorator > > InDecorators, const ISlateStyle* const InDecoratorStyleSet);
 
-	virtual ~FRichTextLayoutMarshaller();
+	SLATE_API virtual ~FRichTextLayoutMarshaller();
 
 	// ITextLayoutMarshaller
-	virtual void SetText(const FString& SourceString, FTextLayout& TargetTextLayout) override;
-	virtual void GetText(FString& TargetString, const FTextLayout& SourceTextLayout) override;
+	SLATE_API virtual void SetText(const FString& SourceString, FTextLayout& TargetTextLayout) override;
+	SLATE_API virtual void GetText(FString& TargetString, const FTextLayout& SourceTextLayout) override;
 
 	/** Set the Multiplier applied on the font size of the text*/
-	void SetFontSizeMultiplier(const float NewFontSizeMultiplier);
+	SLATE_API void SetFontSizeMultiplier(const float NewFontSizeMultiplier);
 
 	/**
 	 * Replace all decorators on this marshaller
@@ -62,12 +62,12 @@ public:
 
 protected:
 
-	FRichTextLayoutMarshaller(TArray< TSharedRef< ITextDecorator > > InDecorators, const ISlateStyle* const InDecoratorStyleSet);
-	FRichTextLayoutMarshaller(TSharedPtr< IRichTextMarkupParser > InParser, TSharedPtr< IRichTextMarkupWriter > InWriter, TArray< TSharedRef< ITextDecorator > > InDecorators, const ISlateStyle* const InDecoratorStyleSet);
+	SLATE_API FRichTextLayoutMarshaller(TArray< TSharedRef< ITextDecorator > > InDecorators, const ISlateStyle* const InDecoratorStyleSet);
+	SLATE_API FRichTextLayoutMarshaller(TSharedPtr< IRichTextMarkupParser > InParser, TSharedPtr< IRichTextMarkupWriter > InWriter, TArray< TSharedRef< ITextDecorator > > InDecorators, const ISlateStyle* const InDecoratorStyleSet);
 
-	TSharedPtr< ITextDecorator > TryGetDecorator(const FString& Line, const FTextRunParseResults& TextRun) const;
+	SLATE_API TSharedPtr< ITextDecorator > TryGetDecorator(const FString& Line, const FTextRunParseResults& TextRun) const;
 
-	virtual void AppendRunsForText(
+	SLATE_API virtual void AppendRunsForText(
 		const int32 LineIndex,
 		const FTextRunParseResults& TextRun, 
 		const FString& ProcessedString, 
diff --git a/Engine/Source/Runtime/Slate/Public/Framework/Text/RichTextMarkupProcessing.h b/Engine/Source/Runtime/Slate/Public/Framework/Text/RichTextMarkupProcessing.h
--- a/Engine/Source/Runtime/Slate/Public/Framework/Text/RichTextMarkupProcessing.h
+++ b/Engine/Source/Runtime/Slate/Public/Framework/Text/RichTextMarkupProcessing.h
@@ -10,14 +10,14 @@
 
 #if WITH_FANCY_TEXT
 
-class SLATE_API FDefaultRichTextMarkupParser : public IRichTextMarkupParser
+class FDefaultRichTextMarkupParser : public IRichTextMarkupParser
 {
 public:
-	static TSharedRef< FDefaultRichTextMarkupParser > Create();
-	static TSharedRef< FDefaultRichTextMarkupParser > GetStaticInstance();
+	static SLATE_API TSharedRef< FDefaultRichTextMarkupParser > Create();
+	static SLATE_API TSharedRef< FDefaultRichTextMarkupParser > GetStaticInstance();
 
 public:
-	virtual void Process(TArray<FTextLineParseResults>& Results, const FString& Input, FString& Output) override;
+	SLATE_API virtual void Process(TArray<FTextLineParseResults>& Results, const FString& Input, FString& Output) override;
 
 private:
 	FDefaultRichTextMarkupParser();
@@ -30,14 +30,14 @@ private:
 	FRegexPattern AttributeRegexPattern;
 };
 
-class SLATE_API FDefaultRichTextMarkupWriter : public IRichTextMarkupWriter
+class FDefaultRichTextMarkupWriter : public IRichTextMarkupWriter
 {
 public:
-	static TSharedRef< FDefaultRichTextMarkupWriter > Create();
-	static TSharedRef< FDefaultRichTextMarkupWriter > GetStaticInstance();
+	static SLATE_API TSharedRef< FDefaultRichTextMarkupWriter > Create();
+	static SLATE_API TSharedRef< FDefaultRichTextMarkupWriter > GetStaticInstance();
 
 public:
-	virtual void Write(const TArray<FRichTextLine>& InLines, FString& Output) override;
+	SLATE_API virtual void Write(const TArray<FRichTextLine>& InLines, FString& Output) override;
 
 private:
 	FDefaultRichTextMarkupWriter() {}
diff --git a/Engine/Source/Runtime/Slate/Public/Framework/Text/ShapedTextCache.h b/Engine/Source/Runtime/Slate/Public/Framework/Text/ShapedTextCache.h
--- a/Engine/Source/Runtime/Slate/Public/Framework/Text/ShapedTextCache.h
+++ b/Engine/Source/Runtime/Slate/Public/Framework/Text/ShapedTextCache.h
@@ -51,7 +51,7 @@ public:
 };
 
 /** Cache of shaped text */
-class SLATE_API FShapedTextCache
+class FShapedTextCache
 {
 public:
 	/** Create a new shaped text cache */
@@ -67,7 +67,7 @@ public:
 	 *
 	 * @return The shaped text instance, or null if it wasn't found or was stale
 	 */
-	FShapedGlyphSequencePtr FindShapedText(const FCachedShapedTextKey& InKey) const;
+	SLATE_API FShapedGlyphSequencePtr FindShapedText(const FCachedShapedTextKey& InKey) const;
 
 	/**
 	 * Add the given shaped text instance to the cache, or generate a new instance and add that based on the parameters provided
@@ -79,9 +79,9 @@ public:
 	 *
 	 * @return The shaped text instance
 	 */
-	FShapedGlyphSequenceRef AddShapedText(const FCachedShapedTextKey& InKey, const TCHAR* InText);
-	FShapedGlyphSequenceRef AddShapedText(const FCachedShapedTextKey& InKey, const TCHAR* InText, const TextBiDi::ETextDirection InTextDirection);
-	FShapedGlyphSequenceRef AddShapedText(const FCachedShapedTextKey& InKey, FShapedGlyphSequenceRef InShapedText);
+	SLATE_API FShapedGlyphSequenceRef AddShapedText(const FCachedShapedTextKey& InKey, const TCHAR* InText);
+	SLATE_API FShapedGlyphSequenceRef AddShapedText(const FCachedShapedTextKey& InKey, const TCHAR* InText, const TextBiDi::ETextDirection InTextDirection);
+	SLATE_API FShapedGlyphSequenceRef AddShapedText(const FCachedShapedTextKey& InKey, FShapedGlyphSequenceRef InShapedText);
 
 	/**
 	 * Try and find an existing shaped text instance, or add a new entry to the cache if one cannot be found
@@ -92,8 +92,8 @@ public:
 	 *
 	 * @return The shaped text instance
 	 */
-	FShapedGlyphSequenceRef FindOrAddShapedText(const FCachedShapedTextKey& InKey, const TCHAR* InText);
-	FShapedGlyphSequenceRef FindOrAddShapedText(const FCachedShapedTextKey& InKey, const TCHAR* InText, const TextBiDi::ETextDirection InTextDirection);
+	SLATE_API FShapedGlyphSequenceRef FindOrAddShapedText(const FCachedShapedTextKey& InKey, const TCHAR* InText);
+	SLATE_API FShapedGlyphSequenceRef FindOrAddShapedText(const FCachedShapedTextKey& InKey, const TCHAR* InText, const TextBiDi::ETextDirection InTextDirection);
 
 	/**
 	 * Try and find an existing overflow shaped text instance, or add a new entry to the cache if one cannot be found.
@@ -101,12 +101,12 @@ public:
 	 *
 	 * @return The shaped text instance
 	 */
-	FShapedGlyphSequenceRef FindOrAddOverflowEllipsisText(const float InScale, const FShapedTextContext& InTextContext, const FSlateFontInfo& InFontInfo);
+	SLATE_API FShapedGlyphSequenceRef FindOrAddOverflowEllipsisText(const float InScale, const FShapedTextContext& InTextContext, const FSlateFontInfo& InFontInfo);
 
 	/**
 	 * Clear this cache
 	 */
-	void Clear();
+	SLATE_API void Clear();
 
 	/**
 	 * Get the font cache used by this instance
diff --git a/Engine/Source/Runtime/Slate/Public/Framework/Text/SlateHyperlinkRun.h b/Engine/Source/Runtime/Slate/Public/Framework/Text/SlateHyperlinkRun.h
--- a/Engine/Source/Runtime/Slate/Public/Framework/Text/SlateHyperlinkRun.h
+++ b/Engine/Source/Runtime/Slate/Public/Framework/Text/SlateHyperlinkRun.h
@@ -19,7 +19,7 @@ enum class ETextHitPoint : uint8;
 
 #if WITH_FANCY_TEXT
 
-class SLATE_API FSlateHyperlinkRun : public ISlateRun, public TSharedFromThis< FSlateHyperlinkRun >
+class FSlateHyperlinkRun : public ISlateRun, public TSharedFromThis< FSlateHyperlinkRun >
 {
 public:
 
@@ -48,62 +48,62 @@ public:
 
 public:
 
-	static TSharedRef< FSlateHyperlinkRun > Create( const FRunInfo& InRunInfo, const TSharedRef< const FString >& InText, const FHyperlinkStyle& InStyle, FOnClick NavigateDelegate, FOnGenerateTooltip InTooltipDelegate, FOnGetTooltipText InTooltipTextDelegate );
+	static SLATE_API TSharedRef< FSlateHyperlinkRun > Create( const FRunInfo& InRunInfo, const TSharedRef< const FString >& InText, const FHyperlinkStyle& InStyle, FOnClick NavigateDelegate, FOnGenerateTooltip InTooltipDelegate, FOnGetTooltipText InTooltipTextDelegate );
 																																	 
-	static TSharedRef< FSlateHyperlinkRun > Create( const FRunInfo& InRunInfo, const TSharedRef< const FString >& InText, const FHyperlinkStyle& InStyle, FOnClick NavigateDelegate, FOnGenerateTooltip InTooltipDelegate, FOnGetTooltipText InTooltipTextDelegate, const FTextRange& InRange );
+	static SLATE_API TSharedRef< FSlateHyperlinkRun > Create( const FRunInfo& InRunInfo, const TSharedRef< const FString >& InText, const FHyperlinkStyle& InStyle, FOnClick NavigateDelegate, FOnGenerateTooltip InTooltipDelegate, FOnGetTooltipText InTooltipTextDelegate, const FTextRange& InRange );
 
 public:
 
 	virtual ~FSlateHyperlinkRun() {}
 
-	virtual FTextRange GetTextRange() const override;
+	SLATE_API virtual FTextRange GetTextRange() const override;
 
-	virtual void SetTextRange( const FTextRange& Value ) override;
+	SLATE_API virtual void SetTextRange( const FTextRange& Value ) override;
 
-	virtual int16 GetBaseLine( float Scale ) const override;
+	SLATE_API virtual int16 GetBaseLine( float Scale ) const override;
 
-	virtual int16 GetMaxHeight( float Scale ) const override;
+	SLATE_API virtual int16 GetMaxHeight( float Scale ) const override;
 
-	virtual FVector2D Measure( int32 StartIndex, int32 EndIndex, float Scale, const FRunTextContext& TextContext ) const override;
+	SLATE_API virtual FVector2D Measure( int32 StartIndex, int32 EndIndex, float Scale, const FRunTextContext& TextContext ) const override;
 
-	virtual int8 GetKerning(int32 CurrentIndex, float Scale, const FRunTextContext& TextContext) const override;
+	SLATE_API virtual int8 GetKerning(int32 CurrentIndex, float Scale, const FRunTextContext& TextContext) const override;
 
-	virtual TSharedRef< ILayoutBlock > CreateBlock( int32 StartIndex, int32 EndIndex, FVector2D Size, const FLayoutBlockTextContext& TextContext, const TSharedPtr< IRunRenderer >& Renderer ) override;
+	SLATE_API virtual TSharedRef< ILayoutBlock > CreateBlock( int32 StartIndex, int32 EndIndex, FVector2D Size, const FLayoutBlockTextContext& TextContext, const TSharedPtr< IRunRenderer >& Renderer ) override;
 
-	virtual int32 OnPaint(const FPaintArgs& PaintArgs, const FTextArgs& TextArgs, const FGeometry& AllottedGeometry, const FSlateRect& MyCullingRect, FSlateWindowElementList& OutDrawElements, int32 LayerId, const FWidgetStyle& InWidgetStyle, bool bParentEnabled) const override;
+	SLATE_API virtual int32 OnPaint(const FPaintArgs& PaintArgs, const FTextArgs& TextArgs, const FGeometry& AllottedGeometry, const FSlateRect& MyCullingRect, FSlateWindowElementList& OutDrawElements, int32 LayerId, const FWidgetStyle& InWidgetStyle, bool bParentEnabled) const override;
 
-	virtual const TArray< TSharedRef<SWidget> >& GetChildren() override;
+	SLATE_API virtual const TArray< TSharedRef<SWidget> >& GetChildren() override;
 
-	virtual void ArrangeChildren( const TSharedRef< ILayoutBlock >& Block, const FGeometry& AllottedGeometry, FArrangedChildren& ArrangedChildren ) const override;
+	SLATE_API virtual void ArrangeChildren( const TSharedRef< ILayoutBlock >& Block, const FGeometry& AllottedGeometry, FArrangedChildren& ArrangedChildren ) const override;
 
-	virtual int32 GetTextIndexAt( const TSharedRef< ILayoutBlock >& Block, const FVector2D& Location, float Scale, ETextHitPoint* const OutHitPoint = nullptr ) const override;
+	SLATE_API virtual int32 GetTextIndexAt( const TSharedRef< ILayoutBlock >& Block, const FVector2D& Location, float Scale, ETextHitPoint* const OutHitPoint = nullptr ) const override;
 
-	virtual FVector2D GetLocationAt( const TSharedRef< ILayoutBlock >& Block, int32 Offset, float Scale ) const override;
+	SLATE_API virtual FVector2D GetLocationAt( const TSharedRef< ILayoutBlock >& Block, int32 Offset, float Scale ) const override;
 
 	virtual void BeginLayout() override { Children.Empty(); }
 	virtual void EndLayout() override {}
 
-	virtual void Move(const TSharedRef<FString>& NewText, const FTextRange& NewRange) override;
-	virtual TSharedRef<IRun> Clone() const override;
+	SLATE_API virtual void Move(const TSharedRef<FString>& NewText, const FTextRange& NewRange) override;
+	SLATE_API virtual TSharedRef<IRun> Clone() const override;
 
-	virtual void AppendTextTo(FString& AppendToText) const override;
-	virtual void AppendTextTo(FString& AppendToText, const FTextRange& PartialRange) const override;
+	SLATE_API virtual void AppendTextTo(FString& AppendToText) const override;
+	SLATE_API virtual void AppendTextTo(FString& AppendToText, const FTextRange& PartialRange) const override;
 
-	virtual const FRunInfo& GetRunInfo() const override;
+	SLATE_API virtual const FRunInfo& GetRunInfo() const override;
 
-	virtual ERunAttributes GetRunAttributes() const override;
+	SLATE_API virtual ERunAttributes GetRunAttributes() const override;
 
 protected:
 
-	FSlateHyperlinkRun( const FRunInfo& InRunInfo, const TSharedRef< const FString >& InText, const FHyperlinkStyle& InStyle, FOnClick InNavigateDelegate, FOnGenerateTooltip InTooltipDelegate, FOnGetTooltipText InTooltipTextDelegate );
+	SLATE_API FSlateHyperlinkRun( const FRunInfo& InRunInfo, const TSharedRef< const FString >& InText, const FHyperlinkStyle& InStyle, FOnClick InNavigateDelegate, FOnGenerateTooltip InTooltipDelegate, FOnGetTooltipText InTooltipTextDelegate );
 																										 
-	FSlateHyperlinkRun( const FRunInfo& InRunInfo, const TSharedRef< const FString >& InText, const FHyperlinkStyle& InStyle, FOnClick InNavigateDelegate, FOnGenerateTooltip InTooltipDelegate, FOnGetTooltipText InTooltipTextDelegate, const FTextRange& InRange );
+	SLATE_API FSlateHyperlinkRun( const FRunInfo& InRunInfo, const TSharedRef< const FString >& InText, const FHyperlinkStyle& InStyle, FOnClick InNavigateDelegate, FOnGenerateTooltip InTooltipDelegate, FOnGetTooltipText InTooltipTextDelegate, const FTextRange& InRange );
 
-	FSlateHyperlinkRun( const FSlateHyperlinkRun& Run );
+	SLATE_API FSlateHyperlinkRun( const FSlateHyperlinkRun& Run );
 
 protected:
 
-	void OnNavigate();
+	SLATE_API void OnNavigate();
 
 protected:
 
diff --git a/Engine/Source/Runtime/Slate/Public/Framework/Text/SlateImageRun.h b/Engine/Source/Runtime/Slate/Public/Framework/Text/SlateImageRun.h
--- a/Engine/Source/Runtime/Slate/Public/Framework/Text/SlateImageRun.h
+++ b/Engine/Source/Runtime/Slate/Public/Framework/Text/SlateImageRun.h
@@ -18,63 +18,63 @@ enum class ETextHitPoint : uint8;
 
 #if WITH_FANCY_TEXT
 
-class SLATE_API FSlateImageRun : public ISlateRun, public TSharedFromThis< FSlateImageRun >
+class FSlateImageRun : public ISlateRun, public TSharedFromThis< FSlateImageRun >
 {
 public:
 
-	static TSharedRef< FSlateImageRun > Create( const FRunInfo& InRunInfo, const TSharedRef< const FString >& InText, const FSlateBrush* InImage, const TAttribute<int16>& InBaseline );
-	static TSharedRef< FSlateImageRun > Create( const FRunInfo& InRunInfo, const TSharedRef< const FString >& InText, const FSlateBrush* InImage, const TAttribute<int16>& InBaseline , const FTextRange& InRange );
-	static TSharedRef< FSlateImageRun > Create( const FRunInfo& InRunInfo, const TSharedRef< const FString >& InText, FName InDynamicBrushName, const TAttribute<int16>& InBaseline );
-	static TSharedRef< FSlateImageRun > Create( const FRunInfo& InRunInfo, const TSharedRef< const FString >& InText, FName InDynamicBrushName, const TAttribute<int16>& InBaseline, const FTextRange& InRange );
+	static SLATE_API TSharedRef< FSlateImageRun > Create( const FRunInfo& InRunInfo, const TSharedRef< const FString >& InText, const FSlateBrush* InImage, const TAttribute<int16>& InBaseline );
+	static SLATE_API TSharedRef< FSlateImageRun > Create( const FRunInfo& InRunInfo, const TSharedRef< const FString >& InText, const FSlateBrush* InImage, const TAttribute<int16>& InBaseline , const FTextRange& InRange );
+	static SLATE_API TSharedRef< FSlateImageRun > Create( const FRunInfo& InRunInfo, const TSharedRef< const FString >& InText, FName InDynamicBrushName, const TAttribute<int16>& InBaseline );
+	static SLATE_API TSharedRef< FSlateImageRun > Create( const FRunInfo& InRunInfo, const TSharedRef< const FString >& InText, FName InDynamicBrushName, const TAttribute<int16>& InBaseline, const FTextRange& InRange );
 
 public:
 
-	virtual ~FSlateImageRun();
+	SLATE_API virtual ~FSlateImageRun();
 
-	virtual FTextRange GetTextRange() const override;
-	virtual void SetTextRange( const FTextRange& Value ) override;
+	SLATE_API virtual FTextRange GetTextRange() const override;
+	SLATE_API virtual void SetTextRange( const FTextRange& Value ) override;
 
-	virtual int16 GetBaseLine( float Scale ) const override;
-	virtual int16 GetMaxHeight( float Scale ) const override;
+	SLATE_API virtual int16 GetBaseLine( float Scale ) const override;
+	SLATE_API virtual int16 GetMaxHeight( float Scale ) const override;
 
-	virtual FVector2D Measure( int32 StartIndex, int32 EndIndex, float Scale, const FRunTextContext& TextContext ) const override;
+	SLATE_API virtual FVector2D Measure( int32 StartIndex, int32 EndIndex, float Scale, const FRunTextContext& TextContext ) const override;
 
-	virtual int8 GetKerning(int32 CurrentIndex, float Scale, const FRunTextContext& TextContext) const override;
+	SLATE_API virtual int8 GetKerning(int32 CurrentIndex, float Scale, const FRunTextContext& TextContext) const override;
 
-	virtual TSharedRef< ILayoutBlock > CreateBlock( int32 StartIndex, int32 EndIndex, FVector2D Size, const FLayoutBlockTextContext& TextContext, const TSharedPtr< IRunRenderer >& Renderer ) override;
+	SLATE_API virtual TSharedRef< ILayoutBlock > CreateBlock( int32 StartIndex, int32 EndIndex, FVector2D Size, const FLayoutBlockTextContext& TextContext, const TSharedPtr< IRunRenderer >& Renderer ) override;
 
-	virtual int32 OnPaint(const FPaintArgs& PaintArgs, const FTextArgs& TextArgs, const FGeometry& AllottedGeometry, const FSlateRect& MyCullingRect, FSlateWindowElementList& OutDrawElements, int32 LayerId, const FWidgetStyle& InWidgetStyle, bool bParentEnabled) const override;
+	SLATE_API virtual int32 OnPaint(const FPaintArgs& PaintArgs, const FTextArgs& TextArgs, const FGeometry& AllottedGeometry, const FSlateRect& MyCullingRect, FSlateWindowElementList& OutDrawElements, int32 LayerId, const FWidgetStyle& InWidgetStyle, bool bParentEnabled) const override;
 
-	virtual const TArray< TSharedRef<SWidget> >& GetChildren() override;
+	SLATE_API virtual const TArray< TSharedRef<SWidget> >& GetChildren() override;
 
-	virtual void ArrangeChildren( const TSharedRef< ILayoutBlock >& Block, const FGeometry& AllottedGeometry, FArrangedChildren& ArrangedChildren ) const override;
+	SLATE_API virtual void ArrangeChildren( const TSharedRef< ILayoutBlock >& Block, const FGeometry& AllottedGeometry, FArrangedChildren& ArrangedChildren ) const override;
 
-	virtual int32 GetTextIndexAt( const TSharedRef< ILayoutBlock >& Block, const FVector2D& Location, float Scale, ETextHitPoint* const OutHitPoint = nullptr ) const override;
+	SLATE_API virtual int32 GetTextIndexAt( const TSharedRef< ILayoutBlock >& Block, const FVector2D& Location, float Scale, ETextHitPoint* const OutHitPoint = nullptr ) const override;
 
-	virtual FVector2D GetLocationAt( const TSharedRef< ILayoutBlock >& Block, int32 Offset, float Scale ) const override;
+	SLATE_API virtual FVector2D GetLocationAt( const TSharedRef< ILayoutBlock >& Block, int32 Offset, float Scale ) const override;
 
 	virtual void BeginLayout() override {}
 	virtual void EndLayout() override {}
 
-	virtual void Move(const TSharedRef<FString>& NewText, const FTextRange& NewRange) override;
-	virtual TSharedRef<IRun> Clone() const override;
+	SLATE_API virtual void Move(const TSharedRef<FString>& NewText, const FTextRange& NewRange) override;
+	SLATE_API virtual TSharedRef<IRun> Clone() const override;
 
-	virtual void AppendTextTo(FString& AppendToText) const override;
-	virtual void AppendTextTo(FString& AppendToText, const FTextRange& PartialRange) const override;
+	SLATE_API virtual void AppendTextTo(FString& AppendToText) const override;
+	SLATE_API virtual void AppendTextTo(FString& AppendToText, const FTextRange& PartialRange) const override;
 
-	virtual const FRunInfo& GetRunInfo() const override;
+	SLATE_API virtual const FRunInfo& GetRunInfo() const override;
 
-	virtual ERunAttributes GetRunAttributes() const override;
+	SLATE_API virtual ERunAttributes GetRunAttributes() const override;
 
 protected:
 
-	FSlateImageRun( const FRunInfo& InRunInfo, const TSharedRef< const FString >& InText, const FSlateBrush* InImage, const TAttribute<int16>& InBaseline );
+	SLATE_API FSlateImageRun( const FRunInfo& InRunInfo, const TSharedRef< const FString >& InText, const FSlateBrush* InImage, const TAttribute<int16>& InBaseline );
 
-	FSlateImageRun( const FRunInfo& InRunInfo, const TSharedRef< const FString >& InText, const FSlateBrush* InImage, const TAttribute<int16>& InBaseline, const FTextRange& InRange );
+	SLATE_API FSlateImageRun( const FRunInfo& InRunInfo, const TSharedRef< const FString >& InText, const FSlateBrush* InImage, const TAttribute<int16>& InBaseline, const FTextRange& InRange );
 
-	FSlateImageRun( const FRunInfo& InRunInfo, const TSharedRef< const FString >& InText, FName InDynamicBrushName, const TAttribute<int16>& InBaseline );
+	SLATE_API FSlateImageRun( const FRunInfo& InRunInfo, const TSharedRef< const FString >& InText, FName InDynamicBrushName, const TAttribute<int16>& InBaseline );
 
-	FSlateImageRun( const FRunInfo& InRunInfo, const TSharedRef< const FString >& InText, FName InDynamicBrushName, const TAttribute<int16>& InBaseline, const FTextRange& InRange );
+	SLATE_API FSlateImageRun( const FRunInfo& InRunInfo, const TSharedRef< const FString >& InText, FName InDynamicBrushName, const TAttribute<int16>& InBaseline, const FTextRange& InRange );
 
 private:
 
diff --git a/Engine/Source/Runtime/Slate/Public/Framework/Text/SlatePasswordRun.h b/Engine/Source/Runtime/Slate/Public/Framework/Text/SlatePasswordRun.h
--- a/Engine/Source/Runtime/Slate/Public/Framework/Text/SlatePasswordRun.h
+++ b/Engine/Source/Runtime/Slate/Public/Framework/Text/SlatePasswordRun.h
@@ -11,37 +11,37 @@ class FPaintArgs;
 class FSlateWindowElementList;
 enum class ETextHitPoint : uint8;
 
-class SLATE_API FSlatePasswordRun : public FSlateTextRun
+class FSlatePasswordRun : public FSlateTextRun
 {
 public:
 
-	static TSharedRef< FSlatePasswordRun > Create( const FRunInfo& InRunInfo, const TSharedRef< const FString >& InText, const FTextBlockStyle& Style );
+	static SLATE_API TSharedRef< FSlatePasswordRun > Create( const FRunInfo& InRunInfo, const TSharedRef< const FString >& InText, const FTextBlockStyle& Style );
 
-	static TSharedRef< FSlatePasswordRun > Create( const FRunInfo& InRunInfo, const TSharedRef< const FString >& InText, const FTextBlockStyle& Style, const FTextRange& InRange );
+	static SLATE_API TSharedRef< FSlatePasswordRun > Create( const FRunInfo& InRunInfo, const TSharedRef< const FString >& InText, const FTextBlockStyle& Style, const FTextRange& InRange );
 
 public:
 
 	virtual ~FSlatePasswordRun() {}
 
-	virtual FVector2D Measure( int32 StartIndex, int32 EndIndex, float Scale, const FRunTextContext& TextContext ) const override;
-	virtual int8 GetKerning(int32 CurrentIndex, float Scale, const FRunTextContext& TextContext) const override;
+	SLATE_API virtual FVector2D Measure( int32 StartIndex, int32 EndIndex, float Scale, const FRunTextContext& TextContext ) const override;
+	SLATE_API virtual int8 GetKerning(int32 CurrentIndex, float Scale, const FRunTextContext& TextContext) const override;
 
-	virtual int32 OnPaint(const FPaintArgs& PaintArgs, const FTextArgs& TextArgs, const FGeometry& AllottedGeometry, const FSlateRect& MyCullingRect, FSlateWindowElementList& OutDrawElements, int32 LayerId, const FWidgetStyle& InWidgetStyle, bool bParentEnabled) const override;
+	SLATE_API virtual int32 OnPaint(const FPaintArgs& PaintArgs, const FTextArgs& TextArgs, const FGeometry& AllottedGeometry, const FSlateRect& MyCullingRect, FSlateWindowElementList& OutDrawElements, int32 LayerId, const FWidgetStyle& InWidgetStyle, bool bParentEnabled) const override;
 
-	virtual int32 GetTextIndexAt( const TSharedRef< ILayoutBlock >& Block, const FVector2D& Location, float Scale, ETextHitPoint* const OutHitPoint = nullptr ) const override;
+	SLATE_API virtual int32 GetTextIndexAt( const TSharedRef< ILayoutBlock >& Block, const FVector2D& Location, float Scale, ETextHitPoint* const OutHitPoint = nullptr ) const override;
 
-	virtual FVector2D GetLocationAt(const TSharedRef< ILayoutBlock >& Block, int32 Offset, float Scale) const override;
+	SLATE_API virtual FVector2D GetLocationAt(const TSharedRef< ILayoutBlock >& Block, int32 Offset, float Scale) const override;
 
-	virtual TSharedRef<IRun> Clone() const override;
+	SLATE_API virtual TSharedRef<IRun> Clone() const override;
 
 protected:
 
-	FSlatePasswordRun( const FRunInfo& InRunInfo, const TSharedRef< const FString >& InText, const FTextBlockStyle& InStyle );
+	SLATE_API FSlatePasswordRun( const FRunInfo& InRunInfo, const TSharedRef< const FString >& InText, const FTextBlockStyle& InStyle );
 
-	FSlatePasswordRun( const FRunInfo& InRunInfo, const TSharedRef< const FString >& InText, const FTextBlockStyle& InStyle, const FTextRange& InRange );
+	SLATE_API FSlatePasswordRun( const FRunInfo& InRunInfo, const TSharedRef< const FString >& InText, const FTextBlockStyle& InStyle, const FTextRange& InRange );
 
-	FSlatePasswordRun( const FSlatePasswordRun& Run );
+	SLATE_API FSlatePasswordRun( const FSlatePasswordRun& Run );
 
-	static TCHAR GetPasswordChar();
-	static FString BuildPasswordString(const int32 InLength);
+	static SLATE_API TCHAR GetPasswordChar();
+	static SLATE_API FString BuildPasswordString(const int32 InLength);
 };
diff --git a/Engine/Source/Runtime/Slate/Public/Framework/Text/SlateTextHighlightRunRenderer.h b/Engine/Source/Runtime/Slate/Public/Framework/Text/SlateTextHighlightRunRenderer.h
--- a/Engine/Source/Runtime/Slate/Public/Framework/Text/SlateTextHighlightRunRenderer.h
+++ b/Engine/Source/Runtime/Slate/Public/Framework/Text/SlateTextHighlightRunRenderer.h
@@ -11,18 +11,18 @@ class FPaintArgs;
 class FSlateWindowElementList;
 struct FTextBlockStyle;
 
-class SLATE_API FSlateTextHighlightRunRenderer : public ISlateRunRenderer
+class FSlateTextHighlightRunRenderer : public ISlateRunRenderer
 {
 public:
 
-	static TSharedRef< FSlateTextHighlightRunRenderer > Create();
+	static SLATE_API TSharedRef< FSlateTextHighlightRunRenderer > Create();
 
 	virtual ~FSlateTextHighlightRunRenderer() {}
 
-	virtual int32 OnPaint( const FPaintArgs& Args, const FTextLayout::FLineView& Line, const TSharedRef< ISlateRun >& Run, const TSharedRef< ILayoutBlock >& Block, const FTextBlockStyle& DefaultStyle, const FGeometry& AllottedGeometry, const FSlateRect& MyCullingRect, FSlateWindowElementList& OutDrawElements, int32 LayerId, const FWidgetStyle& InWidgetStyle, bool bParentEnabled ) const override;
+	SLATE_API virtual int32 OnPaint( const FPaintArgs& Args, const FTextLayout::FLineView& Line, const TSharedRef< ISlateRun >& Run, const TSharedRef< ILayoutBlock >& Block, const FTextBlockStyle& DefaultStyle, const FGeometry& AllottedGeometry, const FSlateRect& MyCullingRect, FSlateWindowElementList& OutDrawElements, int32 LayerId, const FWidgetStyle& InWidgetStyle, bool bParentEnabled ) const override;
 
 private:
 
-	FSlateTextHighlightRunRenderer();
+	SLATE_API FSlateTextHighlightRunRenderer();
 
 };
diff --git a/Engine/Source/Runtime/Slate/Public/Framework/Text/SlateTextLayout.h b/Engine/Source/Runtime/Slate/Public/Framework/Text/SlateTextLayout.h
--- a/Engine/Source/Runtime/Slate/Public/Framework/Text/SlateTextLayout.h
+++ b/Engine/Source/Runtime/Slate/Public/Framework/Text/SlateTextLayout.h
@@ -13,34 +13,34 @@ class FArrangedChildren;
 class FPaintArgs;
 class FSlateWindowElementList;
 
-class SLATE_API FSlateTextLayout : public FTextLayout
+class FSlateTextLayout : public FTextLayout
 {
 public:
 
-	static TSharedRef< FSlateTextLayout > Create(SWidget* InOwner, FTextBlockStyle InDefaultTextStyle);
+	static SLATE_API TSharedRef< FSlateTextLayout > Create(SWidget* InOwner, FTextBlockStyle InDefaultTextStyle);
 
-	FChildren* GetChildren();
+	SLATE_API FChildren* GetChildren();
 
-	virtual void ArrangeChildren( const FGeometry& AllottedGeometry, FArrangedChildren& ArrangedChildren ) const;
+	SLATE_API virtual void ArrangeChildren( const FGeometry& AllottedGeometry, FArrangedChildren& ArrangedChildren ) const;
 
-	virtual int32 OnPaint( const FPaintArgs& Args, const FGeometry& AllottedGeometry, const FSlateRect& MyCullingRect, FSlateWindowElementList& OutDrawElements, int32 LayerId, const FWidgetStyle& InWidgetStyle, bool bParentEnabled ) const;
+	SLATE_API virtual int32 OnPaint( const FPaintArgs& Args, const FGeometry& AllottedGeometry, const FSlateRect& MyCullingRect, FSlateWindowElementList& OutDrawElements, int32 LayerId, const FWidgetStyle& InWidgetStyle, bool bParentEnabled ) const;
 
-	virtual void EndLayout() override;
+	SLATE_API virtual void EndLayout() override;
 
-	void SetDefaultTextStyle(FTextBlockStyle InDefaultTextStyle);
-	const FTextBlockStyle& GetDefaultTextStyle() const;
+	SLATE_API void SetDefaultTextStyle(FTextBlockStyle InDefaultTextStyle);
+	SLATE_API const FTextBlockStyle& GetDefaultTextStyle() const;
 
-	void SetIsPassword(const TAttribute<bool>& InIsPassword);
+	SLATE_API void SetIsPassword(const TAttribute<bool>& InIsPassword);
 
 protected:
 
-	FSlateTextLayout(SWidget* InOwner, FTextBlockStyle InDefaultTextStyle);
+	SLATE_API FSlateTextLayout(SWidget* InOwner, FTextBlockStyle InDefaultTextStyle);
 
-	virtual int32 OnPaintHighlights(const FPaintArgs& Args, const FTextLayout::FLineView& LineView, const TArray<FLineViewHighlight>& Highlights, const FTextBlockStyle& DefaultTextStyle, const FGeometry& AllottedGeometry, const FSlateRect& ClippingRect, FSlateWindowElementList& OutDrawElements, const int32 LayerId, const FWidgetStyle& InWidgetStyle, bool bParentEnabled) const;
+	SLATE_API virtual int32 OnPaintHighlights(const FPaintArgs& Args, const FTextLayout::FLineView& LineView, const TArray<FLineViewHighlight>& Highlights, const FTextBlockStyle& DefaultTextStyle, const FGeometry& AllottedGeometry, const FSlateRect& ClippingRect, FSlateWindowElementList& OutDrawElements, const int32 LayerId, const FWidgetStyle& InWidgetStyle, bool bParentEnabled) const;
 
-	virtual void AggregateChildren();
+	SLATE_API virtual void AggregateChildren();
 
-	virtual TSharedRef<IRun> CreateDefaultTextRun(const TSharedRef<FString>& NewText, const FTextRange& NewRange) const override;
+	SLATE_API virtual TSharedRef<IRun> CreateDefaultTextRun(const TSharedRef<FString>& NewText, const FTextRange& NewRange) const override;
 
 protected:
 	/** Default style used by the TextLayout */
diff --git a/Engine/Source/Runtime/Slate/Public/Framework/Text/SlateTextRun.h b/Engine/Source/Runtime/Slate/Public/Framework/Text/SlateTextRun.h
--- a/Engine/Source/Runtime/Slate/Public/Framework/Text/SlateTextRun.h
+++ b/Engine/Source/Runtime/Slate/Public/Framework/Text/SlateTextRun.h
@@ -14,60 +14,60 @@ class FPaintArgs;
 class FSlateWindowElementList;
 enum class ETextHitPoint : uint8;
 
-class SLATE_API FSlateTextRun : public ISlateRun, public TSharedFromThis< FSlateTextRun >
+class FSlateTextRun : public ISlateRun, public TSharedFromThis< FSlateTextRun >
 {
 public:
 
-	static TSharedRef< FSlateTextRun > Create( const FRunInfo& InRunInfo, const TSharedRef< const FString >& InText, const FTextBlockStyle& Style );
+	static SLATE_API TSharedRef< FSlateTextRun > Create( const FRunInfo& InRunInfo, const TSharedRef< const FString >& InText, const FTextBlockStyle& Style );
 
-	static TSharedRef< FSlateTextRun > Create( const FRunInfo& InRunInfo, const TSharedRef< const FString >& InText, const FTextBlockStyle& Style, const FTextRange& InRange );
+	static SLATE_API TSharedRef< FSlateTextRun > Create( const FRunInfo& InRunInfo, const TSharedRef< const FString >& InText, const FTextBlockStyle& Style, const FTextRange& InRange );
 
 public:
 
 	virtual ~FSlateTextRun() {}
 
-	virtual FTextRange GetTextRange() const override;
-	virtual void SetTextRange( const FTextRange& Value ) override;
+	SLATE_API virtual FTextRange GetTextRange() const override;
+	SLATE_API virtual void SetTextRange( const FTextRange& Value ) override;
 
-	virtual int16 GetBaseLine( float Scale ) const override;
-	virtual int16 GetMaxHeight( float Scale ) const override;
-	virtual FVector2D Measure( int32 StartIndex, int32 EndIndex, float Scale, const FRunTextContext& TextContext ) const override;
-	virtual int8 GetKerning(int32 CurrentIndex, float Scale, const FRunTextContext& TextContext) const override;
+	SLATE_API virtual int16 GetBaseLine( float Scale ) const override;
+	SLATE_API virtual int16 GetMaxHeight( float Scale ) const override;
+	SLATE_API virtual FVector2D Measure( int32 StartIndex, int32 EndIndex, float Scale, const FRunTextContext& TextContext ) const override;
+	SLATE_API virtual int8 GetKerning(int32 CurrentIndex, float Scale, const FRunTextContext& TextContext) const override;
 
-	virtual TSharedRef< ILayoutBlock > CreateBlock( int32 StartIndex, int32 EndIndex, FVector2D Size, const FLayoutBlockTextContext& TextContext, const TSharedPtr< IRunRenderer >& Renderer ) override;
+	SLATE_API virtual TSharedRef< ILayoutBlock > CreateBlock( int32 StartIndex, int32 EndIndex, FVector2D Size, const FLayoutBlockTextContext& TextContext, const TSharedPtr< IRunRenderer >& Renderer ) override;
 
-	virtual int32 OnPaint(const FPaintArgs& PaintArgs, const FTextArgs& TextArgs, const FGeometry& AllottedGeometry, const FSlateRect& MyCullingRect, FSlateWindowElementList& OutDrawElements, int32 LayerId, const FWidgetStyle& InWidgetStyle, bool bParentEnabled) const override;
+	SLATE_API virtual int32 OnPaint(const FPaintArgs& PaintArgs, const FTextArgs& TextArgs, const FGeometry& AllottedGeometry, const FSlateRect& MyCullingRect, FSlateWindowElementList& OutDrawElements, int32 LayerId, const FWidgetStyle& InWidgetStyle, bool bParentEnabled) const override;
 
-	virtual const TArray< TSharedRef<SWidget> >& GetChildren() override;
+	SLATE_API virtual const TArray< TSharedRef<SWidget> >& GetChildren() override;
 
-	virtual void ArrangeChildren( const TSharedRef< ILayoutBlock >& Block, const FGeometry& AllottedGeometry, FArrangedChildren& ArrangedChildren ) const override;
+	SLATE_API virtual void ArrangeChildren( const TSharedRef< ILayoutBlock >& Block, const FGeometry& AllottedGeometry, FArrangedChildren& ArrangedChildren ) const override;
 
 	virtual void BeginLayout() override {}
 	virtual void EndLayout() override {}
 
-	virtual int32 GetTextIndexAt( const TSharedRef< ILayoutBlock >& Block, const FVector2D& Location, float Scale, ETextHitPoint* const OutHitPoint = nullptr ) const override;
+	SLATE_API virtual int32 GetTextIndexAt( const TSharedRef< ILayoutBlock >& Block, const FVector2D& Location, float Scale, ETextHitPoint* const OutHitPoint = nullptr ) const override;
 
-	virtual FVector2D GetLocationAt(const TSharedRef< ILayoutBlock >& Block, int32 Offset, float Scale) const override;
+	SLATE_API virtual FVector2D GetLocationAt(const TSharedRef< ILayoutBlock >& Block, int32 Offset, float Scale) const override;
 
-	virtual void Move(const TSharedRef<FString>& NewText, const FTextRange& NewRange) override;
-	virtual TSharedRef<IRun> Clone() const override;
+	SLATE_API virtual void Move(const TSharedRef<FString>& NewText, const FTextRange& NewRange) override;
+	SLATE_API virtual TSharedRef<IRun> Clone() const override;
 
-	virtual void AppendTextTo(FString& Text) const override;
-	virtual void AppendTextTo(FString& AppendToText, const FTextRange& PartialRange) const override;
+	SLATE_API virtual void AppendTextTo(FString& Text) const override;
+	SLATE_API virtual void AppendTextTo(FString& AppendToText, const FTextRange& PartialRange) const override;
 
-	virtual const FRunInfo& GetRunInfo() const override;
+	SLATE_API virtual const FRunInfo& GetRunInfo() const override;
 
-	virtual ERunAttributes GetRunAttributes() const override;
+	SLATE_API virtual ERunAttributes GetRunAttributes() const override;
 
-	void ApplyFontSizeMultiplierOnTextStyle(float FontSizeMultiplier);
+	SLATE_API void ApplyFontSizeMultiplierOnTextStyle(float FontSizeMultiplier);
 
 protected:
 
-	FSlateTextRun( const FRunInfo& InRunInfo, const TSharedRef< const FString >& InText, const FTextBlockStyle& InStyle );
+	SLATE_API FSlateTextRun( const FRunInfo& InRunInfo, const TSharedRef< const FString >& InText, const FTextBlockStyle& InStyle );
 
-	FSlateTextRun( const FRunInfo& InRunInfo, const TSharedRef< const FString >& InText, const FTextBlockStyle& InStyle, const FTextRange& InRange );
+	SLATE_API FSlateTextRun( const FRunInfo& InRunInfo, const TSharedRef< const FString >& InText, const FTextBlockStyle& InStyle, const FTextRange& InRange );
 
-	FSlateTextRun( const FSlateTextRun& Run );
+	SLATE_API FSlateTextRun( const FSlateTextRun& Run );
 
 protected:
 
diff --git a/Engine/Source/Runtime/Slate/Public/Framework/Text/SlateTextUnderlineLineHighlighter.h b/Engine/Source/Runtime/Slate/Public/Framework/Text/SlateTextUnderlineLineHighlighter.h
--- a/Engine/Source/Runtime/Slate/Public/Framework/Text/SlateTextUnderlineLineHighlighter.h
+++ b/Engine/Source/Runtime/Slate/Public/Framework/Text/SlateTextUnderlineLineHighlighter.h
@@ -15,15 +15,15 @@ struct FGeometry;
 struct FTextBlockStyle;
 
 /** Run highlighter used to draw lines */
-class SLATE_API ISlateTextLineHighlighter : public ISlateLineHighlighter
+class ISlateTextLineHighlighter : public ISlateLineHighlighter
 {
 public:
-	virtual int32 OnPaint(const FPaintArgs& Args, const FTextLayout::FLineView& Line, const float OffsetX, const float Width, const FTextBlockStyle& DefaultStyle, const FGeometry& AllottedGeometry, const FSlateRect& MyCullingRect, FSlateWindowElementList& OutDrawElements, int32 LayerId, const FWidgetStyle& InWidgetStyle, bool bParentEnabled) const override;
+	SLATE_API virtual int32 OnPaint(const FPaintArgs& Args, const FTextLayout::FLineView& Line, const float OffsetX, const float Width, const FTextBlockStyle& DefaultStyle, const FGeometry& AllottedGeometry, const FSlateRect& MyCullingRect, FSlateWindowElementList& OutDrawElements, int32 LayerId, const FWidgetStyle& InWidgetStyle, bool bParentEnabled) const override;
 
 	static const int32 DefaultZIndex = 1;
 
 protected:
-	ISlateTextLineHighlighter(const FSlateBrush& InLineBrush, const FSlateFontInfo& InFontInfo, const FSlateColor InColorAndOpacity, const UE::Slate::FDeprecateVector2DParameter InShadowOffset, const FLinearColor InShadowColorAndOpacity);
+	SLATE_API ISlateTextLineHighlighter(const FSlateBrush& InLineBrush, const FSlateFontInfo& InFontInfo, const FSlateColor InColorAndOpacity, const UE::Slate::FDeprecateVector2DParameter InShadowOffset, const FLinearColor InShadowColorAndOpacity);
 
 	virtual void GetLineMetrics(const float InFontScale, int16& OutLinePos, int16& OutLineThickness) const = 0;
 
@@ -44,25 +44,25 @@ protected:
 };
 
 /** Run highlighter used to draw underlines */
-class SLATE_API FSlateTextUnderlineLineHighlighter : public ISlateTextLineHighlighter
+class FSlateTextUnderlineLineHighlighter : public ISlateTextLineHighlighter
 {
 public:
-	static TSharedRef<FSlateTextUnderlineLineHighlighter> Create(const FSlateBrush& InUnderlineBrush, const FSlateFontInfo& InFontInfo, const FSlateColor InColorAndOpacity, const UE::Slate::FDeprecateVector2DParameter InShadowOffset, const FLinearColor InShadowColorAndOpacity);
+	static SLATE_API TSharedRef<FSlateTextUnderlineLineHighlighter> Create(const FSlateBrush& InUnderlineBrush, const FSlateFontInfo& InFontInfo, const FSlateColor InColorAndOpacity, const UE::Slate::FDeprecateVector2DParameter InShadowOffset, const FLinearColor InShadowColorAndOpacity);
 
 protected:
-	FSlateTextUnderlineLineHighlighter(const FSlateBrush& InUnderlineBrush, const FSlateFontInfo& InFontInfo, const FSlateColor InColorAndOpacity, const UE::Slate::FDeprecateVector2DParameter InShadowOffset, const FLinearColor InShadowColorAndOpacity);
+	SLATE_API FSlateTextUnderlineLineHighlighter(const FSlateBrush& InUnderlineBrush, const FSlateFontInfo& InFontInfo, const FSlateColor InColorAndOpacity, const UE::Slate::FDeprecateVector2DParameter InShadowOffset, const FLinearColor InShadowColorAndOpacity);
 
-	virtual void GetLineMetrics(const float InFontScale, int16& OutLinePos, int16& OutLineThickness) const override;
+	SLATE_API virtual void GetLineMetrics(const float InFontScale, int16& OutLinePos, int16& OutLineThickness) const override;
 };
 
 /** Run highlighter used to draw strikes */
-class SLATE_API FSlateTextStrikeLineHighlighter : public ISlateTextLineHighlighter
+class FSlateTextStrikeLineHighlighter : public ISlateTextLineHighlighter
 {
 public:
-	static TSharedRef<FSlateTextStrikeLineHighlighter> Create(const FSlateBrush& InUnderlineBrush, const FSlateFontInfo& InFontInfo, const FSlateColor InColorAndOpacity, const UE::Slate::FDeprecateVector2DParameter InShadowOffset, const FLinearColor InShadowColorAndOpacity);
+	static SLATE_API TSharedRef<FSlateTextStrikeLineHighlighter> Create(const FSlateBrush& InUnderlineBrush, const FSlateFontInfo& InFontInfo, const FSlateColor InColorAndOpacity, const UE::Slate::FDeprecateVector2DParameter InShadowOffset, const FLinearColor InShadowColorAndOpacity);
 
 protected:
-	FSlateTextStrikeLineHighlighter(const FSlateBrush& InStrikeBrush, const FSlateFontInfo& InFontInfo, const FSlateColor InColorAndOpacity, const UE::Slate::FDeprecateVector2DParameter InShadowOffset, const FLinearColor InShadowColorAndOpacity);
+	SLATE_API FSlateTextStrikeLineHighlighter(const FSlateBrush& InStrikeBrush, const FSlateFontInfo& InFontInfo, const FSlateColor InColorAndOpacity, const UE::Slate::FDeprecateVector2DParameter InShadowOffset, const FLinearColor InShadowColorAndOpacity);
 
-	virtual void GetLineMetrics(const float InFontScale, int16& OutLinePos, int16& OutLineThickness) const override;
+	SLATE_API virtual void GetLineMetrics(const float InFontScale, int16& OutLinePos, int16& OutLineThickness) const override;
 };
diff --git a/Engine/Source/Runtime/Slate/Public/Framework/Text/SlateWidgetRun.h b/Engine/Source/Runtime/Slate/Public/Framework/Text/SlateWidgetRun.h
--- a/Engine/Source/Runtime/Slate/Public/Framework/Text/SlateWidgetRun.h
+++ b/Engine/Source/Runtime/Slate/Public/Framework/Text/SlateWidgetRun.h
@@ -17,7 +17,7 @@ enum class ETextHitPoint : uint8;
 
 #if WITH_FANCY_TEXT
 
-class SLATE_API FSlateWidgetRun : public ISlateRun, public TSharedFromThis< FSlateWidgetRun >
+class FSlateWidgetRun : public ISlateRun, public TSharedFromThis< FSlateWidgetRun >
 {
 public:
 
@@ -36,40 +36,40 @@ public:
 		}
 	};
 
-	static TSharedRef< FSlateWidgetRun > Create(const TSharedRef<class FTextLayout>& TextLayout, const FRunInfo& InRunInfo, const TSharedRef< const FString >& InText, const FWidgetRunInfo& InWidgetInfo);
+	static SLATE_API TSharedRef< FSlateWidgetRun > Create(const TSharedRef<class FTextLayout>& TextLayout, const FRunInfo& InRunInfo, const TSharedRef< const FString >& InText, const FWidgetRunInfo& InWidgetInfo);
 
-	static TSharedRef< FSlateWidgetRun > Create(const TSharedRef<class FTextLayout>& TextLayout, const FRunInfo& InRunInfo, const TSharedRef< const FString >& InText, const FWidgetRunInfo& InWidgetInfo, const FTextRange& InRange);
+	static SLATE_API TSharedRef< FSlateWidgetRun > Create(const TSharedRef<class FTextLayout>& TextLayout, const FRunInfo& InRunInfo, const TSharedRef< const FString >& InText, const FWidgetRunInfo& InWidgetInfo, const FTextRange& InRange);
 
 	virtual ~FSlateWidgetRun() {}
 
-	virtual FTextRange GetTextRange() const override;
-	virtual void SetTextRange( const FTextRange& Value ) override;
-	virtual int16 GetBaseLine( float Scale ) const override;
-	virtual int16 GetMaxHeight( float Scale ) const override;
-	virtual FVector2D Measure( int32 StartIndex, int32 EndIndex, float Scale, const FRunTextContext& TextContext ) const override;
-	virtual int8 GetKerning(int32 CurrentIndex, float Scale, const FRunTextContext& TextContext) const override;
-	virtual TSharedRef< ILayoutBlock > CreateBlock( int32 StartIndex, int32 EndIndex, FVector2D Size, const FLayoutBlockTextContext& TextContext, const TSharedPtr< IRunRenderer >& Renderer ) override;
-	virtual int32 OnPaint(const FPaintArgs& PaintArgs, const FTextArgs& TextArgs, const FGeometry& AllottedGeometry, const FSlateRect& MyCullingRect, FSlateWindowElementList& OutDrawElements, int32 LayerId, const FWidgetStyle& InWidgetStyle, bool bParentEnabled) const override;
-	virtual const TArray< TSharedRef<SWidget> >& GetChildren() override;
-	virtual void ArrangeChildren( const TSharedRef< ILayoutBlock >& Block, const FGeometry& AllottedGeometry, FArrangedChildren& ArrangedChildren ) const override;
-	virtual int32 GetTextIndexAt( const TSharedRef< ILayoutBlock >& Block, const FVector2D& Location, float Scale, ETextHitPoint* const OutHitPoint = nullptr ) const override;
-	virtual FVector2D GetLocationAt( const TSharedRef< ILayoutBlock >& Block, int32 Offset, float Scale ) const override;
+	SLATE_API virtual FTextRange GetTextRange() const override;
+	SLATE_API virtual void SetTextRange( const FTextRange& Value ) override;
+	SLATE_API virtual int16 GetBaseLine( float Scale ) const override;
+	SLATE_API virtual int16 GetMaxHeight( float Scale ) const override;
+	SLATE_API virtual FVector2D Measure( int32 StartIndex, int32 EndIndex, float Scale, const FRunTextContext& TextContext ) const override;
+	SLATE_API virtual int8 GetKerning(int32 CurrentIndex, float Scale, const FRunTextContext& TextContext) const override;
+	SLATE_API virtual TSharedRef< ILayoutBlock > CreateBlock( int32 StartIndex, int32 EndIndex, FVector2D Size, const FLayoutBlockTextContext& TextContext, const TSharedPtr< IRunRenderer >& Renderer ) override;
+	SLATE_API virtual int32 OnPaint(const FPaintArgs& PaintArgs, const FTextArgs& TextArgs, const FGeometry& AllottedGeometry, const FSlateRect& MyCullingRect, FSlateWindowElementList& OutDrawElements, int32 LayerId, const FWidgetStyle& InWidgetStyle, bool bParentEnabled) const override;
+	SLATE_API virtual const TArray< TSharedRef<SWidget> >& GetChildren() override;
+	SLATE_API virtual void ArrangeChildren( const TSharedRef< ILayoutBlock >& Block, const FGeometry& AllottedGeometry, FArrangedChildren& ArrangedChildren ) const override;
+	SLATE_API virtual int32 GetTextIndexAt( const TSharedRef< ILayoutBlock >& Block, const FVector2D& Location, float Scale, ETextHitPoint* const OutHitPoint = nullptr ) const override;
+	SLATE_API virtual FVector2D GetLocationAt( const TSharedRef< ILayoutBlock >& Block, int32 Offset, float Scale ) const override;
 	virtual void BeginLayout() override {}
 	virtual void EndLayout() override {}
-	virtual void Move(const TSharedRef<FString>& NewText, const FTextRange& NewRange) override;
-	virtual TSharedRef<IRun> Clone() const override;
-	virtual void AppendTextTo(FString& Text) const override;
-	virtual void AppendTextTo(FString& AppendToText, const FTextRange& PartialRange) const override;
-	virtual const FRunInfo& GetRunInfo() const override;
-	virtual ERunAttributes GetRunAttributes() const override;
+	SLATE_API virtual void Move(const TSharedRef<FString>& NewText, const FTextRange& NewRange) override;
+	SLATE_API virtual TSharedRef<IRun> Clone() const override;
+	SLATE_API virtual void AppendTextTo(FString& Text) const override;
+	SLATE_API virtual void AppendTextTo(FString& AppendToText, const FTextRange& PartialRange) const override;
+	SLATE_API virtual const FRunInfo& GetRunInfo() const override;
+	SLATE_API virtual ERunAttributes GetRunAttributes() const override;
 
 protected:
 
-	FSlateWidgetRun(const TSharedRef<class FTextLayout>& TextLayout, const FRunInfo& InRunInfo, const TSharedRef< const FString >& InText, const FWidgetRunInfo& InWidgetInfo);
+	SLATE_API FSlateWidgetRun(const TSharedRef<class FTextLayout>& TextLayout, const FRunInfo& InRunInfo, const TSharedRef< const FString >& InText, const FWidgetRunInfo& InWidgetInfo);
 
-	FSlateWidgetRun(const TSharedRef<class FTextLayout>& TextLayout, const FRunInfo& InRunInfo, const TSharedRef< const FString >& InText, const FWidgetRunInfo& InWidgetInfo, const FTextRange& InRange);
+	SLATE_API FSlateWidgetRun(const TSharedRef<class FTextLayout>& TextLayout, const FRunInfo& InRunInfo, const TSharedRef< const FString >& InText, const FWidgetRunInfo& InWidgetInfo, const FTextRange& InRange);
 
-	FSlateWidgetRun( const FSlateWidgetRun& Run );
+	SLATE_API FSlateWidgetRun( const FSlateWidgetRun& Run );
 
 private:
 
diff --git a/Engine/Source/Runtime/Slate/Public/Framework/Text/SyntaxHighlighterTextLayoutMarshaller.h b/Engine/Source/Runtime/Slate/Public/Framework/Text/SyntaxHighlighterTextLayoutMarshaller.h
--- a/Engine/Source/Runtime/Slate/Public/Framework/Text/SyntaxHighlighterTextLayoutMarshaller.h
+++ b/Engine/Source/Runtime/Slate/Public/Framework/Text/SyntaxHighlighterTextLayoutMarshaller.h
@@ -18,24 +18,24 @@ class FTextLayout;
 /**
  * Get/set the raw text to/from a text layout, and also inject syntax highlighting
  */
-class SLATE_API FSyntaxHighlighterTextLayoutMarshaller : public FPlainTextLayoutMarshaller
+class FSyntaxHighlighterTextLayoutMarshaller : public FPlainTextLayoutMarshaller
 {
 public:
 
-	virtual ~FSyntaxHighlighterTextLayoutMarshaller();
+	SLATE_API virtual ~FSyntaxHighlighterTextLayoutMarshaller();
 
 	// ITextLayoutMarshaller
-	virtual void SetText(const FString& SourceString, FTextLayout& TargetTextLayout) override;
-	virtual bool RequiresLiveUpdate() const override;
+	SLATE_API virtual void SetText(const FString& SourceString, FTextLayout& TargetTextLayout) override;
+	SLATE_API virtual bool RequiresLiveUpdate() const override;
 
-	void EnableSyntaxHighlighting(const bool bEnable);
-	bool IsSyntaxHighlightingEnabled() const;
+	SLATE_API void EnableSyntaxHighlighting(const bool bEnable);
+	SLATE_API bool IsSyntaxHighlightingEnabled() const;
 
 protected:
 
 	virtual void ParseTokens(const FString& SourceString, FTextLayout& TargetTextLayout, TArray<ISyntaxTokenizer::FTokenizedLine> TokenizedLines) = 0;
 
-	FSyntaxHighlighterTextLayoutMarshaller(TSharedPtr< ISyntaxTokenizer > InTokenizer);
+	SLATE_API FSyntaxHighlighterTextLayoutMarshaller(TSharedPtr< ISyntaxTokenizer > InTokenizer);
 
 	/** Tokenizer used to style the text */
 	TSharedPtr< ISyntaxTokenizer > Tokenizer;
@@ -48,7 +48,7 @@ protected:
 /**
  * Get/set the raw text to/from a text layout, and also inject syntax highlighting for our rich-text markup
  */
-class SLATE_API FRichTextSyntaxHighlighterTextLayoutMarshaller : public FSyntaxHighlighterTextLayoutMarshaller
+class FRichTextSyntaxHighlighterTextLayoutMarshaller : public FSyntaxHighlighterTextLayoutMarshaller
 {
 public:
 
@@ -79,15 +79,15 @@ public:
 		FTextBlockStyle NodeAttributeValueTextStyle;
 	};
 
-	static TSharedRef< FRichTextSyntaxHighlighterTextLayoutMarshaller > Create(const FSyntaxTextStyle& InSyntaxTextStyle);
+	static SLATE_API TSharedRef< FRichTextSyntaxHighlighterTextLayoutMarshaller > Create(const FSyntaxTextStyle& InSyntaxTextStyle);
 
-	virtual ~FRichTextSyntaxHighlighterTextLayoutMarshaller();
+	SLATE_API virtual ~FRichTextSyntaxHighlighterTextLayoutMarshaller();
 
 protected:
 
-	virtual void ParseTokens(const FString& SourceString, FTextLayout& TargetTextLayout, TArray<ISyntaxTokenizer::FTokenizedLine> TokenizedLines) override;
+	SLATE_API virtual void ParseTokens(const FString& SourceString, FTextLayout& TargetTextLayout, TArray<ISyntaxTokenizer::FTokenizedLine> TokenizedLines) override;
 
-	FRichTextSyntaxHighlighterTextLayoutMarshaller(TSharedPtr< ISyntaxTokenizer > InTokenizer, const FSyntaxTextStyle& InSyntaxTextStyle);
+	SLATE_API FRichTextSyntaxHighlighterTextLayoutMarshaller(TSharedPtr< ISyntaxTokenizer > InTokenizer, const FSyntaxTextStyle& InSyntaxTextStyle);
 
 	/** Styles used to display the text */
 	FSyntaxTextStyle SyntaxTextStyle;
diff --git a/Engine/Source/Runtime/Slate/Public/Framework/Text/SyntaxTokenizer.h b/Engine/Source/Runtime/Slate/Public/Framework/Text/SyntaxTokenizer.h
--- a/Engine/Source/Runtime/Slate/Public/Framework/Text/SyntaxTokenizer.h
+++ b/Engine/Source/Runtime/Slate/Public/Framework/Text/SyntaxTokenizer.h
@@ -8,7 +8,7 @@
 #include "Templates/SharedPointer.h"
 #include "Templates/UnrealTemplate.h"
 
-class SLATE_API ISyntaxTokenizer
+class ISyntaxTokenizer
 {
 public:
 	/** Denotes the type of token */
@@ -48,7 +48,7 @@ public:
 /**
  * Tokenize the text based upon the given rule set
  */
-class SLATE_API FSyntaxTokenizer : public ISyntaxTokenizer
+class FSyntaxTokenizer : public ISyntaxTokenizer
 {
 public:
 	/** Rule used to match syntax token types */
@@ -66,11 +66,11 @@ public:
 	 * Create a new tokenizer which will use the given rules to match syntax tokens
 	 * @param InRules Rules to control the tokenizer, processed in-order so the most greedy matches must come first
 	 */
-	static TSharedRef< FSyntaxTokenizer > Create(TArray<FRule> InRules);
+	static SLATE_API TSharedRef< FSyntaxTokenizer > Create(TArray<FRule> InRules);
 
-	virtual ~FSyntaxTokenizer();
+	SLATE_API virtual ~FSyntaxTokenizer();
 
-	virtual void Process(TArray<FTokenizedLine>& OutTokenizedLines, const FString& Input) override;
+	SLATE_API virtual void Process(TArray<FTokenizedLine>& OutTokenizedLines, const FString& Input) override;
 
 private:
 
diff --git a/Engine/Source/Runtime/Slate/Public/Framework/Text/TextDecorators.h b/Engine/Source/Runtime/Slate/Public/Framework/Text/TextDecorators.h
--- a/Engine/Source/Runtime/Slate/Public/Framework/Text/TextDecorators.h
+++ b/Engine/Source/Runtime/Slate/Public/Framework/Text/TextDecorators.h
@@ -12,7 +12,7 @@ class ISlateStyle;
 
 #if WITH_FANCY_TEXT
 
-class SLATE_API FWidgetDecorator : public ITextDecorator
+class FWidgetDecorator : public ITextDecorator
 {
 public:
 
@@ -20,18 +20,18 @@ public:
 
 public:
 
-	static TSharedRef< FWidgetDecorator > Create( FString InRunName, const FCreateWidget& InCreateWidgetDelegate );
+	static SLATE_API TSharedRef< FWidgetDecorator > Create( FString InRunName, const FCreateWidget& InCreateWidgetDelegate );
 	virtual ~FWidgetDecorator() {}
 
 public:
 
-	virtual bool Supports( const FTextRunParseResults& RunParseResult, const FString& Text ) const override;
+	SLATE_API virtual bool Supports( const FTextRunParseResults& RunParseResult, const FString& Text ) const override;
 
-	virtual TSharedRef< ISlateRun > Create(const TSharedRef<class FTextLayout>& TextLayout, const FTextRunParseResults& RunParseResult, const FString& OriginalText, const TSharedRef< FString >& InOutModelText, const ISlateStyle* Style) override;
+	SLATE_API virtual TSharedRef< ISlateRun > Create(const TSharedRef<class FTextLayout>& TextLayout, const FTextRunParseResults& RunParseResult, const FString& OriginalText, const TSharedRef< FString >& InOutModelText, const ISlateStyle* Style) override;
 
 private:
 
-	FWidgetDecorator( FString InRunName, const FCreateWidget& InCreateWidgetDelegate );
+	SLATE_API FWidgetDecorator( FString InRunName, const FCreateWidget& InCreateWidgetDelegate );
 
 private:
 
@@ -39,22 +39,22 @@ private:
 	FCreateWidget CreateWidgetDelegate;
 };
 
-class SLATE_API FImageDecorator : public ITextDecorator
+class FImageDecorator : public ITextDecorator
 {
 public:
 
-	static TSharedRef< FImageDecorator > Create( FString InRunName, const ISlateStyle* const InOverrideStyle = NULL );
+	static SLATE_API TSharedRef< FImageDecorator > Create( FString InRunName, const ISlateStyle* const InOverrideStyle = NULL );
 	virtual ~FImageDecorator() {}
 
 public:
 
-	virtual bool Supports( const FTextRunParseResults& RunParseResult, const FString& Text ) const override;
+	SLATE_API virtual bool Supports( const FTextRunParseResults& RunParseResult, const FString& Text ) const override;
 
-	virtual TSharedRef< ISlateRun > Create(const TSharedRef<class FTextLayout>& TextLayout, const FTextRunParseResults& RunParseResult, const FString& OriginalText, const TSharedRef< FString >& InOutModelText, const ISlateStyle* Style) override;
+	SLATE_API virtual TSharedRef< ISlateRun > Create(const TSharedRef<class FTextLayout>& TextLayout, const FTextRunParseResults& RunParseResult, const FString& OriginalText, const TSharedRef< FString >& InOutModelText, const ISlateStyle* Style) override;
 
 private:
 
-	FImageDecorator( FString InRunName, const ISlateStyle* const InOverrideStyle );
+	SLATE_API FImageDecorator( FString InRunName, const ISlateStyle* const InOverrideStyle );
 
 private:
 
@@ -62,22 +62,22 @@ private:
 	const ISlateStyle* OverrideStyle;
 };
 
-class SLATE_API FHyperlinkDecorator : public ITextDecorator
+class FHyperlinkDecorator : public ITextDecorator
 {
 public:
 
-	static TSharedRef< FHyperlinkDecorator > Create( FString Id, const FSlateHyperlinkRun::FOnClick& NavigateDelegate, const FSlateHyperlinkRun::FOnGetTooltipText& InToolTipTextDelegate = FSlateHyperlinkRun::FOnGetTooltipText(), const FSlateHyperlinkRun::FOnGenerateTooltip& InToolTipDelegate = FSlateHyperlinkRun::FOnGenerateTooltip() );
+	static SLATE_API TSharedRef< FHyperlinkDecorator > Create( FString Id, const FSlateHyperlinkRun::FOnClick& NavigateDelegate, const FSlateHyperlinkRun::FOnGetTooltipText& InToolTipTextDelegate = FSlateHyperlinkRun::FOnGetTooltipText(), const FSlateHyperlinkRun::FOnGenerateTooltip& InToolTipDelegate = FSlateHyperlinkRun::FOnGenerateTooltip() );
 	virtual ~FHyperlinkDecorator() {}
 
 public:
 
-	virtual bool Supports( const FTextRunParseResults& RunParseResult, const FString& Text ) const override;
+	SLATE_API virtual bool Supports( const FTextRunParseResults& RunParseResult, const FString& Text ) const override;
 
-	virtual TSharedRef< ISlateRun > Create(const TSharedRef<class FTextLayout>& TextLayout, const FTextRunParseResults& RunParseResult, const FString& OriginalText, const TSharedRef< FString >& InOutModelText, const ISlateStyle* Style) override;
+	SLATE_API virtual TSharedRef< ISlateRun > Create(const TSharedRef<class FTextLayout>& TextLayout, const FTextRunParseResults& RunParseResult, const FString& OriginalText, const TSharedRef< FString >& InOutModelText, const ISlateStyle* Style) override;
 
 protected:
 
-	FHyperlinkDecorator( FString InId, const FSlateHyperlinkRun::FOnClick& InNavigateDelegate, const FSlateHyperlinkRun::FOnGetTooltipText& InToolTipTextDelegate, const FSlateHyperlinkRun::FOnGenerateTooltip& InToolTipDelegate );
+	SLATE_API FHyperlinkDecorator( FString InId, const FSlateHyperlinkRun::FOnClick& InNavigateDelegate, const FSlateHyperlinkRun::FOnGetTooltipText& InToolTipTextDelegate, const FSlateHyperlinkRun::FOnGenerateTooltip& InToolTipDelegate );
 
 protected:
 
