[Slate]

diff --git a/Engine/Source/Runtime/Slate/Public/Framework/Text/TextLayout.h b/Engine/Source/Runtime/Slate/Public/Framework/Text/TextLayout.h
--- a/Engine/Source/Runtime/Slate/Public/Framework/Text/TextLayout.h
+++ b/Engine/Source/Runtime/Slate/Public/Framework/Text/TextLayout.h
@@ -198,7 +198,7 @@ public:
 	}
 };
 
-class SLATE_API FTextLayout
+class FTextLayout
 	: public TSharedFromThis<FTextLayout>
 {
 public:
@@ -234,39 +234,39 @@ public:
 #endif
 	};
 
-	class SLATE_API FRunModel
+	class FRunModel
 	{
 	public:
 
-		FRunModel( const TSharedRef< IRun >& InRun);
+		SLATE_API FRunModel( const TSharedRef< IRun >& InRun);
 
 	public:
 
-		TSharedRef< IRun > GetRun() const;;
+		SLATE_API TSharedRef< IRun > GetRun() const;;
 
-		void BeginLayout();
-		void EndLayout();
+		SLATE_API void BeginLayout();
+		SLATE_API void EndLayout();
 
-		FTextRange GetTextRange() const;
-		void SetTextRange( const FTextRange& Value );
+		SLATE_API FTextRange GetTextRange() const;
+		SLATE_API void SetTextRange( const FTextRange& Value );
 		
-		int16 GetBaseLine( float Scale ) const;
-		int16 GetMaxHeight( float Scale ) const;
+		SLATE_API int16 GetBaseLine( float Scale ) const;
+		SLATE_API int16 GetMaxHeight( float Scale ) const;
 
-		FVector2D Measure( int32 BeginIndex, int32 EndIndex, float Scale, const FRunTextContext& InTextContext );
+		SLATE_API FVector2D Measure( int32 BeginIndex, int32 EndIndex, float Scale, const FRunTextContext& InTextContext );
 
-		int8 GetKerning( int32 CurrentIndex, float Scale, const FRunTextContext& InTextContext );
+		SLATE_API int8 GetKerning( int32 CurrentIndex, float Scale, const FRunTextContext& InTextContext );
 
-		static int32 BinarySearchForBeginIndex( const TArray< FTextRange >& Ranges, int32 BeginIndex );
+		static SLATE_API int32 BinarySearchForBeginIndex( const TArray< FTextRange >& Ranges, int32 BeginIndex );
 
-		static int32 BinarySearchForEndIndex( const TArray< FTextRange >& Ranges, int32 RangeBeginIndex, int32 EndIndex );
+		static SLATE_API int32 BinarySearchForEndIndex( const TArray< FTextRange >& Ranges, int32 RangeBeginIndex, int32 EndIndex );
 
-		TSharedRef< ILayoutBlock > CreateBlock( const FBlockDefinition& BlockDefine, float InScale, const FLayoutBlockTextContext& InTextContext ) const;
+		SLATE_API TSharedRef< ILayoutBlock > CreateBlock( const FBlockDefinition& BlockDefine, float InScale, const FLayoutBlockTextContext& InTextContext ) const;
 
-		void ClearCache();
+		SLATE_API void ClearCache();
 
-		void AppendTextTo(FString& Text) const;
-		void AppendTextTo(FString& Text, const FTextRange& Range) const;
+		SLATE_API void AppendTextTo(FString& Text) const;
+		SLATE_API void AppendTextTo(FString& Text, const FTextRange& Range) const;
 
 	private:
 
@@ -288,11 +288,11 @@ public:
 		};
 	};
 
-	struct SLATE_API FLineModel
+	struct FLineModel
 	{
 	public:
 
-		FLineModel( const TSharedRef< FString >& InText );
+		SLATE_API FLineModel( const TSharedRef< FString >& InText );
 
 		TSharedRef< FString > Text;
 		FShapedTextCacheRef ShapedTextCache;
@@ -329,14 +329,14 @@ public:
 	};
 
 	/** A mapping between the offsets into the text as a flat string (with line-breaks), and the internal lines used within a text layout */
-	struct SLATE_API FTextOffsetLocations
+	struct FTextOffsetLocations
 	{
 	friend class FTextLayout;
 
 	public:
-		int32 TextLocationToOffset(const FTextLocation& InLocation) const;
-		FTextLocation OffsetToTextLocation(const int32 InOffset) const;
-		int32 GetTextLength() const;
+		SLATE_API int32 TextLocationToOffset(const FTextLocation& InLocation) const;
+		SLATE_API FTextLocation OffsetToTextLocation(const int32 InOffset) const;
+		SLATE_API int32 GetTextLength() const;
 
 	private:
 		struct FOffsetEntry
@@ -362,7 +362,7 @@ public:
 
 public:
 
-	virtual ~FTextLayout();
+	SLATE_API virtual ~FTextLayout();
 
 	FORCEINLINE const TArray< FTextLayout::FLineView >& GetLineViews() const { return LineViews; }
 	FORCEINLINE const TArray< FTextLayout::FLineModel >& GetLineModels() const { return LineModels; }
@@ -371,73 +371,73 @@ public:
 	 * Get the size of the text layout, including any lines which extend beyond the wrapping boundaries (eg, lines with lots of trailing whitespace, or lines with no break candidates)
 	 * @note This value is unscaled
 	 */
-	FVector2D GetSize() const;
+	SLATE_API FVector2D GetSize() const;
 
 	/**
 	 * Get the size of the text layout that can actually be seen from the parent widget
 	 */
-	FVector2D GetViewSize() const;
+	SLATE_API FVector2D GetViewSize() const;
 
 	/**
 	 * Get the size of the text layout, including any lines which extend beyond the wrapping boundaries (eg, lines with lots of trailing whitespace, or lines with no break candidates)
 	 * @note This value is scaled
 	 */
-	FVector2D GetDrawSize() const;
+	SLATE_API FVector2D GetDrawSize() const;
 
 	/**
 	 * Get the size of the text layout after the text has been wrapped, and including the first piece of trailing whitespace for any given soft-wrapped line
 	 * @note This value is unscaled
 	 */
-	FVector2D GetWrappedSize() const;
+	SLATE_API FVector2D GetWrappedSize() const;
 
 	/**
 	 * Get the size of the text layout after the text has been wrapped, and including the first piece of trailing whitespace for any given soft-wrapped line
 	 * @note This value is scaled
 	 */
-	FVector2D GetWrappedDrawSize() const;
+	SLATE_API FVector2D GetWrappedDrawSize() const;
 
 	FORCEINLINE float GetWrappingWidth() const { return WrappingWidth; }
-	void SetWrappingWidth( float Value );
+	SLATE_API void SetWrappingWidth( float Value );
 
 	FORCEINLINE ETextWrappingPolicy GetWrappingPolicy() const { return WrappingPolicy; }
-	void SetWrappingPolicy(ETextWrappingPolicy Value);
+	SLATE_API void SetWrappingPolicy(ETextWrappingPolicy Value);
 
 	FORCEINLINE float GetLineHeightPercentage() const { return LineHeightPercentage; }
-	void SetLineHeightPercentage( float Value );
+	SLATE_API void SetLineHeightPercentage( float Value );
 
 	FORCEINLINE ETextJustify::Type GetJustification() const { return Justification; }
-	void SetJustification( ETextJustify::Type Value );
+	SLATE_API void SetJustification( ETextJustify::Type Value );
 
 	/** Get the visual justification for this document (based on the visual justification used by the first line of text) */
-	ETextJustify::Type GetVisualJustification() const;
+	SLATE_API ETextJustify::Type GetVisualJustification() const;
 
 	/** @note This option is destructive to the model text, so changing it requires refreshing the text layout from its marshaller */
 	FORCEINLINE ETextTransformPolicy GetTransformPolicy() const { return TransformPolicy; }
-	void SetTransformPolicy(ETextTransformPolicy Value);
+	SLATE_API void SetTransformPolicy(ETextTransformPolicy Value);
 
 	FORCEINLINE float GetScale() const { return Scale; }
-	void SetScale( float Value );
+	SLATE_API void SetScale( float Value );
 
 	FORCEINLINE ETextShapingMethod GetTextShapingMethod() const { return TextShapingMethod; }
-	void SetTextShapingMethod( const ETextShapingMethod InTextShapingMethod );
+	SLATE_API void SetTextShapingMethod( const ETextShapingMethod InTextShapingMethod );
 
 	FORCEINLINE ETextFlowDirection GetTextFlowDirection() const { return TextFlowDirection; }
-	void SetTextFlowDirection( const ETextFlowDirection InTextFlowDirection );
+	SLATE_API void SetTextFlowDirection( const ETextFlowDirection InTextFlowDirection );
 
-	void SetTextOverflowPolicy(const TOptional<ETextOverflowPolicy> InTextOverflowPolicy);
+	SLATE_API void SetTextOverflowPolicy(const TOptional<ETextOverflowPolicy> InTextOverflowPolicy);
 
 	FORCEINLINE FMargin GetMargin() const { return Margin; }
-	void SetMargin( const FMargin& InMargin );
+	SLATE_API void SetMargin( const FMargin& InMargin );
 
-	void SetVisibleRegion( const FVector2D& InViewSize, const FVector2D& InScrollOffset );
+	SLATE_API void SetVisibleRegion( const FVector2D& InViewSize, const FVector2D& InScrollOffset );
 
 	/** Set the iterator to use to detect appropriate soft-wrapping points for lines (or null to go back to using the default) */
-	void SetLineBreakIterator( TSharedPtr<IBreakIterator> InLineBreakIterator );
+	SLATE_API void SetLineBreakIterator( TSharedPtr<IBreakIterator> InLineBreakIterator );
 
 	/** Set the information used to help identify who owns this text layout in the case of an error */
-	void SetDebugSourceInfo(const TAttribute<FString>& InDebugSourceInfo);
+	SLATE_API void SetDebugSourceInfo(const TAttribute<FString>& InDebugSourceInfo);
 
-	void ClearLines();
+	SLATE_API void ClearLines();
 
 	struct FNewLineData
 	{
@@ -451,184 +451,184 @@ public:
 		TArray<TSharedRef<IRun>> Runs;
 	};
 
-	void AddLine( const FNewLineData& NewLine );
+	SLATE_API void AddLine( const FNewLineData& NewLine );
 
-	void AddLines( const TArray<FNewLineData>& NewLines );
+	SLATE_API void AddLines( const TArray<FNewLineData>& NewLines );
 
 	/**
 	* Clears all run renderers
 	*/
-	void ClearRunRenderers();
+	SLATE_API void ClearRunRenderers();
 
 	/**
 	* Replaces the current set of run renderers with the provided renderers.
 	*/
-	void SetRunRenderers( const TArray< FTextRunRenderer >& Renderers );
+	SLATE_API void SetRunRenderers( const TArray< FTextRunRenderer >& Renderers );
 
 	/**
 	* Adds a single run renderer to the existing set of renderers.
 	*/
-	void AddRunRenderer( const FTextRunRenderer& Renderer );
+	SLATE_API void AddRunRenderer( const FTextRunRenderer& Renderer );
 
 	/**
 	* Removes a single run renderer to the existing set of renderers.
 	*/
-	void RemoveRunRenderer( const FTextRunRenderer& Renderer );
+	SLATE_API void RemoveRunRenderer( const FTextRunRenderer& Renderer );
 
 	/**
 	* Clears all line highlights
 	*/
-	void ClearLineHighlights();
+	SLATE_API void ClearLineHighlights();
 
 	/**
 	* Replaces the current set of line highlights with the provided highlights.
 	*/
-	void SetLineHighlights( const TArray< FTextLineHighlight >& Highlights );
+	SLATE_API void SetLineHighlights( const TArray< FTextLineHighlight >& Highlights );
 
 	/**
 	* Adds a single line highlight to the existing set of highlights.
 	*/
-	void AddLineHighlight( const FTextLineHighlight& Highlight );
+	SLATE_API void AddLineHighlight( const FTextLineHighlight& Highlight );
 
 	/**
 	* Removes a single line highlight to the existing set of highlights.
 	*/
-	void RemoveLineHighlight( const FTextLineHighlight& Highlight );
+	SLATE_API void RemoveLineHighlight( const FTextLineHighlight& Highlight );
 
 	/**
 	* Updates the TextLayout's if any changes have occurred since the last update.
 	*/
-	virtual void UpdateIfNeeded();
+	SLATE_API virtual void UpdateIfNeeded();
 
-	virtual void UpdateLayout();
+	SLATE_API virtual void UpdateLayout();
 
-	virtual void UpdateHighlights();
+	SLATE_API virtual void UpdateHighlights();
 
-	void DirtyRunLayout(const TSharedRef<const IRun>& Run);
+	SLATE_API void DirtyRunLayout(const TSharedRef<const IRun>& Run);
 
-	void DirtyLayout();
+	SLATE_API void DirtyLayout();
 
-	bool IsLayoutDirty() const;
+	SLATE_API bool IsLayoutDirty() const;
 
-	int32 GetLineViewIndexForTextLocation(const TArray< FTextLayout::FLineView >& LineViews, const FTextLocation& Location, const bool bPerformInclusiveBoundsCheck) const;
+	SLATE_API int32 GetLineViewIndexForTextLocation(const TArray< FTextLayout::FLineView >& LineViews, const FTextLocation& Location, const bool bPerformInclusiveBoundsCheck) const;
 
 	/**
 	 * 
 	 */
-	FTextLocation GetTextLocationAt( const FVector2D& Relative, ETextHitPoint* const OutHitPoint = nullptr ) const;
+	SLATE_API FTextLocation GetTextLocationAt( const FVector2D& Relative, ETextHitPoint* const OutHitPoint = nullptr ) const;
 
-	FTextLocation GetTextLocationAt( const FLineView& LineView, const FVector2D& Relative, ETextHitPoint* const OutHitPoint = nullptr ) const;
+	SLATE_API FTextLocation GetTextLocationAt( const FLineView& LineView, const FVector2D& Relative, ETextHitPoint* const OutHitPoint = nullptr ) const;
 
-	FVector2D GetLocationAt( const FTextLocation& Location, const bool bPerformInclusiveBoundsCheck ) const;
+	SLATE_API FVector2D GetLocationAt( const FTextLocation& Location, const bool bPerformInclusiveBoundsCheck ) const;
 
-	bool SplitLineAt(const FTextLocation& Location);
+	SLATE_API bool SplitLineAt(const FTextLocation& Location);
 
-	bool JoinLineWithNextLine(int32 LineIndex);
+	SLATE_API bool JoinLineWithNextLine(int32 LineIndex);
 
-	bool InsertAt(const FTextLocation& Location, TCHAR Character);
+	SLATE_API bool InsertAt(const FTextLocation& Location, TCHAR Character);
 
-	bool InsertAt(const FTextLocation& Location, const FString& Text);
+	SLATE_API bool InsertAt(const FTextLocation& Location, const FString& Text);
 
-	bool InsertAt(const FTextLocation& Location, TSharedRef<IRun> InRun, const bool bAlwaysKeepRightRun = false);
+	SLATE_API bool InsertAt(const FTextLocation& Location, TSharedRef<IRun> InRun, const bool bAlwaysKeepRightRun = false);
 
-	bool RemoveAt(const FTextLocation& Location, int32 Count = 1);
+	SLATE_API bool RemoveAt(const FTextLocation& Location, int32 Count = 1);
 
-	bool RemoveLine(int32 LineIndex);
+	SLATE_API bool RemoveLine(int32 LineIndex);
 
-	bool IsEmpty() const;
+	SLATE_API bool IsEmpty() const;
 
-	int32 GetLineCount() const;
+	SLATE_API int32 GetLineCount() const;
 
-	void GetAsText(FString& DisplayText, FTextOffsetLocations* const OutTextOffsetLocations = nullptr) const;
+	SLATE_API void GetAsText(FString& DisplayText, FTextOffsetLocations* const OutTextOffsetLocations = nullptr) const;
 
-	void GetAsText(FText& DisplayText, FTextOffsetLocations* const OutTextOffsetLocations = nullptr) const;
+	SLATE_API void GetAsText(FText& DisplayText, FTextOffsetLocations* const OutTextOffsetLocations = nullptr) const;
 
 	/** Constructs an array containing the mappings between the text that would be returned by GetAsText, and the internal FTextLocation points used within this text layout */
-	void GetTextOffsetLocations(FTextOffsetLocations& OutTextOffsetLocations) const;
+	SLATE_API void GetTextOffsetLocations(FTextOffsetLocations& OutTextOffsetLocations) const;
 
-	void GetSelectionAsText(FString& DisplayText, const FTextSelection& Selection) const;
+	SLATE_API void GetSelectionAsText(FString& DisplayText, const FTextSelection& Selection) const;
 
-	FTextSelection GetGraphemeAt(const FTextLocation& Location) const;
+	SLATE_API FTextSelection GetGraphemeAt(const FTextLocation& Location) const;
 
-	FTextSelection GetWordAt(const FTextLocation& Location) const;
+	SLATE_API FTextSelection GetWordAt(const FTextLocation& Location) const;
 
 protected:
 
-	FTextLayout();
+	SLATE_API FTextLayout();
 
 	/**
 	 * Calculates the text direction for each line based upon the current shaping method and document flow direction
 	 * When changing the shaping method, or document flow direction, all the lines need to be dirtied (see DirtyAllLineModels(ELineModelDirtyState::TextBaseDirection))
 	 */
-	void CalculateTextDirection();
+	SLATE_API void CalculateTextDirection();
 
 	/**
 	 * Calculates the text direction for the given line based upon the current shaping method and document flow direction
 	 */
-	void CalculateLineTextDirection(FLineModel& LineModel) const;
+	SLATE_API void CalculateLineTextDirection(FLineModel& LineModel) const;
 
 	/**
 	 * Calculates the visual justification for the given line view
 	 */
-	ETextJustify::Type CalculateLineViewVisualJustification(const FLineView& LineView) const;
+	SLATE_API ETextJustify::Type CalculateLineViewVisualJustification(const FLineView& LineView) const;
 
 	/**
 	* Create the wrapping cache for the current text based upon the current scale
 	* Each line keeps its own cached state, so needs to be cleared when changing the text within a line
 	* When changing the scale, all the lines need to be cleared (see DirtyAllLineModels(ELineModelDirtyState::WrappingInformation))
 	*/
-	void CreateWrappingCache();
+	SLATE_API void CreateWrappingCache();
 
 	/**
 	* Create the wrapping cache for the given line based upon the current scale
 	*/
-	void CreateLineWrappingCache(FLineModel& LineModel);
+	SLATE_API void CreateLineWrappingCache(FLineModel& LineModel);
 
 	/**
 	 * Flushes the text shaping cache for each line
 	 */
-	void FlushTextShapingCache();
+	SLATE_API void FlushTextShapingCache();
 
 	/**
 	 * Flushes the text shaping cache for the given line
 	 */
-	void FlushLineTextShapingCache(FLineModel& LineModel);
+	SLATE_API void FlushLineTextShapingCache(FLineModel& LineModel);
 
 	/**
 	 * Set the given dirty flags on all line models in this layout
 	 */
-	void DirtyAllLineModels(const ELineModelDirtyState::Flags InDirtyFlags);
+	SLATE_API void DirtyAllLineModels(const ELineModelDirtyState::Flags InDirtyFlags);
 
 	/**
 	* Clears the current layouts view information.
 	*/
-	void ClearView();
+	SLATE_API void ClearView();
 
 	/**
 	 * Transform the given line model text based on the active transform policy.
 	 * @note This is destructive to the model text!
 	 */
-	void TransformLineText(FLineModel& LineModel) const;
+	SLATE_API void TransformLineText(FLineModel& LineModel) const;
 
 	/**
 	* Notifies all Runs that we are beginning to generate a new layout.
 	*/
-	virtual void BeginLayout();
+	SLATE_API virtual void BeginLayout();
 
 	/**
 	* Notifies all Runs on the given line is beginning to have a new layout generated.
 	*/
-	void BeginLineLayout(FLineModel& LineModel);
+	SLATE_API void BeginLineLayout(FLineModel& LineModel);
 
 	/**
 	* Notifies all Runs that the layout has finished generating.
 	*/
-	virtual void EndLayout();
+	SLATE_API virtual void EndLayout();
 
 	/**
 	* Notifies all Runs on the given line has finished having a new layout generated.
 	*/
-	void EndLineLayout(FLineModel& LineModel);
+	SLATE_API void EndLineLayout(FLineModel& LineModel);
 
 	/**
 	 * Called to generate a new empty text run for this text layout
@@ -637,22 +637,22 @@ protected:
 
 private:
 
-	float GetWrappingDrawWidth() const;
+	SLATE_API float GetWrappingDrawWidth() const;
 
-	void FlowLayout();
-	void MarginLayout();
+	SLATE_API void FlowLayout();
+	SLATE_API void MarginLayout();
 
-	void FlowLineLayout(const int32 LineModelIndex, const float WrappingDrawWidth, TArray<TSharedRef<ILayoutBlock>>& SoftLine);
+	SLATE_API void FlowLineLayout(const int32 LineModelIndex, const float WrappingDrawWidth, TArray<TSharedRef<ILayoutBlock>>& SoftLine);
 
-	void FlowHighlights();
+	SLATE_API void FlowHighlights();
 
-	void JustifyLayout();
+	SLATE_API void JustifyLayout();
 
-	void CreateLineViewBlocks( int32 LineModelIndex, const int32 StopIndex, const float WrappedLineWidth, const TOptional<float>& JustificationWidth, int32& OutRunIndex, int32& OutRendererIndex, int32& OutPreviousBlockEnd, TArray< TSharedRef< ILayoutBlock > >& OutSoftLine );
+	SLATE_API void CreateLineViewBlocks( int32 LineModelIndex, const int32 StopIndex, const float WrappedLineWidth, const TOptional<float>& JustificationWidth, int32& OutRunIndex, int32& OutRendererIndex, int32& OutPreviousBlockEnd, TArray< TSharedRef< ILayoutBlock > >& OutSoftLine );
 
-	FBreakCandidate CreateBreakCandidate( int32& OutRunIndex, FLineModel& Line, int32 PreviousBreak, int32 CurrentBreak );
+	SLATE_API FBreakCandidate CreateBreakCandidate( int32& OutRunIndex, FLineModel& Line, int32 PreviousBreak, int32 CurrentBreak );
 
-	void GetAsTextAndOffsets(FString* const OutDisplayText, FTextOffsetLocations* const OutTextOffsetLocations) const;
+	SLATE_API void GetAsTextAndOffsets(FString* const OutDisplayText, FTextOffsetLocations* const OutTextOffsetLocations) const;
 
 protected:
 	
diff --git a/Engine/Source/Runtime/Slate/Public/Framework/Text/TextLineHighlight.h b/Engine/Source/Runtime/Slate/Public/Framework/Text/TextLineHighlight.h
--- a/Engine/Source/Runtime/Slate/Public/Framework/Text/TextLineHighlight.h
+++ b/Engine/Source/Runtime/Slate/Public/Framework/Text/TextLineHighlight.h
@@ -5,7 +5,7 @@
 
 class ILineHighlighter;
 
-struct SLATE_API FTextLineHighlight
+struct FTextLineHighlight
 {
 	FTextLineHighlight( int32 InLineIndex, const FTextRange& InRange, int32 InZOrder, const TSharedRef< ILineHighlighter >& InHighlighter )
 		: LineIndex( InLineIndex )
diff --git a/Engine/Source/Runtime/Slate/Public/Framework/Text/TextRunRenderer.h b/Engine/Source/Runtime/Slate/Public/Framework/Text/TextRunRenderer.h
--- a/Engine/Source/Runtime/Slate/Public/Framework/Text/TextRunRenderer.h
+++ b/Engine/Source/Runtime/Slate/Public/Framework/Text/TextRunRenderer.h
@@ -5,7 +5,7 @@
 
 class IRunRenderer;
 
-struct SLATE_API FTextRunRenderer
+struct FTextRunRenderer
 {
 	FTextRunRenderer( int32 InLineIndex, const FTextRange& InRange, const TSharedRef< IRunRenderer >& InRenderer )
 		: LineIndex( InLineIndex )
diff --git a/Engine/Source/Runtime/Slate/Public/Framework/Text/WidgetLayoutBlock.h b/Engine/Source/Runtime/Slate/Public/Framework/Text/WidgetLayoutBlock.h
--- a/Engine/Source/Runtime/Slate/Public/Framework/Text/WidgetLayoutBlock.h
+++ b/Engine/Source/Runtime/Slate/Public/Framework/Text/WidgetLayoutBlock.h
@@ -6,7 +6,7 @@
 #include "Framework/Text/IRun.h"
 #include "Framework/Text/ILayoutBlock.h"
 
-class SLATE_API FWidgetLayoutBlock : public ILayoutBlock
+class FWidgetLayoutBlock : public ILayoutBlock
 {
 public:
 
diff --git a/Engine/Source/Runtime/Slate/Public/SlateSettings.h b/Engine/Source/Runtime/Slate/Public/SlateSettings.h
--- a/Engine/Source/Runtime/Slate/Public/SlateSettings.h
+++ b/Engine/Source/Runtime/Slate/Public/SlateSettings.h
@@ -12,8 +12,8 @@
 /**
 * Settings that control Slate functionality
 */
-UCLASS(config=Engine, defaultconfig, meta=(DisplayName="Slate"))
-class SLATE_API USlateSettings : public UObject
+UCLASS(config=Engine, defaultconfig, meta=(DisplayName="Slate"), MinimalAPI)
+class USlateSettings : public UObject
 {
 	GENERATED_UCLASS_BODY()
 
diff --git a/Engine/Source/Runtime/Slate/Public/Widgets/Accessibility/SlateAccessibleWidgets.h b/Engine/Source/Runtime/Slate/Public/Widgets/Accessibility/SlateAccessibleWidgets.h
--- a/Engine/Source/Runtime/Slate/Public/Widgets/Accessibility/SlateAccessibleWidgets.h
+++ b/Engine/Source/Runtime/Slate/Public/Widgets/Accessibility/SlateAccessibleWidgets.h
@@ -8,7 +8,7 @@
 #include "Widgets/Accessibility/SlateCoreAccessibleWidgets.h"
 
 // SButton
-class SLATE_API FSlateAccessibleButton
+class FSlateAccessibleButton
 	: public FSlateAccessibleWidget
 	, public IAccessibleActivatable
 {
@@ -19,13 +19,13 @@ public:
 	virtual IAccessibleActivatable* AsActivatable() override { return this; }
 
 	// IAccessibleActivatable
-	virtual void Activate() override;
+	SLATE_API virtual void Activate() override;
 	// ~
 };
 // ~
 
 // SCheckBox
-class SLATE_API FSlateAccessibleCheckBox
+class FSlateAccessibleCheckBox
 	: public FSlateAccessibleWidget
 	, public IAccessibleActivatable
 	, public IAccessibleProperty
@@ -38,20 +38,20 @@ public:
 	virtual IAccessibleProperty* AsProperty() override { return this; }
 
 	// IAccessibleActivatable
-	virtual void Activate() override;
-	virtual bool IsCheckable() const override;
-	virtual bool GetCheckedState() const override;
+	SLATE_API virtual void Activate() override;
+	SLATE_API virtual bool IsCheckable() const override;
+	SLATE_API virtual bool GetCheckedState() const override;
 	// ~
 	
 	// IAccessibleProperty
-	virtual FString GetValue() const override;
-	virtual FVariant GetValueAsVariant() const override;
+	SLATE_API virtual FString GetValue() const override;
+	SLATE_API virtual FVariant GetValueAsVariant() const override;
 	// ~
 };
 // ~
 
 // SEditableText
-class SLATE_API FSlateAccessibleEditableText
+class FSlateAccessibleEditableText
 	: public FSlateAccessibleWidget
 	, public IAccessibleText
 	, public IAccessibleProperty
@@ -65,21 +65,21 @@ public:
 	// ~
 
 	// IAccessibleText
-	virtual const FString& GetText() const override;
+	SLATE_API virtual const FString& GetText() const override;
 	// ~
 
 	// IAccessibleProperty
-	virtual bool IsReadOnly() const override;
-	virtual bool IsPassword() const override;
-	virtual FString GetValue() const override;
-	virtual FVariant GetValueAsVariant() const override;
-	virtual void SetValue(const FString& Value) override;
+	SLATE_API virtual bool IsReadOnly() const override;
+	SLATE_API virtual bool IsPassword() const override;
+	SLATE_API virtual FString GetValue() const override;
+	SLATE_API virtual FVariant GetValueAsVariant() const override;
+	SLATE_API virtual void SetValue(const FString& Value) override;
 	// ~
 };
 // ~
 
 // SEditableTextBox
-class SLATE_API FSlateAccessibleEditableTextBox
+class FSlateAccessibleEditableTextBox
 	: public FSlateAccessibleWidget
 	, public IAccessibleText
 	, public IAccessibleProperty
@@ -93,21 +93,21 @@ public:
 	// ~
 
 	// IAccessibleText
-	virtual const FString& GetText() const override;
+	SLATE_API virtual const FString& GetText() const override;
 	// ~
 
 	// IAccessibleProperty
-	virtual bool IsReadOnly() const override;
-	virtual bool IsPassword() const override;
-	virtual FString GetValue() const override;
-	virtual FVariant GetValueAsVariant() const override;
-	virtual void SetValue(const FString& Value) override;
+	SLATE_API virtual bool IsReadOnly() const override;
+	SLATE_API virtual bool IsPassword() const override;
+	SLATE_API virtual FString GetValue() const override;
+	SLATE_API virtual FVariant GetValueAsVariant() const override;
+	SLATE_API virtual void SetValue(const FString& Value) override;
 	// ~
 };
 // ~
 
 // SHyperlink
-class SLATE_API FSlateAccessibleHyperlink
+class FSlateAccessibleHyperlink
 	: public FSlateAccessibleButton
 {
 public:
@@ -118,7 +118,7 @@ public:
 // ~
 
 // Layouts
-class SLATE_API FSlateAccessibleLayout
+class FSlateAccessibleLayout
 	: public FSlateAccessibleWidget
 {
 public:
@@ -127,7 +127,7 @@ public:
 // ~
 
 // SSlider
-class SLATE_API FSlateAccessibleSlider
+class FSlateAccessibleSlider
 	: public FSlateAccessibleWidget
 	, public IAccessibleProperty
 {
@@ -139,19 +139,19 @@ public:
 	// ~
 
 	// IAccessibleProperty
-	virtual bool IsReadOnly() const override;
-	virtual float GetStepSize() const override;
-	virtual float GetMaximum() const override;
-	virtual float GetMinimum() const override;
-	virtual FString GetValue() const override;
-	virtual FVariant GetValueAsVariant() const override;
-	virtual void SetValue(const FString& Value) override;
+	SLATE_API virtual bool IsReadOnly() const override;
+	SLATE_API virtual float GetStepSize() const override;
+	SLATE_API virtual float GetMaximum() const override;
+	SLATE_API virtual float GetMinimum() const override;
+	SLATE_API virtual FString GetValue() const override;
+	SLATE_API virtual FVariant GetValueAsVariant() const override;
+	SLATE_API virtual void SetValue(const FString& Value) override;
 	// ~
 };
 // ~
 
 // STextBlock
-class SLATE_API FSlateAccessibleTextBlock
+class FSlateAccessibleTextBlock
 	: public FSlateAccessibleWidget
 	//, public IAccessibleText // Disabled until we have better support for text. JAWS will not read these properly as-is.
 {
@@ -163,7 +163,7 @@ public:
 	// ~
 
 	// IAccessibleText
-	virtual const FString& GetText() const /*override*/;
+	SLATE_API virtual const FString& GetText() const /*override*/;
 	// ~
 };
 // ~
diff --git a/Engine/Source/Runtime/Slate/Public/Widgets/Colors/SColorBlock.h b/Engine/Source/Runtime/Slate/Public/Widgets/Colors/SColorBlock.h
--- a/Engine/Source/Runtime/Slate/Public/Widgets/Colors/SColorBlock.h
+++ b/Engine/Source/Runtime/Slate/Public/Widgets/Colors/SColorBlock.h
@@ -22,9 +22,9 @@ enum class EColorBlockAlphaDisplayMode : uint8
 	Ignore,
 };
 
-class SLATE_API SColorBlock : public SLeafWidget
+class SColorBlock : public SLeafWidget
 {
-	SLATE_DECLARE_WIDGET(SColorBlock, SLeafWidget)
+	SLATE_DECLARE_WIDGET_API(SColorBlock, SLeafWidget, SLATE_API)
 
 public:
 
@@ -80,22 +80,22 @@ public:
 	SLATE_END_ARGS()
 
 public:
-	SColorBlock();
+	SLATE_API SColorBlock();
 
 	/**
 	 * Construct this widget
 	 *
 	 * @param	InArgs	The declaration data for this widget
 	 */
-	void Construct(const FArguments& InArgs);
+	SLATE_API void Construct(const FArguments& InArgs);
 
 private:
 	// SWidget overrides
-	virtual int32 OnPaint(const FPaintArgs& Args, const FGeometry& AllottedGeometry, const FSlateRect& MyCullingRect, FSlateWindowElementList& OutDrawElements, int32 LayerId, const FWidgetStyle& InWidgetStyle, bool bParentEnabled) const override;
-	virtual FReply OnMouseButtonDown(const FGeometry& MyGeometry, const FPointerEvent& MouseEvent) override;
-	virtual FVector2D ComputeDesiredSize(float) const override;
+	SLATE_API virtual int32 OnPaint(const FPaintArgs& Args, const FGeometry& AllottedGeometry, const FSlateRect& MyCullingRect, FSlateWindowElementList& OutDrawElements, int32 LayerId, const FWidgetStyle& InWidgetStyle, bool bParentEnabled) const override;
+	SLATE_API virtual FReply OnMouseButtonDown(const FGeometry& MyGeometry, const FPointerEvent& MouseEvent) override;
+	SLATE_API virtual FVector2D ComputeDesiredSize(float) const override;
 
-	void MakeSection(TArray<FSlateGradientStop>& OutGradientStops, FVector2D StartPt, FVector2D EndPt, FLinearColor Color, const FWidgetStyle& InWidgetStyle, bool bIgnoreAlpha) const;
+	SLATE_API void MakeSection(TArray<FSlateGradientStop>& OutGradientStops, FVector2D StartPt, FVector2D EndPt, FLinearColor Color, const FWidgetStyle& InWidgetStyle, bool bIgnoreAlpha) const;
 
 private:
 	/** The color to display for this color block */
diff --git a/Engine/Source/Runtime/Slate/Public/Widgets/Colors/SColorGradingWheel.h b/Engine/Source/Runtime/Slate/Public/Widgets/Colors/SColorGradingWheel.h
--- a/Engine/Source/Runtime/Slate/Public/Widgets/Colors/SColorGradingWheel.h
+++ b/Engine/Source/Runtime/Slate/Public/Widgets/Colors/SColorGradingWheel.h
@@ -12,10 +12,10 @@
 /**
  * Implements the color wheel widget.
  */
-class SLATE_API SColorGradingWheel
+class SColorGradingWheel
 	: public SLeafWidget
 {
-	SLATE_DECLARE_WIDGET(SColorGradingWheel, SLeafWidget)
+	SLATE_DECLARE_WIDGET_API(SColorGradingWheel, SLeafWidget, SLATE_API)
 
 public:
 
@@ -50,25 +50,25 @@ public:
 	SLATE_END_ARGS()
 	
 public:
-	SColorGradingWheel();
+	SLATE_API SColorGradingWheel();
 
 	/**
 	 * Construct this widget.
 	 *
 	 * @param InArgs The declaration data for this widget.
 	 */
-	void Construct(const FArguments& InArgs);
+	SLATE_API void Construct(const FArguments& InArgs);
 
 public:
 
 	// SWidget overrides
 
-	virtual FVector2D ComputeDesiredSize(float) const override;
-	virtual FReply OnMouseButtonDoubleClick(const FGeometry& InMyGeometry, const FPointerEvent& InMouseEvent) override;
-	virtual FReply OnMouseButtonDown(const FGeometry& MyGeometry, const FPointerEvent& MouseEvent) override;
-	virtual FReply OnMouseButtonUp(const FGeometry& MyGeometry, const FPointerEvent& MouseEvent) override;
-	virtual FReply OnMouseMove(const FGeometry& MyGeometry, const FPointerEvent& MouseEvent) override;
-	virtual int32 OnPaint(const FPaintArgs& Args, const FGeometry& AllottedGeometry, const FSlateRect& MyCullingRect, FSlateWindowElementList& OutDrawElements, int32 LayerId, const FWidgetStyle& InWidgetStyle, bool bParentEnabled) const override;
+	SLATE_API virtual FVector2D ComputeDesiredSize(float) const override;
+	SLATE_API virtual FReply OnMouseButtonDoubleClick(const FGeometry& InMyGeometry, const FPointerEvent& InMouseEvent) override;
+	SLATE_API virtual FReply OnMouseButtonDown(const FGeometry& MyGeometry, const FPointerEvent& MouseEvent) override;
+	SLATE_API virtual FReply OnMouseButtonUp(const FGeometry& MyGeometry, const FPointerEvent& MouseEvent) override;
+	SLATE_API virtual FReply OnMouseMove(const FGeometry& MyGeometry, const FPointerEvent& MouseEvent) override;
+	SLATE_API virtual int32 OnPaint(const FPaintArgs& Args, const FGeometry& AllottedGeometry, const FSlateRect& MyCullingRect, FSlateWindowElementList& OutDrawElements, int32 LayerId, const FWidgetStyle& InWidgetStyle, bool bParentEnabled) const override;
 	
 protected:
 
@@ -77,23 +77,23 @@ protected:
 	 *
 	 * @return The position relative to the widget.
 	 */
-	UE::Slate::FDeprecateVector2DResult CalcRelativePositionFromCenter() const;
+	SLATE_API UE::Slate::FDeprecateVector2DResult CalcRelativePositionFromCenter() const;
 
 	/**
 	 * Performs actions according to mouse click / move
 	 *
 	 * @return	True if the mouse action occurred within the color wheel radius
 	 */
-	bool ProcessMouseAction(const FGeometry& MyGeometry, const FPointerEvent& MouseEvent, bool bProcessWhenOutsideColorWheel);
+	SLATE_API bool ProcessMouseAction(const FGeometry& MyGeometry, const FPointerEvent& MouseEvent, bool bProcessWhenOutsideColorWheel);
 
 	/** */
-	void SetSelectedColorAttribute(TAttribute<FLinearColor> InSelectedColor);
+	SLATE_API void SetSelectedColorAttribute(TAttribute<FLinearColor> InSelectedColor);
 
 	/** */
-	void SetDesiredWheelSizeAttribute(TAttribute<int32> InDesiredWheelSize);
+	SLATE_API void SetDesiredWheelSizeAttribute(TAttribute<int32> InDesiredWheelSize);
 
 	/** */
-	void SetExponentDisplacementAttribute(TAttribute<float> InExponentDisplacement);
+	SLATE_API void SetExponentDisplacementAttribute(TAttribute<float> InExponentDisplacement);
 
 	/** @return an attribute reference of SelectedColor */
 	TSlateAttributeRef<FLinearColor> GetSelectedColorAttribute() const { return TSlateAttributeRef<FLinearColor>(SharedThis(this), SelectedColorAttribute); }
diff --git a/Engine/Source/Runtime/Slate/Public/Widgets/Colors/SColorSpectrum.h b/Engine/Source/Runtime/Slate/Public/Widgets/Colors/SColorSpectrum.h
--- a/Engine/Source/Runtime/Slate/Public/Widgets/Colors/SColorSpectrum.h
+++ b/Engine/Source/Runtime/Slate/Public/Widgets/Colors/SColorSpectrum.h
@@ -16,7 +16,7 @@ struct FSlateBrush;
 /**
  * Implements the color spectrum widget.
  */
-class SLATE_API SColorSpectrum
+class SColorSpectrum
 	: public SLeafWidget
 {
 public:
@@ -43,25 +43,25 @@ public:
 	SLATE_END_ARGS()
 	
 public:
-	SColorSpectrum();
+	SLATE_API SColorSpectrum();
 
 	/**
 	 * Construct this widget.
 	 *
 	 * @param InArgs The declaration data for this widget.
 	 */
-	void Construct( const FArguments& InArgs );
+	SLATE_API void Construct( const FArguments& InArgs );
 
 public:
 
 	// SWidget overrides
 
-	virtual FVector2D ComputeDesiredSize(float) const override;
-	virtual FReply OnMouseButtonDoubleClick( const FGeometry& InMyGeometry, const FPointerEvent& InMouseEvent ) override;
-	virtual FReply OnMouseButtonDown( const FGeometry& MyGeometry, const FPointerEvent& MouseEvent ) override;
-	virtual FReply OnMouseButtonUp( const FGeometry& MyGeometry, const FPointerEvent& MouseEvent ) override;
-	virtual FReply OnMouseMove( const FGeometry& MyGeometry, const FPointerEvent& MouseEvent ) override;
-	virtual int32 OnPaint( const FPaintArgs& Args, const FGeometry& AllottedGeometry, const FSlateRect& MyCullingRect, FSlateWindowElementList& OutDrawElements, int32 LayerId, const FWidgetStyle& InWidgetStyle, bool bParentEnabled ) const override;
+	SLATE_API virtual FVector2D ComputeDesiredSize(float) const override;
+	SLATE_API virtual FReply OnMouseButtonDoubleClick( const FGeometry& InMyGeometry, const FPointerEvent& InMouseEvent ) override;
+	SLATE_API virtual FReply OnMouseButtonDown( const FGeometry& MyGeometry, const FPointerEvent& MouseEvent ) override;
+	SLATE_API virtual FReply OnMouseButtonUp( const FGeometry& MyGeometry, const FPointerEvent& MouseEvent ) override;
+	SLATE_API virtual FReply OnMouseMove( const FGeometry& MyGeometry, const FPointerEvent& MouseEvent ) override;
+	SLATE_API virtual int32 OnPaint( const FPaintArgs& Args, const FGeometry& AllottedGeometry, const FSlateRect& MyCullingRect, FSlateWindowElementList& OutDrawElements, int32 LayerId, const FWidgetStyle& InWidgetStyle, bool bParentEnabled ) const override;
 	
 protected:
 
@@ -70,12 +70,12 @@ protected:
 	 *
 	 * @return The position relative to the widget.
 	 */
-	FVector2D CalcRelativeSelectedPosition( ) const;
+	SLATE_API FVector2D CalcRelativeSelectedPosition( ) const;
 
 	/**
 	 * Performs actions according to mouse click / move
 	 */
-	void ProcessMouseAction(const FGeometry& MyGeometry, const FPointerEvent& MouseEvent);
+	SLATE_API void ProcessMouseAction(const FGeometry& MyGeometry, const FPointerEvent& MouseEvent);
 
 private:
 
diff --git a/Engine/Source/Runtime/Slate/Public/Widgets/Colors/SColorWheel.h b/Engine/Source/Runtime/Slate/Public/Widgets/Colors/SColorWheel.h
--- a/Engine/Source/Runtime/Slate/Public/Widgets/Colors/SColorWheel.h
+++ b/Engine/Source/Runtime/Slate/Public/Widgets/Colors/SColorWheel.h
@@ -16,7 +16,7 @@ struct FSlateBrush;
 /**
  * Implements the color wheel widget.
  */
-class SLATE_API SColorWheel
+class SColorWheel
 	: public SLeafWidget
 {
 public:
@@ -43,25 +43,25 @@ public:
 	SLATE_END_ARGS()
 	
 public:
-	SColorWheel();
+	SLATE_API SColorWheel();
 
 	/**
 	 * Construct this widget.
 	 *
 	 * @param InArgs The declaration data for this widget.
 	 */
-	void Construct(const FArguments& InArgs);
+	SLATE_API void Construct(const FArguments& InArgs);
 
 public:
 
 	//~ SWidget overrides
 
-	virtual FVector2D ComputeDesiredSize(float) const override;
-	virtual FReply OnMouseButtonDoubleClick(const FGeometry& InMyGeometry, const FPointerEvent& InMouseEvent) override;
-	virtual FReply OnMouseButtonDown(const FGeometry& MyGeometry, const FPointerEvent& MouseEvent) override;
-	virtual FReply OnMouseButtonUp(const FGeometry& MyGeometry, const FPointerEvent& MouseEvent) override;
-	virtual FReply OnMouseMove(const FGeometry& MyGeometry, const FPointerEvent& MouseEvent) override;
-	virtual int32 OnPaint(const FPaintArgs& Args, const FGeometry& AllottedGeometry, const FSlateRect& MyCullingRect, FSlateWindowElementList& OutDrawElements, int32 LayerId, const FWidgetStyle& InWidgetStyle, bool bParentEnabled) const override;
+	SLATE_API virtual FVector2D ComputeDesiredSize(float) const override;
+	SLATE_API virtual FReply OnMouseButtonDoubleClick(const FGeometry& InMyGeometry, const FPointerEvent& InMouseEvent) override;
+	SLATE_API virtual FReply OnMouseButtonDown(const FGeometry& MyGeometry, const FPointerEvent& MouseEvent) override;
+	SLATE_API virtual FReply OnMouseButtonUp(const FGeometry& MyGeometry, const FPointerEvent& MouseEvent) override;
+	SLATE_API virtual FReply OnMouseMove(const FGeometry& MyGeometry, const FPointerEvent& MouseEvent) override;
+	SLATE_API virtual int32 OnPaint(const FPaintArgs& Args, const FGeometry& AllottedGeometry, const FSlateRect& MyCullingRect, FSlateWindowElementList& OutDrawElements, int32 LayerId, const FWidgetStyle& InWidgetStyle, bool bParentEnabled) const override;
 	
 protected:
 
@@ -70,14 +70,14 @@ protected:
 	 *
 	 * @return The position relative to the widget.
 	 */
-	UE::Slate::FDeprecateVector2DResult CalcRelativePositionFromCenter() const;
+	SLATE_API UE::Slate::FDeprecateVector2DResult CalcRelativePositionFromCenter() const;
 
 	/**
 	 * Performs actions according to mouse click / move
 	 *
 	 * @return	True if the mouse action occurred within the color wheel radius
 	 */
-	bool ProcessMouseAction(const FGeometry& MyGeometry, const FPointerEvent& MouseEvent, bool bProcessWhenOutsideColorWheel);
+	SLATE_API bool ProcessMouseAction(const FGeometry& MyGeometry, const FPointerEvent& MouseEvent, bool bProcessWhenOutsideColorWheel);
 
 private:
 
diff --git a/Engine/Source/Runtime/Slate/Public/Widgets/Colors/SComplexGradient.h b/Engine/Source/Runtime/Slate/Public/Widgets/Colors/SComplexGradient.h
--- a/Engine/Source/Runtime/Slate/Public/Widgets/Colors/SComplexGradient.h
+++ b/Engine/Source/Runtime/Slate/Public/Widgets/Colors/SComplexGradient.h
@@ -13,7 +13,7 @@ class FSlateWindowElementList;
 /**
  * Implements a Slate widget that renders a color gradient consisting of multiple stops.
  */
-class SLATE_API SComplexGradient
+class SComplexGradient
 	: public SCompoundWidget
 {
 public:
@@ -40,21 +40,21 @@ public:
 
 	SLATE_END_ARGS()
 	
-	SComplexGradient();
+	SLATE_API SComplexGradient();
 
 	/**
 	 * Construct this widget
 	 *
 	 * @param InArgs The declaration data for this widget.
 	 */
-	void Construct( const FArguments& InArgs );
+	SLATE_API void Construct( const FArguments& InArgs );
 
 protected:
 
 	// SCompoundWidget overrides
 
-	virtual int32 OnPaint( const FPaintArgs& Args, const FGeometry& AllottedGeometry, const FSlateRect& MyCullingRect, FSlateWindowElementList& OutDrawElements, int32 LayerId, const FWidgetStyle& InWidgetStyle, bool bParentEnabled ) const override;
-	virtual FVector2D ComputeDesiredSize(float LayoutScaleMultiplier) const override;
+	SLATE_API virtual int32 OnPaint( const FPaintArgs& Args, const FGeometry& AllottedGeometry, const FSlateRect& MyCullingRect, FSlateWindowElementList& OutDrawElements, int32 LayerId, const FWidgetStyle& InWidgetStyle, bool bParentEnabled ) const override;
+	SLATE_API virtual FVector2D ComputeDesiredSize(float LayoutScaleMultiplier) const override;
 private:
 
 	/** The colors used in the gradient. */
diff --git a/Engine/Source/Runtime/Slate/Public/Widgets/Colors/SSimpleGradient.h b/Engine/Source/Runtime/Slate/Public/Widgets/Colors/SSimpleGradient.h
--- a/Engine/Source/Runtime/Slate/Public/Widgets/Colors/SSimpleGradient.h
+++ b/Engine/Source/Runtime/Slate/Public/Widgets/Colors/SSimpleGradient.h
@@ -13,7 +13,7 @@ class FSlateWindowElementList;
 /**
  * Implements a Slate widget that renders a simple color gradient.
  */
-class SLATE_API SSimpleGradient
+class SSimpleGradient
 	: public SCompoundWidget
 {
 public:
@@ -46,20 +46,20 @@ public:
 	SLATE_END_ARGS()
 
 public:
-	SSimpleGradient();
+	SLATE_API SSimpleGradient();
 
 	/**
 	 * Constructs the widget
 	 *
 	 * @param InArgs The declaration data for this widget
 	 */
-	void Construct( const FArguments& InArgs );
+	SLATE_API void Construct( const FArguments& InArgs );
 
 protected:
 
 	// SCompoundWidget overrides
 
-	virtual int32 OnPaint( const FPaintArgs& Args, const FGeometry& AllottedGeometry, const FSlateRect& MyCullingRect, FSlateWindowElementList& OutDrawElements, int32 LayerId, const FWidgetStyle& InWidgetStyle, bool bParentEnabled ) const override;
+	SLATE_API virtual int32 OnPaint( const FPaintArgs& Args, const FGeometry& AllottedGeometry, const FSlateRect& MyCullingRect, FSlateWindowElementList& OutDrawElements, int32 LayerId, const FWidgetStyle& InWidgetStyle, bool bParentEnabled ) const override;
 
 private:
 
diff --git a/Engine/Source/Runtime/Slate/Public/Widgets/Docking/SDockTab.h b/Engine/Source/Runtime/Slate/Public/Widgets/Docking/SDockTab.h
--- a/Engine/Source/Runtime/Slate/Public/Widgets/Docking/SDockTab.h
+++ b/Engine/Source/Runtime/Slate/Public/Widgets/Docking/SDockTab.h
@@ -46,7 +46,7 @@ class FMenuBuilder;
  * A tab widget that also holds on to some content that should be shown when this tab is selected.
  * Intended to be used in conjunction with SDockingTabStack.
  */
-class SLATE_API SDockTab : public SBorder
+class SDockTab : public SBorder
 {
 	friend class FTabManager;
 	friend class STabSidebar;
@@ -125,196 +125,196 @@ public:
 
 
 	/** Construct the widget from the declaration. */
-	void Construct( const FArguments& InArgs );
+	SLATE_API void Construct( const FArguments& InArgs );
 
 	// SWidget interface
-	virtual FReply OnMouseButtonDown( const FGeometry& MyGeometry, const FPointerEvent& MouseEvent ) override;
-	virtual FReply OnMouseButtonDoubleClick( const FGeometry& MyGeometry, const FPointerEvent& MouseEvent ) override;
-	virtual FReply OnDragDetected( const FGeometry& MyGeometry, const FPointerEvent& MouseEvent  ) override;
-	virtual FReply OnMouseButtonUp( const FGeometry& MyGeometry, const FPointerEvent& MouseEvent ) override;
-	virtual void OnDragEnter( const FGeometry& MyGeometry, const FDragDropEvent& DragDropEvent ) override;
-	virtual void OnDragLeave( const FDragDropEvent& DragDropEvent ) override;
-	virtual FReply OnDrop( const FGeometry& MyGeometry, const FDragDropEvent& DragDropEvent ) override;
-	virtual FReply OnTouchStarted( const FGeometry& MyGeometry, const FPointerEvent& InTouchEvent ) override;
-	virtual FReply OnTouchEnded( const FGeometry& MyGeometry, const FPointerEvent& InTouchEvent ) override;
+	SLATE_API virtual FReply OnMouseButtonDown( const FGeometry& MyGeometry, const FPointerEvent& MouseEvent ) override;
+	SLATE_API virtual FReply OnMouseButtonDoubleClick( const FGeometry& MyGeometry, const FPointerEvent& MouseEvent ) override;
+	SLATE_API virtual FReply OnDragDetected( const FGeometry& MyGeometry, const FPointerEvent& MouseEvent  ) override;
+	SLATE_API virtual FReply OnMouseButtonUp( const FGeometry& MyGeometry, const FPointerEvent& MouseEvent ) override;
+	SLATE_API virtual void OnDragEnter( const FGeometry& MyGeometry, const FDragDropEvent& DragDropEvent ) override;
+	SLATE_API virtual void OnDragLeave( const FDragDropEvent& DragDropEvent ) override;
+	SLATE_API virtual FReply OnDrop( const FGeometry& MyGeometry, const FDragDropEvent& DragDropEvent ) override;
+	SLATE_API virtual FReply OnTouchStarted( const FGeometry& MyGeometry, const FPointerEvent& InTouchEvent ) override;
+	SLATE_API virtual FReply OnTouchEnded( const FGeometry& MyGeometry, const FPointerEvent& InTouchEvent ) override;
 	// End of SWidget interface
 
 	// SBorder interface
-	virtual void SetContent(TSharedRef<SWidget> InContent) override;
+	SLATE_API virtual void SetContent(TSharedRef<SWidget> InContent) override;
 	// End of SBorder interface
 
 	/** Content that appears in the TabWell to the left of the tabs */
-	void SetLeftContent(TSharedRef<SWidget> InContent);
+	SLATE_API void SetLeftContent(TSharedRef<SWidget> InContent);
 	/** Content that appears in the TabWell to the right of the tabs */
-	void SetRightContent(TSharedRef<SWidget> InContent);
+	SLATE_API void SetRightContent(TSharedRef<SWidget> InContent);
 	/** Content that appears on the right side of the title bar in the window this stack is in */
-	void SetTitleBarRightContent(TSharedRef<SWidget> InContent);
+	SLATE_API void SetTitleBarRightContent(TSharedRef<SWidget> InContent);
 
 	/** @return True if this tab is currently focused */
-	bool IsActive() const;
+	SLATE_API bool IsActive() const;
 
 	/** @return True if this tab appears active; False otherwise */
-	bool IsForeground() const;
+	SLATE_API bool IsForeground() const;
 
 	/** @return the Foreground color that this widget sets; unset options if the widget does not set a foreground color */
-	virtual FSlateColor GetForegroundColor() const;
+	SLATE_API virtual FSlateColor GetForegroundColor() const;
 
 	/** Add any entries specific to this tab to the tab context menu */
-	void ExtendContextMenu(FMenuBuilder& MenuBuilder);
+	SLATE_API void ExtendContextMenu(FMenuBuilder& MenuBuilder);
 
 	/** Is this an MajorTab? A tool panel tab? */
-	ETabRole GetTabRole() const;
+	SLATE_API ETabRole GetTabRole() const;
 
 	/** Similar to GetTabRole() but returns the correct role for UI style and user input purposes */
-	ETabRole GetVisualTabRole() const;
+	SLATE_API ETabRole GetVisualTabRole() const;
 
 	/**
 	 * What should the content area look like for this type of tab?
 	 * Documents, Apps, and Tool Panels have different backgrounds.
 	 */
-	const FSlateBrush* GetContentAreaBrush() const;
+	SLATE_API const FSlateBrush* GetContentAreaBrush() const;
 
 	/** Depending on the tabs we put into the tab well, we want a different background brush. */
-	const FSlateBrush* GetTabWellBrush() const;
+	SLATE_API const FSlateBrush* GetTabWellBrush() const;
 
 	/** @return the content associated with this tab */
-	TSharedRef<SWidget> GetContent();
-	TSharedRef<SWidget> GetLeftContent();
-	TSharedRef<SWidget> GetRightContent();
-	TSharedRef<SWidget> GetTitleBarRightContent();
+	SLATE_API TSharedRef<SWidget> GetContent();
+	SLATE_API TSharedRef<SWidget> GetLeftContent();
+	SLATE_API TSharedRef<SWidget> GetRightContent();
+	SLATE_API TSharedRef<SWidget> GetTitleBarRightContent();
 
 	/** Padding around the content when it is presented by the SDockingTabStack */
-	FMargin GetContentPadding() const;
+	SLATE_API FMargin GetContentPadding() const;
 
 	/** Gets this tab's layout identifier */
-	const FTabId GetLayoutIdentifier() const;
+	SLATE_API const FTabId GetLayoutIdentifier() const;
 
 	/** Sets the tab's tab well parent, or resets it if nothing is passed in */
-	void SetParent(TSharedPtr<SDockingTabWell> Parent = TSharedPtr<SDockingTabWell>());
+	SLATE_API void SetParent(TSharedPtr<SDockingTabWell> Parent = TSharedPtr<SDockingTabWell>());
 
 	/** Gets the tab's tab well parent, or nothing, if it has none */
-	TSharedPtr<SDockingTabWell> GetParent() const;
+	SLATE_API TSharedPtr<SDockingTabWell> GetParent() const;
 
 	/** Gets the dock tab stack this dockable tab resides within, if any */
-	TSharedPtr<SDockingTabStack> GetParentDockTabStack() const;
+	SLATE_API TSharedPtr<SDockingTabStack> GetParentDockTabStack() const;
 
 	/** Gets the dock area that this resides in */
-	TSharedPtr<class SDockingArea> GetDockArea() const;
+	SLATE_API TSharedPtr<class SDockingArea> GetDockArea() const;
 
 	/** Get the window in which this tab's tabmanager has placed it */
-	TSharedPtr<SWindow> GetParentWindow() const;
+	SLATE_API TSharedPtr<SWindow> GetParentWindow() const;
 
 	/** The width that this tab will overlap with side-by-side tabs. */
-	float GetOverlapWidth() const;
+	SLATE_API float GetOverlapWidth() const;
 
 	/** The label on the tab */
-	FText GetTabLabel() const;
+	SLATE_API FText GetTabLabel() const;
 
 	/** The label that appears on the tab. */
-	void SetLabel( const TAttribute<FText>& InTabLabel );
+	SLATE_API void SetLabel( const TAttribute<FText>& InTabLabel );
 
 	/** Get Label Suffix */
-	FText GetTabLabelSuffix() const;
+	SLATE_API FText GetTabLabelSuffix() const;
 
 	/** Set Label Suffix.  A second text field at the end of the Label that takes precedence and isn't lost when space is restricted */
-	void SetTabLabelSuffix(const TAttribute<FText>& InTabLabelSuffix);
+	SLATE_API void SetTabLabelSuffix(const TAttribute<FText>& InTabLabelSuffix);
 
 	/** The tooltip text that appears on the tab. */
-	void SetTabToolTipWidget(TSharedPtr<SToolTip> InTabToolTipWidget);
+	SLATE_API void SetTabToolTipWidget(TSharedPtr<SToolTip> InTabToolTipWidget);
 
 	/** Gets the tab icon */
-	const FSlateBrush* GetTabIcon() const;
+	SLATE_API const FSlateBrush* GetTabIcon() const;
 
 	/** Sets the tab icon */
-	void SetTabIcon( const TAttribute<const FSlateBrush*> InTabIcon );
+	SLATE_API void SetTabIcon( const TAttribute<const FSlateBrush*> InTabIcon );
 
 	/** Should this tab be sized based on its content. */
-	bool ShouldAutosize() const;
+	SLATE_API bool ShouldAutosize() const;
 
 	/** Set whether this tab should be sized based on its content. */
-	void SetShouldAutosize(const bool bNewShouldAutosize);
+	SLATE_API void SetShouldAutosize(const bool bNewShouldAutosize);
 
 	/** @return true if the tab can be closed right now. Example: Callback could ask user via dialog. */
-	bool CanCloseTab() const;
+	SLATE_API bool CanCloseTab() const;
 
 	/** Requests that the tab be closed.  Tabs may prevent closing depending on their state */	
-	bool RequestCloseTab();
+	SLATE_API bool RequestCloseTab();
 
 	/** A chance for the tab's content to save any internal layout info */
-	void PersistVisualState();
+	SLATE_API void PersistVisualState();
 
 	/** 
 	 * Pulls this tab out of it's parent tab stack and destroys it
 	 * Note: This does not check if its safe to remove the tab.  Use RequestCloseTab to do this safely.
 	 */
-	void RemoveTabFromParent();
+	SLATE_API void RemoveTabFromParent();
 
 	/** Protected constructor; Widgets may only be constructed via a FArguments (i.e.: SNew(SDockTab) ) */
-	SDockTab();
+	SLATE_API SDockTab();
 
 	/** 
 	 * Make this tab active in its tabwell 
 	 * @param	InActivationMethod	How this tab was activated.
 	 */
-	void ActivateInParent(ETabActivationCause InActivationCause);
+	SLATE_API void ActivateInParent(ETabActivationCause InActivationCause);
 
 	/** Set the tab manager that is controlling this tab */
-	void SetTabManager( const TSharedPtr<FTabManager>& InTabManager );
+	SLATE_API void SetTabManager( const TSharedPtr<FTabManager>& InTabManager );
 
 	/**
 	 * Set the custom code to execute for saving visual state in this tab.
 	 * e.g. ContentBrowser saves the visible filters.OnExtendContextMenu
 	 */
-	void SetOnPersistVisualState( const FOnPersistVisualState& Handler );
+	SLATE_API void SetOnPersistVisualState( const FOnPersistVisualState& Handler );
 
 	/** Set the handler to be invoked when the user requests that this tab be closed. */
-	void SetCanCloseTab( const FCanCloseTab& InOnTabClosing );
+	SLATE_API void SetCanCloseTab( const FCanCloseTab& InOnTabClosing );
 
 	/** Set the handler that will be invoked when the tab is closed */
-	void SetOnTabClosed( const FOnTabClosedCallback& InDelegate );
+	SLATE_API void SetOnTabClosed( const FOnTabClosedCallback& InDelegate );
 
 	/** Set the handler that will be invoked when the tab is activated */
-	void SetOnTabActivated( const FOnTabActivatedCallback& InDelegate );
+	SLATE_API void SetOnTabActivated( const FOnTabActivatedCallback& InDelegate );
 
 	/** Set the handler that will be invoked when the tab is relocated to a new tab well */
-	void SetOnTabRelocated(const FSimpleDelegate InDelegate);
+	SLATE_API void SetOnTabRelocated(const FSimpleDelegate InDelegate);
 
 	/** Set the handler that will be invoked when the tab is dragged over dock area */
-	void SetOnTabDraggedOverDockArea(const FSimpleDelegate InDelegate);
+	SLATE_API void SetOnTabDraggedOverDockArea(const FSimpleDelegate InDelegate);
 
 	/** Set the handler that will be invoked when the tab is renamed */
-	void SetOnTabRenamed(const FOnTabRenamed& InDelegate);
+	SLATE_API void SetOnTabRenamed(const FOnTabRenamed& InDelegate);
 
 	/** Set the handler that will be invoked when the tab is opened from a drawer */
-	void SetOnTabDrawerOpened(const FSimpleDelegate InDelegate);
+	SLATE_API void SetOnTabDrawerOpened(const FSimpleDelegate InDelegate);
 
 	/** Set the handler that will be invoked when the tab is closed from a drawer */
-	void SetOnTabDrawerClosed(const FSimpleDelegate InDelegate);
+	SLATE_API void SetOnTabDrawerClosed(const FSimpleDelegate InDelegate);
 
 	/** Set the handler for extending the tab context menu */
-	void SetOnExtendContextMenu( const FExtendContextMenu& Handler );
+	SLATE_API void SetOnExtendContextMenu( const FExtendContextMenu& Handler );
 
 	/** Get the tab manager currently managing this tab. Note that a user move the tab between Tab Managers, so this return value may change. */
 	UE_DEPRECATED(5.0, "The tab manager is not guaranteed to exist, which will cause GetTabManager() to crash. Use GetTabManagerPtr() instead.")
-	TSharedRef<FTabManager> GetTabManager() const;
+	SLATE_API TSharedRef<FTabManager> GetTabManager() const;
 
 	/** Get the tab manager currently managing this tab. Note that a user move the tab between Tab Managers, so this return value may change. */
-	TSharedPtr<FTabManager> GetTabManagerPtr() const;
+	SLATE_API TSharedPtr<FTabManager> GetTabManagerPtr() const;
 
 	/** Draws attention to the tab. */
-	void DrawAttention();
+	SLATE_API void DrawAttention();
 
 	/** Provide a default tab label in case the spawner did not set one. */
-	void ProvideDefaultLabel( const FText& InDefaultLabel );
+	SLATE_API void ProvideDefaultLabel( const FText& InDefaultLabel );
 
 	/** Play an animation showing this tab as opening */
-	void PlaySpawnAnim();
+	SLATE_API void PlaySpawnAnim();
 
 	/** Flash the tab, used for drawing attention to it */
-	void FlashTab();
+	SLATE_API void FlashTab();
 
 	/** Used by the drag/drop operation to signal to this tab what it is dragging over. */
-	void SetDraggedOverDockArea( const TSharedPtr<SDockingArea>& Area );
+	SLATE_API void SetDraggedOverDockArea( const TSharedPtr<SDockingArea>& Area );
 
 	/** 
 	 * Check to see whether this tab has a sibling tab with the given tab ID
@@ -324,10 +324,10 @@ public:
 	 *										If true, we will consider this a "wildcard" search (matching any tab with the correct TabType, regardless 
 	 *										of its InstanceId). If false, we will explicitly look for a tab with an InstanceId of INDEX_NONE
 	 */
-	bool HasSiblingTab(const FTabId& SiblingTabId, const bool TreatIndexNoneAsWildcard = true) const;
+	SLATE_API bool HasSiblingTab(const FTabId& SiblingTabId, const bool TreatIndexNoneAsWildcard = true) const;
 
 	/** Updates the 'last activated' time to the current time */
-	void UpdateActivationTime();
+	SLATE_API void UpdateActivationTime();
 
 	/** Returns the time this tab was last activated */
 	double GetLastActivationTime()
@@ -335,71 +335,71 @@ public:
 		return LastActivationTime;
 	}
 
-	void SetParentDockTabStackTabWellHidden(bool bIsTabWellHidden);
+	SLATE_API void SetParentDockTabStackTabWellHidden(bool bIsTabWellHidden);
 
 protected:
 	/** Provide a default tab icon. */
-	void ProvideDefaultIcon(const FSlateBrush* InDefaultIcon);
+	SLATE_API void ProvideDefaultIcon(const FSlateBrush* InDefaultIcon);
 
 	/** @return the style currently applied to the dock tab */
-	const FDockTabStyle& GetCurrentStyle() const;
+	SLATE_API const FDockTabStyle& GetCurrentStyle() const;
 
 	/** @return the image brush that best represents this tab's in its current state */
-	const FSlateBrush* GetImageBrush() const;
+	SLATE_API const FSlateBrush* GetImageBrush() const;
 
 	/** @return the padding for the tab widget */
-	FMargin GetTabPadding() const;
+	SLATE_API FMargin GetTabPadding() const;
 
 	/** @return the image brush for the tab's color overlay */
-	const FSlateBrush* GetColorOverlayImageBrush() const;
+	SLATE_API const FSlateBrush* GetColorOverlayImageBrush() const;
 
 	/** @return The visibility of the active tab indicator */
-	EVisibility GetActiveTabIndicatorVisibility() const;
+	SLATE_API EVisibility GetActiveTabIndicatorVisibility() const;
 
 	/** @return Returns a color to scale the background of this tab by */
-	FSlateColor GetTabColor() const;
+	SLATE_API FSlateColor GetTabColor() const;
 
 	/** @return Returns the color of this tab's icon */
-	FSlateColor GetIconColor() const;
+	SLATE_API FSlateColor GetIconColor() const;
 
 	/** @return the image brush for the tab's flasher overlay */
-	const FSlateBrush* GetFlashOverlayImageBrush() const;
+	SLATE_API const FSlateBrush* GetFlashOverlayImageBrush() const;
 
 	/** @return Returns a color to flash the background of this tab with */
-	FSlateColor GetFlashColor() const;
+	SLATE_API FSlateColor GetFlashColor() const;
 
 	/** Called when the close button is clicked on the tab. */
-	FReply OnCloseButtonClicked();
+	SLATE_API FReply OnCloseButtonClicked();
 
 	/** The close button tooltip showing the appropriate close command shortcut */
-	FText GetCloseButtonToolTipText() const;
+	SLATE_API FText GetCloseButtonToolTipText() const;
 
 	/** Specify the TabId that was used to spawn this tab. */
-	void SetLayoutIdentifier( const FTabId& TabId );
+	SLATE_API void SetLayoutIdentifier( const FTabId& TabId );
 
 	/** @return if the close button should be visible. */
-	EVisibility HandleIsCloseButtonVisible() const;
+	SLATE_API EVisibility HandleIsCloseButtonVisible() const;
 
 	/** @return the size the tab icon should be */
-	TOptional<FVector2D> GetTabIconSize() const;
+	SLATE_API TOptional<FVector2D> GetTabIconSize() const;
 
 	/** @return the padding for the tab icon border */
-	FMargin GetTabIconBorderPadding() const;
+	SLATE_API FMargin GetTabIconBorderPadding() const;
 
 private:
 	/** Activates the tab in its tab well */
-	EActiveTimerReturnType TriggerActivateTab( double InCurrentTime, float InDeltaTime );
+	SLATE_API EActiveTimerReturnType TriggerActivateTab( double InCurrentTime, float InDeltaTime );
 
-	EActiveTimerReturnType OnHandleUpdateStyle(double InCurrentTime, float InDeltaTime);
+	SLATE_API EActiveTimerReturnType OnHandleUpdateStyle(double InCurrentTime, float InDeltaTime);
 
-	void OnParentSet();
+	SLATE_API void OnParentSet();
 
-	void UpdateTabStyle();
+	SLATE_API void UpdateTabStyle();
 
-	void OnTabDrawerOpened();
-	void OnTabDrawerClosed();
+	SLATE_API void OnTabDrawerOpened();
+	SLATE_API void OnTabDrawerClosed();
 
-	void NotifyTabRelocated();
+	SLATE_API void NotifyTabRelocated();
 
 	/** The handle to the active tab activation tick */
 	TWeakPtr<FActiveTimerHandle> DragDropTimerHandle;
@@ -480,7 +480,7 @@ protected:
 	TAttribute<FLinearColor> IconColor;
 
 	/** @return the scaling of the tab based on the opening/closing animation */
-	UE::Slate::FDeprecateVector2DResult GetAnimatedScale() const;
+	SLATE_API UE::Slate::FDeprecateVector2DResult GetAnimatedScale() const;
 
 	/** Animation that shows the tab opening up */
 	FCurveSequence SpawnAnimCurve;
@@ -489,7 +489,7 @@ protected:
 	FCurveSequence FlashTabCurve;
 
 	/** Get the desired color of tab. These change during flashing. */
-	float GetFlashValue() const;
+	SLATE_API float GetFlashValue() const;
 
 	/** The dock area this tab is currently being dragged over. Used in nomad tabs to change style */
 	TSharedPtr<SDockingArea> DraggedOverDockingArea;
diff --git a/Engine/Source/Runtime/Slate/Public/Widgets/Images/SSpinningImage.h b/Engine/Source/Runtime/Slate/Public/Widgets/Images/SSpinningImage.h
--- a/Engine/Source/Runtime/Slate/Public/Widgets/Images/SSpinningImage.h
+++ b/Engine/Source/Runtime/Slate/Public/Widgets/Images/SSpinningImage.h
@@ -13,7 +13,7 @@ class FPaintArgs;
 class FSlateWindowElementList;
 
 /** A widget that displays a spinning image.*/
-class SLATE_API SSpinningImage : public SImage
+class SSpinningImage : public SImage
 {
 public:
 	SLATE_BEGIN_ARGS(SSpinningImage)
@@ -38,11 +38,11 @@ public:
 	SLATE_END_ARGS()
 
 	/* Construct this widget */
-	void Construct( const FArguments& InArgs );
+	SLATE_API void Construct( const FArguments& InArgs );
 
 	//~ Begin SWidget interface
-	virtual int32 OnPaint( const FPaintArgs& Args, const FGeometry& AllottedGeometry, const FSlateRect& MyCullingRect, FSlateWindowElementList& OutDrawElements, int32 LayerId, const FWidgetStyle& InWidgetStyle, bool bParentEnabled ) const override;
-	virtual bool ComputeVolatility() const override;
+	SLATE_API virtual int32 OnPaint( const FPaintArgs& Args, const FGeometry& AllottedGeometry, const FSlateRect& MyCullingRect, FSlateWindowElementList& OutDrawElements, int32 LayerId, const FWidgetStyle& InWidgetStyle, bool bParentEnabled ) const override;
+	SLATE_API virtual bool ComputeVolatility() const override;
 	//~ End SWidget interface
 	
 private:
diff --git a/Engine/Source/Runtime/Slate/Public/Widgets/Images/SThrobber.h b/Engine/Source/Runtime/Slate/Public/Widgets/Images/SThrobber.h
--- a/Engine/Source/Runtime/Slate/Public/Widgets/Images/SThrobber.h
+++ b/Engine/Source/Runtime/Slate/Public/Widgets/Images/SThrobber.h
@@ -18,7 +18,7 @@ class SHorizontalBox;
 /**
  * A throbber widget that uses 5 zooming circles in a row.
  */
-class SLATE_API SThrobber
+class SThrobber
 	: public SCompoundWidget
 {
 	static const int32 DefaultNumPieces = 3;
@@ -49,25 +49,25 @@ public:
 		SLATE_ARGUMENT( EAnimation, Animate )
 	SLATE_END_ARGS()
 
-	void Construct(const FArguments& InArgs);
+	SLATE_API void Construct(const FArguments& InArgs);
 
 	/** Sets what each segment of the throbber looks like. */
-	void SetPieceImage(const FSlateBrush* InPieceImage);
+	SLATE_API void SetPieceImage(const FSlateBrush* InPieceImage);
 
 	/**
 	 * Invalidate the SCircularThrobber with the correct reason.
 	 * You should invalidate when you are changing a property of SlateBrush after the SlateBrush was set to the SCircularThrobber.
 	 */
-	void InvalidatePieceImage();
+	SLATE_API void InvalidatePieceImage();
 
 	/** Sets how many pieces there are */
-	void SetNumPieces(int InNumPieces);
+	SLATE_API void SetNumPieces(int InNumPieces);
 
 	/** Sets which aspects of the throbber to animate */
-	void SetAnimate(EAnimation InAnimate);
+	SLATE_API void SetAnimate(EAnimation InAnimate);
 
 	//~ Begin SWidget interface
-	virtual bool ComputeVolatility() const override;
+	SLATE_API virtual bool ComputeVolatility() const override;
 	//~ End SWidget interface
 
 private:
@@ -103,13 +103,13 @@ private:
 /**
  * A throbber widget that orients images in a spinning circle.
  */
-class SLATE_API SCircularThrobber
+class SCircularThrobber
 	: public SLeafWidget
 {
-	SLATE_DECLARE_WIDGET(SCircularThrobber, SLeafWidget)
+	SLATE_DECLARE_WIDGET_API(SCircularThrobber, SLeafWidget, SLATE_API)
 
 public:
-	static const float MinimumPeriodValue;
+	static SLATE_API const float MinimumPeriodValue;
 
 	SLATE_BEGIN_ARGS(SCircularThrobber)	
 		: _PieceImage( FCoreStyle::Get().GetBrush( "Throbber.CircleChunk" ) )
@@ -131,33 +131,33 @@ public:
 
 	SLATE_END_ARGS()
 
-	SCircularThrobber();
+	SLATE_API SCircularThrobber();
 
 	/** Constructs the widget */
-	void Construct(const FArguments& InArgs);
+	SLATE_API void Construct(const FArguments& InArgs);
 
 	/** Sets what each segment of the throbber looks like */
-	void SetPieceImage(const FSlateBrush* InPieceImage);
+	SLATE_API void SetPieceImage(const FSlateBrush* InPieceImage);
 
 	/**
 	 * Invalidate the SCircularThrobber with the correct reason. 
 	 * You should invalidate when you are changing a property of SlateBrush after the SlateBrush was set to the SCircularThrobber.
 	 */
-	void InvalidatePieceImage();
+	SLATE_API void InvalidatePieceImage();
 
 	/** Sets how many pieces there are */
-	void SetNumPieces(int32 InNumPieces);
+	SLATE_API void SetNumPieces(int32 InNumPieces);
 
 	/** Sets the amount of time in seconds for a full circle */
-	void SetPeriod(float InPeriod);
+	SLATE_API void SetPeriod(float InPeriod);
 
 	/** Sets the radius of the circle */
-	void SetRadius(float InRadius);
+	SLATE_API void SetRadius(float InRadius);
 
 	//~ Begin SWidget interface
-	virtual int32 OnPaint( const FPaintArgs& Args, const FGeometry& AllottedGeometry, const FSlateRect& MyCullingRect, FSlateWindowElementList& OutDrawElements, int32 LayerId, const FWidgetStyle& InWidgetStyle, bool bParentEnabled ) const override;
-	virtual FVector2D ComputeDesiredSize(float) const override;
-	virtual bool ComputeVolatility() const override;
+	SLATE_API virtual int32 OnPaint( const FPaintArgs& Args, const FGeometry& AllottedGeometry, const FSlateRect& MyCullingRect, FSlateWindowElementList& OutDrawElements, int32 LayerId, const FWidgetStyle& InWidgetStyle, bool bParentEnabled ) const override;
+	SLATE_API virtual FVector2D ComputeDesiredSize(float) const override;
+	SLATE_API virtual bool ComputeVolatility() const override;
 	//~ End SWidget interface
 
 private:
diff --git a/Engine/Source/Runtime/Slate/Public/Widgets/Input/IVirtualKeyboardEntry.h b/Engine/Source/Runtime/Slate/Public/Widgets/Input/IVirtualKeyboardEntry.h
--- a/Engine/Source/Runtime/Slate/Public/Widgets/Input/IVirtualKeyboardEntry.h
+++ b/Engine/Source/Runtime/Slate/Public/Widgets/Input/IVirtualKeyboardEntry.h
@@ -48,7 +48,7 @@ public:
 
 DECLARE_DELEGATE(FOnSelectionChangedDelegateVK);
 
-class SLATE_API IVirtualKeyboardEntry
+class IVirtualKeyboardEntry
 {
 
 public:
diff --git a/Engine/Source/Runtime/Slate/Public/Widgets/Input/SButton.h b/Engine/Source/Runtime/Slate/Public/Widgets/Input/SButton.h
--- a/Engine/Source/Runtime/Slate/Public/Widgets/Input/SButton.h
+++ b/Engine/Source/Runtime/Slate/Public/Widgets/Input/SButton.h
@@ -23,9 +23,9 @@ enum class ETextShapingMethod : uint8;
 /**
  * Slate's Buttons are clickable Widgets that can contain arbitrary widgets as its Content().
  */
-class SLATE_API SButton : public SBorder
+class SButton : public SBorder
 {
-	SLATE_DECLARE_WIDGET(SButton, SBorder)
+	SLATE_DECLARE_WIDGET_API(SButton, SBorder, SLATE_API)
 
 #if WITH_ACCESSIBILITY
 	// Allow the accessible button to "click" this button
@@ -121,7 +121,7 @@ public:
 	SLATE_END_ARGS()
 
 protected:
-		SButton();
+		SLATE_API SButton();
 
 public:
 	/**
@@ -140,7 +140,7 @@ public:
 	}
 
 	/** @return the Foreground color that this widget sets when this widget or any of its ancestors are disabled; unset options if the widget does not set a foreground color */
-	virtual FSlateColor GetDisabledForegroundColor() const final;
+	SLATE_API virtual FSlateColor GetDisabledForegroundColor() const final;
 
 	/**
 	 * Returns true if this button is currently pressed
@@ -158,87 +158,87 @@ public:
 	 *
 	 * @param	InArgs	The declaration data for this widget
 	 */
-	void Construct( const FArguments& InArgs );
+	SLATE_API void Construct( const FArguments& InArgs );
 
 	/** See ContentPadding attribute */
-	void SetContentPadding(TAttribute<FMargin> InContentPadding);
+	SLATE_API void SetContentPadding(TAttribute<FMargin> InContentPadding);
 
 	/** See HoveredSound attribute */
-	void SetHoveredSound(TOptional<FSlateSound> InHoveredSound);
+	SLATE_API void SetHoveredSound(TOptional<FSlateSound> InHoveredSound);
 
 	/** See PressedSound attribute */
-	void SetPressedSound(TOptional<FSlateSound> InPressedSound);
+	SLATE_API void SetPressedSound(TOptional<FSlateSound> InPressedSound);
 
 	/** See OnClicked event */
-	void SetOnClicked(FOnClicked InOnClicked);
+	SLATE_API void SetOnClicked(FOnClicked InOnClicked);
 
 	/** Set OnHovered event */
-	void SetOnHovered(FSimpleDelegate InOnHovered);
+	SLATE_API void SetOnHovered(FSimpleDelegate InOnHovered);
 
 	/** Set OnUnhovered event */
-	void SetOnUnhovered(FSimpleDelegate InOnUnhovered);
+	SLATE_API void SetOnUnhovered(FSimpleDelegate InOnUnhovered);
 
 	/** See ButtonStyle attribute */
-	void SetButtonStyle(const FButtonStyle* ButtonStyle);
+	SLATE_API void SetButtonStyle(const FButtonStyle* ButtonStyle);
 
-	void SetClickMethod(EButtonClickMethod::Type InClickMethod);
-	void SetTouchMethod(EButtonTouchMethod::Type InTouchMethod);
-	void SetPressMethod(EButtonPressMethod::Type InPressMethod);
+	SLATE_API void SetClickMethod(EButtonClickMethod::Type InClickMethod);
+	SLATE_API void SetTouchMethod(EButtonTouchMethod::Type InTouchMethod);
+	SLATE_API void SetPressMethod(EButtonPressMethod::Type InPressMethod);
 
 #if !UE_BUILD_SHIPPING
-	void SimulateClick();
+	SLATE_API void SimulateClick();
 #endif // !UE_BUILD_SHIPPING
 
 public:
 
 	//~ SWidget overrides
-	virtual int32 OnPaint(const FPaintArgs& Args, const FGeometry& AllottedGeometry, const FSlateRect& MyCullingRect, FSlateWindowElementList& OutDrawElements, int32 LayerId, const FWidgetStyle& InWidgetStyle, bool bParentEnabled) const override;
-	virtual bool SupportsKeyboardFocus() const override;
-	virtual void OnFocusLost( const FFocusEvent& InFocusEvent ) override;
-	virtual FReply OnKeyDown( const FGeometry& MyGeometry, const FKeyEvent& InKeyEvent ) override;
-	virtual FReply OnKeyUp( const FGeometry& MyGeometry, const FKeyEvent& InKeyEvent ) override;
-	virtual FReply OnMouseButtonDown( const FGeometry& MyGeometry, const FPointerEvent& MouseEvent ) override;
-	virtual FReply OnMouseButtonDoubleClick(const FGeometry& InMyGeometry, const FPointerEvent& InMouseEvent) override;
-	virtual FReply OnMouseButtonUp( const FGeometry& MyGeometry, const FPointerEvent& MouseEvent ) override;
-	virtual FReply OnMouseMove( const FGeometry& MyGeometry, const FPointerEvent& MouseEvent ) override;
-	virtual void OnMouseEnter( const FGeometry& MyGeometry, const FPointerEvent& MouseEvent ) override;
-	virtual void OnMouseLeave( const FPointerEvent& MouseEvent ) override;
-	virtual void OnMouseCaptureLost(const FCaptureLostEvent& CaptureLostEvent) override;
-	virtual bool IsInteractable() const override;
+	SLATE_API virtual int32 OnPaint(const FPaintArgs& Args, const FGeometry& AllottedGeometry, const FSlateRect& MyCullingRect, FSlateWindowElementList& OutDrawElements, int32 LayerId, const FWidgetStyle& InWidgetStyle, bool bParentEnabled) const override;
+	SLATE_API virtual bool SupportsKeyboardFocus() const override;
+	SLATE_API virtual void OnFocusLost( const FFocusEvent& InFocusEvent ) override;
+	SLATE_API virtual FReply OnKeyDown( const FGeometry& MyGeometry, const FKeyEvent& InKeyEvent ) override;
+	SLATE_API virtual FReply OnKeyUp( const FGeometry& MyGeometry, const FKeyEvent& InKeyEvent ) override;
+	SLATE_API virtual FReply OnMouseButtonDown( const FGeometry& MyGeometry, const FPointerEvent& MouseEvent ) override;
+	SLATE_API virtual FReply OnMouseButtonDoubleClick(const FGeometry& InMyGeometry, const FPointerEvent& InMouseEvent) override;
+	SLATE_API virtual FReply OnMouseButtonUp( const FGeometry& MyGeometry, const FPointerEvent& MouseEvent ) override;
+	SLATE_API virtual FReply OnMouseMove( const FGeometry& MyGeometry, const FPointerEvent& MouseEvent ) override;
+	SLATE_API virtual void OnMouseEnter( const FGeometry& MyGeometry, const FPointerEvent& MouseEvent ) override;
+	SLATE_API virtual void OnMouseLeave( const FPointerEvent& MouseEvent ) override;
+	SLATE_API virtual void OnMouseCaptureLost(const FCaptureLostEvent& CaptureLostEvent) override;
+	SLATE_API virtual bool IsInteractable() const override;
 #if WITH_ACCESSIBILITY
-	virtual TSharedRef<FSlateAccessibleWidget> CreateAccessibleWidget() override;
+	SLATE_API virtual TSharedRef<FSlateAccessibleWidget> CreateAccessibleWidget() override;
 #endif
 protected:
-	virtual FVector2D ComputeDesiredSize(float) const override;
+	SLATE_API virtual FVector2D ComputeDesiredSize(float) const override;
 	//~ SWidget
 
 protected:
 	/** Press the button */
-	virtual void Press();
+	SLATE_API virtual void Press();
 
 	/** Release the button */
-	virtual void Release();
+	SLATE_API virtual void Release();
 
 	/** Execute the "OnClicked" delegate, and get the reply */
-	FReply ExecuteOnClick();
+	SLATE_API FReply ExecuteOnClick();
 
 	/** @return combines the user-specified margin and the button's internal margin. */
-	FMargin GetCombinedPadding() const;
+	SLATE_API FMargin GetCombinedPadding() const;
 
 	/** @return True if the disabled effect should be shown. */
-	bool GetShowDisabledEffect() const;
+	SLATE_API bool GetShowDisabledEffect() const;
 
 	/** Utility function to translate other input click methods to regular ones. */
-	TEnumAsByte<EButtonClickMethod::Type> GetClickMethodFromInputType(const FPointerEvent& MouseEvent) const;
+	SLATE_API TEnumAsByte<EButtonClickMethod::Type> GetClickMethodFromInputType(const FPointerEvent& MouseEvent) const;
 
 	/** Utility function to determine if the incoming mouse event is for a precise tap or click */
-	bool IsPreciseTapOrClick(const FPointerEvent& MouseEvent) const;
+	SLATE_API bool IsPreciseTapOrClick(const FPointerEvent& MouseEvent) const;
 
 	/** Play the pressed sound */
-	void PlayPressedSound() const;
+	SLATE_API void PlayPressedSound() const;
 
 	/** Play the hovered sound */
-	void PlayHoverSound() const;
+	SLATE_API void PlayHoverSound() const;
 
 	/** Set if this button can be focused */
 	void SetIsFocusable(bool bInIsFocusable)
@@ -246,7 +246,7 @@ protected:
 		bIsFocusable = bInIsFocusable;
 	}
 
-	void ExecuteHoverStateChanged(bool bPlaySound);
+	SLATE_API void ExecuteHoverStateChanged(bool bPlaySound);
 
 protected:
 	/** @return the BorderForegroundColor attribute. */
@@ -265,13 +265,13 @@ protected:
 	TSlateAttributeRef<bool> GetAppearPressedAttribute() const { return TSlateAttributeRef<bool>(SharedThis(this), AppearPressedAttribute); }
 
 private:
-	void UpdatePressStateChanged();
+	SLATE_API void UpdatePressStateChanged();
 
-	void UpdatePadding();
-	void UpdateShowDisabledEffect();
-	void UpdateBorderImage();
-	void UpdateForegroundColor();
-	void UpdateDisabledForegroundColor();
+	SLATE_API void UpdatePadding();
+	SLATE_API void UpdateShowDisabledEffect();
+	SLATE_API void UpdateBorderImage();
+	SLATE_API void UpdateForegroundColor();
+	SLATE_API void UpdateDisabledForegroundColor();
 
 protected:
 #if WITH_EDITORONLY_DATA
diff --git a/Engine/Source/Runtime/Slate/Public/Widgets/Input/SCheckBox.h b/Engine/Source/Runtime/Slate/Public/Widgets/Input/SCheckBox.h
--- a/Engine/Source/Runtime/Slate/Public/Widgets/Input/SCheckBox.h
+++ b/Engine/Source/Runtime/Slate/Public/Widgets/Input/SCheckBox.h
@@ -26,7 +26,7 @@ DECLARE_DELEGATE_OneParam( FOnCheckStateChanged, ECheckBoxState );
 /**
  * Check box Slate control
  */
-class SLATE_API SCheckBox : public SCompoundWidget
+class SCheckBox : public SCompoundWidget
 {
 public:
 
@@ -153,28 +153,28 @@ public:
 
 	SLATE_END_ARGS()
 
-	SCheckBox();
+	SLATE_API SCheckBox();
 
 	/**
 	 * Construct this widget
 	 *
 	 * @param	InArgs	The declaration data for this widget
 	 */
-	void Construct( const FArguments& InArgs );
+	SLATE_API void Construct( const FArguments& InArgs );
 
 	// SWidget interface
-	virtual bool SupportsKeyboardFocus() const override;
-	virtual FReply OnKeyDown(const FGeometry& MyGeometry, const FKeyEvent& InKeyEvent) override;
-	virtual FReply OnKeyUp( const FGeometry& MyGeometry, const FKeyEvent& InKeyEvent ) override;
-	virtual FReply OnMouseButtonDown( const FGeometry& MyGeometry, const FPointerEvent& MouseEvent ) override;
-	virtual FReply OnMouseButtonDoubleClick( const FGeometry& InMyGeometry, const FPointerEvent& InMouseEvent ) override;
-	virtual FReply OnMouseButtonUp( const FGeometry& MyGeometry, const FPointerEvent& MouseEvent ) override;
-	virtual void OnMouseEnter( const FGeometry& MyGeometry, const FPointerEvent& MouseEvent ) override;
-	virtual void OnMouseLeave( const FPointerEvent& MouseEvent ) override;
-	virtual bool IsInteractable() const override;
-	virtual FSlateColor GetForegroundColor() const;
+	SLATE_API virtual bool SupportsKeyboardFocus() const override;
+	SLATE_API virtual FReply OnKeyDown(const FGeometry& MyGeometry, const FKeyEvent& InKeyEvent) override;
+	SLATE_API virtual FReply OnKeyUp( const FGeometry& MyGeometry, const FKeyEvent& InKeyEvent ) override;
+	SLATE_API virtual FReply OnMouseButtonDown( const FGeometry& MyGeometry, const FPointerEvent& MouseEvent ) override;
+	SLATE_API virtual FReply OnMouseButtonDoubleClick( const FGeometry& InMyGeometry, const FPointerEvent& InMouseEvent ) override;
+	SLATE_API virtual FReply OnMouseButtonUp( const FGeometry& MyGeometry, const FPointerEvent& MouseEvent ) override;
+	SLATE_API virtual void OnMouseEnter( const FGeometry& MyGeometry, const FPointerEvent& MouseEvent ) override;
+	SLATE_API virtual void OnMouseLeave( const FPointerEvent& MouseEvent ) override;
+	SLATE_API virtual bool IsInteractable() const override;
+	SLATE_API virtual FSlateColor GetForegroundColor() const;
 #if WITH_ACCESSIBILITY
-	virtual TSharedRef<FSlateAccessibleWidget> CreateAccessibleWidget() override;
+	SLATE_API virtual TSharedRef<FSlateAccessibleWidget> CreateAccessibleWidget() override;
 #endif
 	// End of SWidget interface
 
@@ -189,7 +189,7 @@ public:
 	}
 
 	/** @return The current checked state of the checkbox. */
-	ECheckBoxState GetCheckedState() const;
+	SLATE_API ECheckBoxState GetCheckedState() const;
 
 	/**
 	 * Returns true if this button is currently pressed
@@ -204,78 +204,78 @@ public:
 	/**
 	 * Toggles the checked state for this check box, fire events as needed
 	 */
-	void ToggleCheckedState();
+	SLATE_API void ToggleCheckedState();
 
 	/** See the IsChecked attribute */
-	void SetIsChecked(TAttribute<ECheckBoxState> InIsChecked);
+	SLATE_API void SetIsChecked(TAttribute<ECheckBoxState> InIsChecked);
 	
 	/** See the Content slot */
-	void SetContent(const TSharedRef< SWidget >& InContent);
+	SLATE_API void SetContent(const TSharedRef< SWidget >& InContent);
 	
 	/** See the Style attribute */
-	void SetStyle(const FCheckBoxStyle* InStyle);
+	SLATE_API void SetStyle(const FCheckBoxStyle* InStyle);
 	
 	/** See the UncheckedImage attribute */
-	void SetUncheckedImage(const FSlateBrush* Brush);
+	SLATE_API void SetUncheckedImage(const FSlateBrush* Brush);
 	/** See the UncheckedHoveredImage attribute */
-	void SetUncheckedHoveredImage(const FSlateBrush* Brush);
+	SLATE_API void SetUncheckedHoveredImage(const FSlateBrush* Brush);
 	/** See the UncheckedPressedImage attribute */
-	void SetUncheckedPressedImage(const FSlateBrush* Brush);
+	SLATE_API void SetUncheckedPressedImage(const FSlateBrush* Brush);
 	
 	/** See the CheckedImage attribute */
-	void SetCheckedImage(const FSlateBrush* Brush);
+	SLATE_API void SetCheckedImage(const FSlateBrush* Brush);
 	/** See the CheckedHoveredImage attribute */
-	void SetCheckedHoveredImage(const FSlateBrush* Brush);
+	SLATE_API void SetCheckedHoveredImage(const FSlateBrush* Brush);
 	/** See the CheckedPressedImage attribute */
-	void SetCheckedPressedImage(const FSlateBrush* Brush);
+	SLATE_API void SetCheckedPressedImage(const FSlateBrush* Brush);
 	
 	/** See the UndeterminedImage attribute */
-	void SetUndeterminedImage(const FSlateBrush* Brush);
+	SLATE_API void SetUndeterminedImage(const FSlateBrush* Brush);
 	/** See the UndeterminedHoveredImage attribute */
-	void SetUndeterminedHoveredImage(const FSlateBrush* Brush);
+	SLATE_API void SetUndeterminedHoveredImage(const FSlateBrush* Brush);
 	/** See the UndeterminedPressedImage attribute */
-	void SetUndeterminedPressedImage(const FSlateBrush* Brush);
+	SLATE_API void SetUndeterminedPressedImage(const FSlateBrush* Brush);
 
-	void SetClickMethod(EButtonClickMethod::Type InClickMethod);
-	void SetTouchMethod(EButtonTouchMethod::Type InTouchMethod);
-	void SetPressMethod(EButtonPressMethod::Type InPressMethod);
+	SLATE_API void SetClickMethod(EButtonClickMethod::Type InClickMethod);
+	SLATE_API void SetTouchMethod(EButtonTouchMethod::Type InTouchMethod);
+	SLATE_API void SetPressMethod(EButtonPressMethod::Type InPressMethod);
 
 protected:
 
 	/** Rebuilds the checkbox based on the current ESlateCheckBoxType */
-	void BuildCheckBox(TSharedRef<SWidget> InContent);
+	SLATE_API void BuildCheckBox(TSharedRef<SWidget> InContent);
 
 	/** Attribute getter for the padding */
-	FMargin OnGetPadding() const;
+	SLATE_API FMargin OnGetPadding() const;
 	/** Attribute getter for the border background color */
-	FSlateColor OnGetBorderBackgroundColor() const;
+	SLATE_API FSlateColor OnGetBorderBackgroundColor() const;
 	/** Attribute getter for the checkbox type */
-	ESlateCheckBoxType::Type OnGetCheckBoxType() const;
+	SLATE_API ESlateCheckBoxType::Type OnGetCheckBoxType() const;
 
 	/**
 	 * Gets the check image to display for the current state of the check box
 	 * @return	The name of the image to display
 	 */
-	const FSlateBrush* OnGetCheckImage() const;
+	SLATE_API const FSlateBrush* OnGetCheckImage() const;
 	
-	const FSlateBrush* GetUncheckedImage() const;
-	const FSlateBrush* GetUncheckedHoveredImage() const;
-	const FSlateBrush* GetUncheckedPressedImage() const;
+	SLATE_API const FSlateBrush* GetUncheckedImage() const;
+	SLATE_API const FSlateBrush* GetUncheckedHoveredImage() const;
+	SLATE_API const FSlateBrush* GetUncheckedPressedImage() const;
 	
-	const FSlateBrush* GetCheckedImage() const;
-	const FSlateBrush* GetCheckedHoveredImage() const;
-	const FSlateBrush* GetCheckedPressedImage() const;
+	SLATE_API const FSlateBrush* GetCheckedImage() const;
+	SLATE_API const FSlateBrush* GetCheckedHoveredImage() const;
+	SLATE_API const FSlateBrush* GetCheckedPressedImage() const;
 	
-	const FSlateBrush* GetUndeterminedImage() const;
-	const FSlateBrush* GetUndeterminedHoveredImage() const;
-	const FSlateBrush* GetUndeterminedPressedImage() const;
+	SLATE_API const FSlateBrush* GetUndeterminedImage() const;
+	SLATE_API const FSlateBrush* GetUndeterminedHoveredImage() const;
+	SLATE_API const FSlateBrush* GetUndeterminedPressedImage() const;
 
 	/** Attribute getter for the background image */
-	const FSlateBrush* OnGetBackgroundImage() const;
+	SLATE_API const FSlateBrush* OnGetBackgroundImage() const;
 
-	const FSlateBrush* GetBackgroundImage() const;
-	const FSlateBrush* GetBackgroundHoveredImage() const;
-	const FSlateBrush* GetBackgroundPressedImage() const;
+	SLATE_API const FSlateBrush* GetBackgroundImage() const;
+	SLATE_API const FSlateBrush* GetBackgroundHoveredImage() const;
+	SLATE_API const FSlateBrush* GetBackgroundPressedImage() const;
 
 	
 protected:
@@ -347,19 +347,19 @@ protected:
 	FOnGetContent OnGetMenuContent;
 
 	/** Play the checked sound */
-	void PlayCheckedSound() const;
+	SLATE_API void PlayCheckedSound() const;
 
 	/** Play the unchecked sound */
-	void PlayUncheckedSound() const;
+	SLATE_API void PlayUncheckedSound() const;
 
 	/** Play the hovered sound */
-	void PlayHoverSound() const;
+	SLATE_API void PlayHoverSound() const;
 
 	/** Utility function to translate other input click methods to regular ones. */
-	TEnumAsByte<EButtonClickMethod::Type> GetClickMethodFromInputType(const FPointerEvent& MouseEvent) const;
+	SLATE_API TEnumAsByte<EButtonClickMethod::Type> GetClickMethodFromInputType(const FPointerEvent& MouseEvent) const;
 
 	/** Utility function to determine if the incoming mouse event is for a precise tap or click */
-	bool IsPreciseTapOrClick(const FPointerEvent& MouseEvent) const;
+	SLATE_API bool IsPreciseTapOrClick(const FPointerEvent& MouseEvent) const;
 
 	/** The Sound to play when the check box is hovered  */
 	FSlateSound HoveredSound;
diff --git a/Engine/Source/Runtime/Slate/Public/Widgets/Input/SComboButton.h b/Engine/Source/Runtime/Slate/Public/Widgets/Input/SComboButton.h
--- a/Engine/Source/Runtime/Slate/Public/Widgets/Input/SComboButton.h
+++ b/Engine/Source/Runtime/Slate/Public/Widgets/Input/SComboButton.h
@@ -21,7 +21,7 @@ class SImage;
 /**
  * A button that, when clicked, brings up a popup.
  */
-class SLATE_API SComboButton : public SMenuAnchor
+class SComboButton : public SMenuAnchor
 {
 public:
 
@@ -75,18 +75,18 @@ public:
 	SLATE_END_ARGS()
 
 	// SMenuAnchor interface
-	virtual void SetMenuContent(TSharedRef<SWidget> InContent) override;
+	SLATE_API virtual void SetMenuContent(TSharedRef<SWidget> InContent) override;
 	// End of SMenuAnchor interface
 
 	/** See the OnGetMenuContent event */
-	void SetOnGetMenuContent( FOnGetContent InOnGetMenuContent );
+	SLATE_API void SetOnGetMenuContent( FOnGetContent InOnGetMenuContent );
 
 	/**
 	 * Construct the widget from a declaration
 	 *
 	 * @param InArgs  The declaration from which to construct
 	 */
-	void Construct(const FArguments& InArgs);
+	SLATE_API void Construct(const FArguments& InArgs);
 
 	void SetMenuContentWidgetToFocus( TWeakPtr<SWidget> InWidgetToFocusPtr )
 	{
@@ -94,17 +94,17 @@ public:
 	}
 
 	/** See the padding for button content. */
-	void SetButtonContentPadding(FMargin InPadding);
+	SLATE_API void SetButtonContentPadding(FMargin InPadding);
 
 	/** add/remove the expanding arrow. */
-	void SetHasDownArrow(bool InHasArrowDown);
+	SLATE_API void SetHasDownArrow(bool InHasArrowDown);
 
 protected:
 	/**
 	 * Handle the button being clicked by summoning the ComboButton.
 	 */
-	virtual FReply OnButtonClicked();
-	virtual FReply OnKeyDown(const FGeometry& MyGeometry, const FKeyEvent& InKeyEvent) override;
+	SLATE_API virtual FReply OnButtonClicked();
+	SLATE_API virtual FReply OnKeyDown(const FGeometry& MyGeometry, const FKeyEvent& InKeyEvent) override;
 
 	/**
 	 * Called to query the tool tip text for this widget, but will return an empty text when the menu is already open
@@ -113,7 +113,7 @@ protected:
 	 *
 	 * @return	Tool tip text, or an empty text if filtered out
 	 */
-	FText GetFilteredToolTipText(TAttribute<FText> ToolTipText) const;
+	SLATE_API FText GetFilteredToolTipText(TAttribute<FText> ToolTipText) const;
 
 protected:
 	/** Area where the button's content resides */
