[Slate]

diff --git a/Engine/Source/Runtime/Slate/Public/Widgets/Input/SEditableText.h b/Engine/Source/Runtime/Slate/Public/Widgets/Input/SEditableText.h
--- a/Engine/Source/Runtime/Slate/Public/Widgets/Input/SEditableText.h
+++ b/Engine/Source/Runtime/Slate/Public/Widgets/Input/SEditableText.h
@@ -34,7 +34,7 @@ enum class ETextShapingMethod : uint8;
 /**
  * Editable text widget
  */
-class SLATE_API SEditableText : public SWidget, public ISlateEditableTextWidget
+class SEditableText : public SWidget, public ISlateEditableTextWidget
 {
 public:
 	SLATE_BEGIN_ARGS( SEditableText )
@@ -178,75 +178,75 @@ public:
 	SLATE_END_ARGS()
 
 	/** Constructor */
-	SEditableText();
+	SLATE_API SEditableText();
 
 	/** Destructor */
-	~SEditableText();
+	SLATE_API ~SEditableText();
 
 	/**
 	 * Construct this widget
 	 *
 	 * @param	InArgs	The declaration data for this widget
 	 */
-	void Construct( const FArguments& InArgs );
+	SLATE_API void Construct( const FArguments& InArgs );
 
 	/**
 	 * Sets the text currently being edited 
 	 *
 	 * @param  InNewText  The new text
 	 */
-	void SetText( const TAttribute< FText >& InNewText );
+	SLATE_API void SetText( const TAttribute< FText >& InNewText );
 	
 	/**
 	 * Returns the text string
 	 *
 	 * @return  Text string
 	 */
-	FText GetText() const;
+	SLATE_API FText GetText() const;
 
 	/** See the HintText attribute */
-	void SetHintText( const TAttribute< FText >& InHintText );
+	SLATE_API void SetHintText( const TAttribute< FText >& InHintText );
 	
 	/** Get the text that appears when there is no text in the text box */
-	FText GetHintText() const;
+	SLATE_API FText GetHintText() const;
 
 	/** Set the text that is currently being searched for (if any) */
-	void SetSearchText(const TAttribute<FText>& InSearchText);
+	SLATE_API void SetSearchText(const TAttribute<FText>& InSearchText);
 
 	/** Get the text that is currently being searched for (if any) */
-	FText GetSearchText() const;
+	SLATE_API FText GetSearchText() const;
 
 	/** See the IsReadOnly attribute */
-	void SetIsReadOnly( TAttribute< bool > InIsReadOnly );
+	SLATE_API void SetIsReadOnly( TAttribute< bool > InIsReadOnly );
 	
 	/** See the IsPassword attribute */
-	void SetIsPassword( TAttribute< bool > InIsPassword );
+	SLATE_API void SetIsPassword( TAttribute< bool > InIsPassword );
 	
 	/** See the ColorAndOpacity attribute */
-	void SetColorAndOpacity(TAttribute<FSlateColor> Color);
+	SLATE_API void SetColorAndOpacity(TAttribute<FSlateColor> Color);
 
 	/** See the AllowContextMenu attribute */
-	void SetAllowContextMenu(const TAttribute< bool >& InAllowContextMenu);
+	SLATE_API void SetAllowContextMenu(const TAttribute< bool >& InAllowContextMenu);
 
 	/** Set the VirtualKeyboardDismissAction attribute */
-	void SetVirtualKeyboardDismissAction(TAttribute< EVirtualKeyboardDismissAction > InVirtualKeyboardDismissAction);
+	SLATE_API void SetVirtualKeyboardDismissAction(TAttribute< EVirtualKeyboardDismissAction > InVirtualKeyboardDismissAction);
 
 	/**
 	 * Sets the font used to draw the text
 	 *
 	 * @param  InNewFont	The new font to use
 	 */
-	void SetFont(const TAttribute< FSlateFontInfo >& InNewFont);
+	SLATE_API void SetFont(const TAttribute< FSlateFontInfo >& InNewFont);
 
 	/** Gets the font used to draw the text. */
-	FSlateFontInfo GetFont() const;
+	SLATE_API FSlateFontInfo GetFont() const;
 
 	/**
 	 * Sets the text style used to draw the text
 	 *
 	 * @param  NewTextStyle	The new text style to use
 	 */
-	void SetTextStyle( const FEditableTextStyle& InNewTextStyle );
+	SLATE_API void SetTextStyle( const FEditableTextStyle& InNewTextStyle );
 
 	/** @See TextStyle */
 	/**
@@ -254,59 +254,59 @@ public:
 	 *
 	 * @param  InTextStyle	The new text block style to use
 	 */
-	void SetTextBlockStyle(const FTextBlockStyle* InTextStyle);
+	SLATE_API void SetTextBlockStyle(const FTextBlockStyle* InTextStyle);
 
 	/**
 	 * Sets the minimum width that a text block should be.
 	 *
 	 * @param  InMinDesiredWidth	The minimum width
 	 */
-	void SetMinDesiredWidth(const TAttribute<float>& InMinDesiredWidth);
+	SLATE_API void SetMinDesiredWidth(const TAttribute<float>& InMinDesiredWidth);
 
 	/**
 	 * Workaround as we loose focus when the auto completion closes.
 	 *
 	 * @param  InIsCaretMovedWhenGainFocus	Workaround
 	 */
-	void SetIsCaretMovedWhenGainFocus(const TAttribute<bool>& InIsCaretMovedWhenGainFocus);
+	SLATE_API void SetIsCaretMovedWhenGainFocus(const TAttribute<bool>& InIsCaretMovedWhenGainFocus);
 
 	/**
 	 * Sets whether to select all text when the user clicks to give focus on the widget
 	 *
 	 * @param  InSelectAllTextWhenFocused	Select all text when the user clicks?
 	 */
-	void SetSelectAllTextWhenFocused(const TAttribute<bool>& InSelectAllTextWhenFocused);
+	SLATE_API void SetSelectAllTextWhenFocused(const TAttribute<bool>& InSelectAllTextWhenFocused);
 
 	/**
 	 * Sets whether to allow the user to back out of changes when they press the escape key
 	 *
 	 * @param  InRevertTextOnEscape			Allow the user to back out of changes?
 	 */
-	void SetRevertTextOnEscape(const TAttribute<bool>& InRevertTextOnEscape);
+	SLATE_API void SetRevertTextOnEscape(const TAttribute<bool>& InRevertTextOnEscape);
 
 	/**
 	 * Sets whether to clear keyboard focus when pressing enter to commit changes
 	 *
 	 * @param  InClearKeyboardFocusOnCommit		Clear keyboard focus when pressing enter?
 	 */
-	void SetClearKeyboardFocusOnCommit(const TAttribute<bool>& InClearKeyboardFocusOnCommit);
+	SLATE_API void SetClearKeyboardFocusOnCommit(const TAttribute<bool>& InClearKeyboardFocusOnCommit);
 
 	/**
 	 * Sets whether to select all text when pressing enter to commit changes
 	 *
 	 * @param  InSelectAllTextOnCommit		Select all text when pressing enter?
 	 */
-	void SetSelectAllTextOnCommit(const TAttribute<bool>& InSelectAllTextOnCommit);
+	SLATE_API void SetSelectAllTextOnCommit(const TAttribute<bool>& InSelectAllTextOnCommit);
 
 	/**
 	 * Sets whether to select word on the mouse double click
 	 *
 	 * @param  InSelectWordOnMouseDoubleClick		Select word on the mouse double click
 	 */
-	void SetSelectWordOnMouseDoubleClick(const TAttribute<bool>& InSelectWordOnMouseDoubleClick);
+	SLATE_API void SetSelectWordOnMouseDoubleClick(const TAttribute<bool>& InSelectWordOnMouseDoubleClick);
 
 	/** See Justification attribute */
-	void SetJustification(const TAttribute<ETextJustify::Type>& InJustification);
+	SLATE_API void SetJustification(const TAttribute<ETextJustify::Type>& InJustification);
 	/**
 	 * Sets the OnKeyCharHandler to provide first chance handling of the OnKeyChar event
 	 *
@@ -328,112 +328,112 @@ public:
 	}
 
 	/** See TextShapingMethod attribute */
-	void SetTextShapingMethod(const TOptional<ETextShapingMethod>& InTextShapingMethod);
+	SLATE_API void SetTextShapingMethod(const TOptional<ETextShapingMethod>& InTextShapingMethod);
 
 	/** See TextFlowDirection attribute */
-	void SetTextFlowDirection(const TOptional<ETextFlowDirection>& InTextFlowDirection);
+	SLATE_API void SetTextFlowDirection(const TOptional<ETextFlowDirection>& InTextFlowDirection);
 
 	/** Sets the overflow policy for this text block */
-	void SetOverflowPolicy(TOptional<ETextOverflowPolicy> InOverflowPolicy);
+	SLATE_API void SetOverflowPolicy(TOptional<ETextOverflowPolicy> InOverflowPolicy);
 
 	/** Query to see if any text is selected within the document */
-	bool AnyTextSelected() const;
+	SLATE_API bool AnyTextSelected() const;
 
 	/** Select all the text in the document */
-	void SelectAllText();
+	SLATE_API void SelectAllText();
 
 	/** Clear the active text selection */
-	void ClearSelection();
+	SLATE_API void ClearSelection();
 
 	/** Get the currently selected text */
-	FText GetSelectedText() const;
+	SLATE_API FText GetSelectedText() const;
 
 	/** Move the cursor to the given location in the document (will also scroll to this point) */
-	void GoTo(const FTextLocation& NewLocation);
+	SLATE_API void GoTo(const FTextLocation& NewLocation);
 
 	/** Move the cursor specified location */
-	void GoTo(const ETextLocation NewLocation);
+	SLATE_API void GoTo(const ETextLocation NewLocation);
 
 	/** Scroll to the given location in the document (without moving the cursor) */
-	void ScrollTo(const FTextLocation& NewLocation);
+	SLATE_API void ScrollTo(const FTextLocation& NewLocation);
 
 	/** Scroll to the given location in the document (without moving the cursor) */
-	void ScrollTo(const ETextLocation NewLocation);
+	SLATE_API void ScrollTo(const ETextLocation NewLocation);
 
 	/** Begin a new text search (this is called automatically when the bound search text changes) */
-	void BeginSearch(const FText& InSearchText, const ESearchCase::Type InSearchCase = ESearchCase::IgnoreCase, const bool InReverse = false);
+	SLATE_API void BeginSearch(const FText& InSearchText, const ESearchCase::Type InSearchCase = ESearchCase::IgnoreCase, const bool InReverse = false);
 
 	/** Advance the current search to the next match (does nothing if not currently searching) */
-	void AdvanceSearch(const bool InReverse = false);
+	SLATE_API void AdvanceSearch(const bool InReverse = false);
 
 	/** Register and activate the IME context for this text layout */
-	void EnableTextInputMethodContext();
+	SLATE_API void EnableTextInputMethodContext();
 
 protected:
 	//~ Begin SWidget Interface
-	virtual void Tick( const FGeometry& AllottedGeometry, const double InCurrentTime, const float InDeltaTime ) override;
-	virtual int32 OnPaint(const FPaintArgs& Args, const FGeometry& AllottedGeometry, const FSlateRect& MyCullingRect, FSlateWindowElementList& OutDrawElements, int32 LayerId, const FWidgetStyle& InWidgetStyle, bool bParentEnabled) const override;
-	virtual void CacheDesiredSize(float LayoutScaleMultiplier) override;
-	virtual FVector2D ComputeDesiredSize(float LayoutScaleMultiplier) const override;
-	virtual FChildren* GetChildren() override;
-	virtual void OnArrangeChildren(const FGeometry& AllottedGeometry, FArrangedChildren& ArrangedChildren) const override;
-	virtual FReply OnDragOver( const FGeometry& MyGeometry, const FDragDropEvent& DragDropEvent ) override;
-	virtual FReply OnDrop( const FGeometry& MyGeometry, const FDragDropEvent& DragDropEvent ) override;
-	virtual bool SupportsKeyboardFocus() const override;
-	virtual FReply OnFocusReceived( const FGeometry& MyGeometry, const FFocusEvent& InFocusEvent ) override;
-	virtual void OnFocusLost( const FFocusEvent& InFocusEvent ) override;
-	virtual FReply OnKeyChar( const FGeometry& MyGeometry,  const FCharacterEvent& InCharacterEvent ) override;
-	virtual FReply OnKeyDown( const FGeometry& MyGeometry, const FKeyEvent& InKeyEvent ) override;
-	virtual FReply OnKeyUp( const FGeometry& MyGeometry, const FKeyEvent& InKeyEvent ) override;
-	virtual FReply OnMouseButtonDown( const FGeometry& InMyGeometry, const FPointerEvent& InMouseEvent ) override;
-	virtual FReply OnMouseButtonUp( const FGeometry& InMyGeometry, const FPointerEvent& InMouseEvent ) override;
-	virtual FReply OnMouseMove( const FGeometry& InMyGeometry, const FPointerEvent& InMouseEvent ) override;
-	virtual FReply OnMouseButtonDoubleClick( const FGeometry& InMyGeometry, const FPointerEvent& InMouseEvent ) override;
-	virtual FCursorReply OnCursorQuery( const FGeometry& MyGeometry, const FPointerEvent& CursorEvent ) const override;
-	virtual const FSlateBrush* GetFocusBrush() const;
-	virtual bool IsInteractable() const override;
-	virtual bool ComputeVolatility() const override;
+	SLATE_API virtual void Tick( const FGeometry& AllottedGeometry, const double InCurrentTime, const float InDeltaTime ) override;
+	SLATE_API virtual int32 OnPaint(const FPaintArgs& Args, const FGeometry& AllottedGeometry, const FSlateRect& MyCullingRect, FSlateWindowElementList& OutDrawElements, int32 LayerId, const FWidgetStyle& InWidgetStyle, bool bParentEnabled) const override;
+	SLATE_API virtual void CacheDesiredSize(float LayoutScaleMultiplier) override;
+	SLATE_API virtual FVector2D ComputeDesiredSize(float LayoutScaleMultiplier) const override;
+	SLATE_API virtual FChildren* GetChildren() override;
+	SLATE_API virtual void OnArrangeChildren(const FGeometry& AllottedGeometry, FArrangedChildren& ArrangedChildren) const override;
+	SLATE_API virtual FReply OnDragOver( const FGeometry& MyGeometry, const FDragDropEvent& DragDropEvent ) override;
+	SLATE_API virtual FReply OnDrop( const FGeometry& MyGeometry, const FDragDropEvent& DragDropEvent ) override;
+	SLATE_API virtual bool SupportsKeyboardFocus() const override;
+	SLATE_API virtual FReply OnFocusReceived( const FGeometry& MyGeometry, const FFocusEvent& InFocusEvent ) override;
+	SLATE_API virtual void OnFocusLost( const FFocusEvent& InFocusEvent ) override;
+	SLATE_API virtual FReply OnKeyChar( const FGeometry& MyGeometry,  const FCharacterEvent& InCharacterEvent ) override;
+	SLATE_API virtual FReply OnKeyDown( const FGeometry& MyGeometry, const FKeyEvent& InKeyEvent ) override;
+	SLATE_API virtual FReply OnKeyUp( const FGeometry& MyGeometry, const FKeyEvent& InKeyEvent ) override;
+	SLATE_API virtual FReply OnMouseButtonDown( const FGeometry& InMyGeometry, const FPointerEvent& InMouseEvent ) override;
+	SLATE_API virtual FReply OnMouseButtonUp( const FGeometry& InMyGeometry, const FPointerEvent& InMouseEvent ) override;
+	SLATE_API virtual FReply OnMouseMove( const FGeometry& InMyGeometry, const FPointerEvent& InMouseEvent ) override;
+	SLATE_API virtual FReply OnMouseButtonDoubleClick( const FGeometry& InMyGeometry, const FPointerEvent& InMouseEvent ) override;
+	SLATE_API virtual FCursorReply OnCursorQuery( const FGeometry& MyGeometry, const FPointerEvent& CursorEvent ) const override;
+	SLATE_API virtual const FSlateBrush* GetFocusBrush() const;
+	SLATE_API virtual bool IsInteractable() const override;
+	SLATE_API virtual bool ComputeVolatility() const override;
 #if WITH_ACCESSIBILITY
-	virtual TSharedRef<FSlateAccessibleWidget> CreateAccessibleWidget() override;
-	virtual TOptional<FText> GetDefaultAccessibleText(EAccessibleType AccessibleType = EAccessibleType::Main) const override;
+	SLATE_API virtual TSharedRef<FSlateAccessibleWidget> CreateAccessibleWidget() override;
+	SLATE_API virtual TOptional<FText> GetDefaultAccessibleText(EAccessibleType AccessibleType = EAccessibleType::Main) const override;
 #endif
 	//~ End SWidget Interface
 
 protected:
 	/** Synchronize the text style currently set (including from overrides) and update the text layout if required */
-	void SynchronizeTextStyle();
+	SLATE_API void SynchronizeTextStyle();
 
 public:
 	//~ Begin ISlateEditableTextWidget Interface
-	virtual bool IsTextReadOnly() const override;
-	virtual bool IsTextPassword() const override;
-	virtual bool IsMultiLineTextEdit() const override;
+	SLATE_API virtual bool IsTextReadOnly() const override;
+	SLATE_API virtual bool IsTextPassword() const override;
+	SLATE_API virtual bool IsMultiLineTextEdit() const override;
 	//~ End ISlateEditableTextWidget Interface
 
 protected:
 	//~ Begin ISlateEditableTextWidget Interface
-	virtual bool ShouldJumpCursorToEndWhenFocused() const override;
-	virtual bool ShouldSelectAllTextWhenFocused() const override;
-	virtual bool ShouldClearTextSelectionOnFocusLoss() const override;
-	virtual bool ShouldRevertTextOnEscape() const override;
-	virtual bool ShouldClearKeyboardFocusOnCommit() const override;
-	virtual bool ShouldSelectAllTextOnCommit() const override;
-	virtual bool ShouldSelectWordOnMouseDoubleClick() const override;
-	virtual bool CanInsertCarriageReturn() const override;
-	virtual bool CanTypeCharacter(const TCHAR InChar) const override;
-	virtual void EnsureActiveTick() override;
-	virtual EKeyboardType GetVirtualKeyboardType() const override;
-	virtual FVirtualKeyboardOptions GetVirtualKeyboardOptions() const override;
-	virtual EVirtualKeyboardTrigger GetVirtualKeyboardTrigger() const override;
-	virtual EVirtualKeyboardDismissAction GetVirtualKeyboardDismissAction() const override;
-	virtual TSharedRef<SWidget> GetSlateWidget() override;
-	virtual TSharedPtr<SWidget> GetSlateWidgetPtr() override;
-	virtual TSharedPtr<SWidget> BuildContextMenuContent() const override;
-	virtual void OnTextChanged(const FText& InText) override;
-	virtual void OnTextCommitted(const FText& InText, const ETextCommit::Type InTextAction) override;
-	virtual void OnCursorMoved(const FTextLocation& InLocation) override;
-	virtual float UpdateAndClampHorizontalScrollBar(const float InViewOffset, const float InViewFraction, const EVisibility InVisiblityOverride) override;
-	virtual float UpdateAndClampVerticalScrollBar(const float InViewOffset, const float InViewFraction, const EVisibility InVisiblityOverride) override;
+	SLATE_API virtual bool ShouldJumpCursorToEndWhenFocused() const override;
+	SLATE_API virtual bool ShouldSelectAllTextWhenFocused() const override;
+	SLATE_API virtual bool ShouldClearTextSelectionOnFocusLoss() const override;
+	SLATE_API virtual bool ShouldRevertTextOnEscape() const override;
+	SLATE_API virtual bool ShouldClearKeyboardFocusOnCommit() const override;
+	SLATE_API virtual bool ShouldSelectAllTextOnCommit() const override;
+	SLATE_API virtual bool ShouldSelectWordOnMouseDoubleClick() const override;
+	SLATE_API virtual bool CanInsertCarriageReturn() const override;
+	SLATE_API virtual bool CanTypeCharacter(const TCHAR InChar) const override;
+	SLATE_API virtual void EnsureActiveTick() override;
+	SLATE_API virtual EKeyboardType GetVirtualKeyboardType() const override;
+	SLATE_API virtual FVirtualKeyboardOptions GetVirtualKeyboardOptions() const override;
+	SLATE_API virtual EVirtualKeyboardTrigger GetVirtualKeyboardTrigger() const override;
+	SLATE_API virtual EVirtualKeyboardDismissAction GetVirtualKeyboardDismissAction() const override;
+	SLATE_API virtual TSharedRef<SWidget> GetSlateWidget() override;
+	SLATE_API virtual TSharedPtr<SWidget> GetSlateWidgetPtr() override;
+	SLATE_API virtual TSharedPtr<SWidget> BuildContextMenuContent() const override;
+	SLATE_API virtual void OnTextChanged(const FText& InText) override;
+	SLATE_API virtual void OnTextCommitted(const FText& InText, const ETextCommit::Type InTextAction) override;
+	SLATE_API virtual void OnCursorMoved(const FTextLocation& InLocation) override;
+	SLATE_API virtual float UpdateAndClampHorizontalScrollBar(const float InViewOffset, const float InViewFraction, const EVisibility InVisiblityOverride) override;
+	SLATE_API virtual float UpdateAndClampVerticalScrollBar(const float InViewOffset, const float InViewFraction, const EVisibility InVisiblityOverride) override;
 	//~ End ISlateEditableTextWidget Interface
 
 protected:
diff --git a/Engine/Source/Runtime/Slate/Public/Widgets/Input/SEditableTextBox.h b/Engine/Source/Runtime/Slate/Public/Widgets/Input/SEditableTextBox.h
--- a/Engine/Source/Runtime/Slate/Public/Widgets/Input/SEditableTextBox.h
+++ b/Engine/Source/Runtime/Slate/Public/Widgets/Input/SEditableTextBox.h
@@ -25,7 +25,7 @@ enum class ETextShapingMethod : uint8;
 /**
  * Editable text box widget
  */
-class SLATE_API SEditableTextBox : public SBorder
+class SEditableTextBox : public SBorder
 {
 
 public:
@@ -169,14 +169,14 @@ public:
 		SLATE_ARGUMENT(TOptional<ETextOverflowPolicy>, OverflowPolicy)
 	SLATE_END_ARGS()
 
-	SEditableTextBox();
+	SLATE_API SEditableTextBox();
 	
 	/**
 	 * Construct this widget
 	 *
 	 * @param	InArgs	The declaration data for this widget
 	 */
-	void Construct( const FArguments& InArgs );
+	SLATE_API void Construct( const FArguments& InArgs );
 
 	/**
 	 * Returns the text string
@@ -189,21 +189,21 @@ public:
 	}
 
 	/** See attribute Style */
-	void SetStyle(const FEditableTextBoxStyle* InStyle);
+	SLATE_API void SetStyle(const FEditableTextBoxStyle* InStyle);
 
 	/**
 	 * Sets the text block style currently used
 	 *
 	 * @param  InTextStyle  The new text string
 	 */
-	void SetTextBlockStyle(const FTextBlockStyle* InTextStyle);
+	SLATE_API void SetTextBlockStyle(const FTextBlockStyle* InTextStyle);
 
 	/**
 	 * Sets the text string currently being edited 
 	 *
 	 * @param  InNewText  The new text string
 	 */
-	void SetText( const TAttribute< FText >& InNewText );
+	SLATE_API void SetText( const TAttribute< FText >& InNewText );
 	
 	/**
 	 * Returns the Hint text string
@@ -216,161 +216,161 @@ public:
 	}
 
 	/** See the HintText attribute */
-	void SetHintText( const TAttribute< FText >& InHintText );
+	SLATE_API void SetHintText( const TAttribute< FText >& InHintText );
 	
 	/** Set the text that is currently being searched for (if any) */
-	void SetSearchText(const TAttribute<FText>& InSearchText);
+	SLATE_API void SetSearchText(const TAttribute<FText>& InSearchText);
 
 	/** Get the text that is currently being searched for (if any) */
-	FText GetSearchText() const;
+	SLATE_API FText GetSearchText() const;
 
 	/** See the IsReadOnly attribute */
-	void SetIsReadOnly( TAttribute< bool > InIsReadOnly );
+	SLATE_API void SetIsReadOnly( TAttribute< bool > InIsReadOnly );
 
 	bool IsReadOnly() const { return EditableText->IsTextReadOnly(); }
 	
 	/** See the IsPassword attribute */
-	void SetIsPassword( TAttribute< bool > InIsPassword );
+	SLATE_API void SetIsPassword( TAttribute< bool > InIsPassword );
 
 	bool IsPassword() const { return EditableText->IsTextPassword(); }
 
 	/** See the AllowContextMenu attribute */
-	void SetAllowContextMenu(TAttribute< bool > InAllowContextMenu);
+	SLATE_API void SetAllowContextMenu(TAttribute< bool > InAllowContextMenu);
 
 	/** Set the VirtualKeyboardDismissAction attribute */
-	void SetVirtualKeyboardDismissAction(TAttribute< EVirtualKeyboardDismissAction > InVirtualKeyboardDismissAction);
+	SLATE_API void SetVirtualKeyboardDismissAction(TAttribute< EVirtualKeyboardDismissAction > InVirtualKeyboardDismissAction);
 	
 	/**
 	 * Sets the font used to draw the text
 	 *
 	 * @param  InFont	The new font to use
 	 */
-	void SetFont(const TAttribute<FSlateFontInfo>& InFont);
+	SLATE_API void SetFont(const TAttribute<FSlateFontInfo>& InFont);
 
 	/**
 	 * Sets the text color and opacity (overrides Style)
 	 *
 	 * @param  InForegroundColor 	The text color and opacity
 	 */
-	void SetTextBoxForegroundColor(const TAttribute<FSlateColor>& InForegroundColor);
+	SLATE_API void SetTextBoxForegroundColor(const TAttribute<FSlateColor>& InForegroundColor);
 
 	/**
 	 * Sets the color of the background/border around the editable text (overrides Style) 
 	 *
 	 * @param  InBackgroundColor 	The background/border color
 	 */
-	void SetTextBoxBackgroundColor(const TAttribute<FSlateColor>& InBackgroundColor);
+	SLATE_API void SetTextBoxBackgroundColor(const TAttribute<FSlateColor>& InBackgroundColor);
 
 	/**
 	 * Sets the text color and opacity when read-only (overrides Style) 
 	 *
 	 * @param  InReadOnlyForegroundColor 	The read-only text color and opacity
 	 */
-	void SetReadOnlyForegroundColor(const TAttribute<FSlateColor>& InReadOnlyForegroundColor);
+	SLATE_API void SetReadOnlyForegroundColor(const TAttribute<FSlateColor>& InReadOnlyForegroundColor);
 
 	/**
 	 * Sets the text color and opacity when this box has keyboard focus(overrides Style)
 	 *
 	 * @param  InFocusedForegroundColor 	The focused color and opacity
 	 */
-	void SetFocusedForegroundColor(const TAttribute<FSlateColor>& InFocusedForegroundColor);
+	SLATE_API void SetFocusedForegroundColor(const TAttribute<FSlateColor>& InFocusedForegroundColor);
 
 	/**
 	 * Sets the minimum width that a text box should be.
 	 *
 	 * @param  InMinimumDesiredWidth	The minimum width
 	 */
-	void SetMinimumDesiredWidth(const TAttribute<float>& InMinimumDesiredWidth);
+	SLATE_API void SetMinimumDesiredWidth(const TAttribute<float>& InMinimumDesiredWidth);
 
 	/**
 	 * Workaround as we loose focus when the auto completion closes.
 	 *
 	 * @param  InIsCaretMovedWhenGainFocus	Workaround
 	 */
-	void SetIsCaretMovedWhenGainFocus(const TAttribute<bool>& InIsCaretMovedWhenGainFocus);
+	SLATE_API void SetIsCaretMovedWhenGainFocus(const TAttribute<bool>& InIsCaretMovedWhenGainFocus);
 
 	/**
 	 * Sets whether to select all text when the user clicks to give focus on the widget
 	 *
 	 * @param  InSelectAllTextWhenFocused	Select all text when the user clicks?
 	 */
-	void SetSelectAllTextWhenFocused(const TAttribute<bool>& InSelectAllTextWhenFocused);
+	SLATE_API void SetSelectAllTextWhenFocused(const TAttribute<bool>& InSelectAllTextWhenFocused);
 
 	/**
 	 * Sets whether to allow the user to back out of changes when they press the escape key
 	 *
 	 * @param  InRevertTextOnEscape			Allow the user to back out of changes?
 	 */
-	void SetRevertTextOnEscape(const TAttribute<bool>& InRevertTextOnEscape);
+	SLATE_API void SetRevertTextOnEscape(const TAttribute<bool>& InRevertTextOnEscape);
 
 	/**
 	 * Sets whether to clear keyboard focus when pressing enter to commit changes
 	 *
 	 * @param  InClearKeyboardFocusOnCommit		Clear keyboard focus when pressing enter?
 	 */
-	void SetClearKeyboardFocusOnCommit(const TAttribute<bool>& InClearKeyboardFocusOnCommit);
+	SLATE_API void SetClearKeyboardFocusOnCommit(const TAttribute<bool>& InClearKeyboardFocusOnCommit);
 
 	/**
 	 * Sets whether to select all text when pressing enter to commit changes
 	 *
 	 * @param  InSelectAllTextOnCommit		Select all text when pressing enter?
 	 */
-	void SetSelectAllTextOnCommit(const TAttribute<bool>& InSelectAllTextOnCommit);
+	SLATE_API void SetSelectAllTextOnCommit(const TAttribute<bool>& InSelectAllTextOnCommit);
 
 	/**
 	 * Sets whether to select select word on mouse double click
 	 *
 	 * @param  InSelectWordOnMouseDoubleClick		Select select word on mouse double click?
 	 */
-	void SetSelectWordOnMouseDoubleClick(const TAttribute<bool>& InSelectWordOnMouseDoubleClick);
+	SLATE_API void SetSelectWordOnMouseDoubleClick(const TAttribute<bool>& InSelectWordOnMouseDoubleClick);
 
 	/** See Justification attribute */
-	void SetJustification(const TAttribute<ETextJustify::Type>& InJustification);
+	SLATE_API void SetJustification(const TAttribute<ETextJustify::Type>& InJustification);
 
 	/**
 	 * If InError is a non-empty string the TextBox will the ErrorReporting provided during construction
 	 * If no error reporting was provided, the TextBox will create a default error reporter.
 	 */
-	void SetError( const FText& InError );
-	void SetError( const FString& InError );
+	SLATE_API void SetError( const FText& InError );
+	SLATE_API void SetError( const FString& InError );
 
 	/**
 	 * Sets the OnKeyCharHandler to provide first chance handling of the SEditableText's OnKeyChar event
 	 *
 	 * @param InOnKeyCharHandler			Delegate to call during OnKeyChar event
 	 */
-	void SetOnKeyCharHandler(FOnKeyChar InOnKeyCharHandler);
+	SLATE_API void SetOnKeyCharHandler(FOnKeyChar InOnKeyCharHandler);
 
 	/**
 	 * Sets the OnKeyDownHandler to provide first chance handling of the SEditableText's OnKeyDown event
 	 *
 	 * @param InOnKeyDownHandler			Delegate to call during OnKeyDown event
 	 */
-	void SetOnKeyDownHandler(FOnKeyDown InOnKeyDownHandler);
+	SLATE_API void SetOnKeyDownHandler(FOnKeyDown InOnKeyDownHandler);
 
 	/** See TextShapingMethod attribute */
-	void SetTextShapingMethod(const TOptional<ETextShapingMethod>& InTextShapingMethod);
+	SLATE_API void SetTextShapingMethod(const TOptional<ETextShapingMethod>& InTextShapingMethod);
 
 	/** See TextFlowDirection attribute */
-	void SetTextFlowDirection(const TOptional<ETextFlowDirection>& InTextFlowDirection);
+	SLATE_API void SetTextFlowDirection(const TOptional<ETextFlowDirection>& InTextFlowDirection);
 
 	/** Sets the overflow policy for this text block */
-	void SetOverflowPolicy(TOptional<ETextOverflowPolicy> InOverflowPolicy);
+	SLATE_API void SetOverflowPolicy(TOptional<ETextOverflowPolicy> InOverflowPolicy);
 
 	/** Query to see if any text is selected within the document */
-	bool AnyTextSelected() const;
+	SLATE_API bool AnyTextSelected() const;
 
 	/** Select all the text in the document */
-	void SelectAllText();
+	SLATE_API void SelectAllText();
 
 	/** Clear the active text selection */
-	void ClearSelection();
+	SLATE_API void ClearSelection();
 
 	/** Get the currently selected text */
-	FText GetSelectedText() const;
+	SLATE_API FText GetSelectedText() const;
 
 	/** Move the cursor to the given location in the document */
-	void GoTo(const FTextLocation& NewLocation);
+	SLATE_API void GoTo(const FTextLocation& NewLocation);
 
 	/** Move the cursor to the specified location */
 	void GoTo(const ETextLocation NewLocation)
@@ -379,7 +379,7 @@ public:
 	}
 
 	/** Scroll to the given location in the document (without moving the cursor) */
-	void ScrollTo(const FTextLocation& NewLocation);
+	SLATE_API void ScrollTo(const FTextLocation& NewLocation);
 
 	/** Scroll to the given location in the document (without moving the cursor) */
 	void ScrollTo(const ETextLocation NewLocation)
@@ -388,32 +388,32 @@ public:
 	}
 
 	/** Begin a new text search (this is called automatically when the bound search text changes) */
-	void BeginSearch(const FText& InSearchText, const ESearchCase::Type InSearchCase = ESearchCase::IgnoreCase, const bool InReverse = false);
+	SLATE_API void BeginSearch(const FText& InSearchText, const ESearchCase::Type InSearchCase = ESearchCase::IgnoreCase, const bool InReverse = false);
 
 	/** Advance the current search to the next match (does nothing if not currently searching) */
-	void AdvanceSearch(const bool InReverse = false);
+	SLATE_API void AdvanceSearch(const bool InReverse = false);
 
 	/** Register and activate the IME context for the text layout of this textbox */
-	void EnableTextInputMethodContext();
+	SLATE_API void EnableTextInputMethodContext();
 
-	bool HasError() const;
+	SLATE_API bool HasError() const;
 
 	// SWidget overrides
-	virtual bool SupportsKeyboardFocus() const override;
-	virtual bool HasKeyboardFocus() const override;
-	virtual FReply OnFocusReceived( const FGeometry& MyGeometry, const FFocusEvent& InFocusEvent ) override;
-	virtual FReply OnKeyDown(const FGeometry& MyGeometry, const FKeyEvent& InKeyEvent) override;
+	SLATE_API virtual bool SupportsKeyboardFocus() const override;
+	SLATE_API virtual bool HasKeyboardFocus() const override;
+	SLATE_API virtual FReply OnFocusReceived( const FGeometry& MyGeometry, const FFocusEvent& InFocusEvent ) override;
+	SLATE_API virtual FReply OnKeyDown(const FGeometry& MyGeometry, const FKeyEvent& InKeyEvent) override;
 
 protected:
 #if WITH_ACCESSIBILITY
-	virtual TSharedRef<FSlateAccessibleWidget> CreateAccessibleWidget() override;
-	virtual TOptional<FText> GetDefaultAccessibleText(EAccessibleType AccessibleType = EAccessibleType::Main) const override;
+	SLATE_API virtual TSharedRef<FSlateAccessibleWidget> CreateAccessibleWidget() override;
+	SLATE_API virtual TOptional<FText> GetDefaultAccessibleText(EAccessibleType AccessibleType = EAccessibleType::Main) const override;
 #endif
 	/** Callback for the editable text's OnTextChanged event */
-	void OnEditableTextChanged(const FText& InText);
+	SLATE_API void OnEditableTextChanged(const FText& InText);
 
 	/** Callback when the editable text is committed. */
-	void OnEditableTextCommitted(const FText& InText, ETextCommit::Type InCommitType);
+	SLATE_API void OnEditableTextCommitted(const FText& InText, ETextCommit::Type InCommitType);
 
 	const FEditableTextBoxStyle* Style;
 
@@ -458,10 +458,10 @@ protected:
 
 private:
 
-	FMargin DeterminePadding() const;
-	FSlateFontInfo DetermineFont() const;
-	FSlateColor DetermineBackgroundColor() const;
-	FSlateColor DetermineForegroundColor() const;
+	SLATE_API FMargin DeterminePadding() const;
+	SLATE_API FSlateFontInfo DetermineFont() const;
+	SLATE_API FSlateColor DetermineBackgroundColor() const;
+	SLATE_API FSlateColor DetermineForegroundColor() const;
 
 	/** Styling: border image to draw when not hovered or focused */
 	const FSlateBrush* BorderImageNormal;
@@ -473,6 +473,6 @@ private:
 	const FSlateBrush* BorderImageReadOnly;
 
 	/** @return Border image for the text box based on the hovered and focused state */
-	const FSlateBrush* GetBorderImage() const;
+	SLATE_API const FSlateBrush* GetBorderImage() const;
 
 };
diff --git a/Engine/Source/Runtime/Slate/Public/Widgets/Input/SExpandableButton.h b/Engine/Source/Runtime/Slate/Public/Widgets/Input/SExpandableButton.h
--- a/Engine/Source/Runtime/Slate/Public/Widgets/Input/SExpandableButton.h
+++ b/Engine/Source/Runtime/Slate/Public/Widgets/Input/SExpandableButton.h
@@ -12,10 +12,10 @@
 /**
  * A button that can either be collapsed or expanded, containing different content in each state.
  */
-class SLATE_API SExpandableButton
+class SExpandableButton
 	: public SBorder
 {
-	SLATE_DECLARE_WIDGET(SExpandableButton, SBorder)
+	SLATE_DECLARE_WIDGET_API(SExpandableButton, SBorder, SLATE_API)
 
 public:
 
@@ -49,21 +49,21 @@ public:
 
 	SLATE_END_ARGS()
 
-	SExpandableButton();
+	SLATE_API SExpandableButton();
 	/**
 	 * Construct this widget
 	 *
 	 * @param	InArgs	The declaration data for this widget
 	 */
-	void Construct(const FArguments& InArgs);
+	SLATE_API void Construct(const FArguments& InArgs);
 
 protected:
 
 	/** Callbacks to determine visibility of parts that should be shown when the button state is collapsed or expanded */
-	EVisibility GetCollapsedVisibility() const;
-	EVisibility GetExpandedVisibility() const;
+	SLATE_API EVisibility GetCollapsedVisibility() const;
+	SLATE_API EVisibility GetExpandedVisibility() const;
 
-	void UpdateVisibility();
+	SLATE_API void UpdateVisibility();
 
 private:
 
diff --git a/Engine/Source/Runtime/Slate/Public/Widgets/Input/SInputKeySelector.h b/Engine/Source/Runtime/Slate/Public/Widgets/Input/SInputKeySelector.h
--- a/Engine/Source/Runtime/Slate/Public/Widgets/Input/SInputKeySelector.h
+++ b/Engine/Source/Runtime/Slate/Public/Widgets/Input/SInputKeySelector.h
@@ -23,7 +23,7 @@ class STextBlock;
 /**
  * A widget for selecting keys or input chords.
  */
-class SLATE_API SInputKeySelector : public SCompoundWidget
+class SInputKeySelector : public SCompoundWidget
 {
 public:
 	DECLARE_DELEGATE_OneParam( FOnKeySelected, const FInputChord& )
@@ -84,24 +84,24 @@ public:
 		SLATE_ARGUMENT(bool, IsFocusable)
 	SLATE_END_ARGS()
 
-	void Construct( const FArguments& InArgs );
+	SLATE_API void Construct( const FArguments& InArgs );
 
 public:
 
 	/** Gets the currently selected key chord. */
-	FInputChord GetSelectedKey() const;
+	SLATE_API FInputChord GetSelectedKey() const;
 
 	/** Sets the currently selected key chord. */
-	void SetSelectedKey( TAttribute<FInputChord> InSelectedKey );
+	SLATE_API void SetSelectedKey( TAttribute<FInputChord> InSelectedKey );
 
 	/** Sets the margin around the text used to display the currently selected key */
-	void SetMargin( TAttribute<FMargin> InMargin );
+	SLATE_API void SetMargin( TAttribute<FMargin> InMargin );
 
 	/** Sets the style of the button which is used enter key selection mode. */
-	void SetButtonStyle( const FButtonStyle* ButtonStyle );
+	SLATE_API void SetButtonStyle( const FButtonStyle* ButtonStyle );
 
 	/** Sets the style of the text on the button which is used enter key selection mode. */
-	void SetTextStyle(const FTextBlockStyle* InTextStyle);
+	SLATE_API void SetTextStyle(const FTextBlockStyle* InTextStyle);
 
 	/** Sets the text which is displayed when selecting a key. */
 	void SetKeySelectionText( FText InKeySelectionText ) { KeySelectionText = MoveTemp(InKeySelectionText); }
@@ -122,38 +122,38 @@ public:
 	bool GetIsSelectingKey() const { return bIsSelectingKey; }
 
 	/** Sets the visibility of the text block. */
-	void SetTextBlockVisibility(EVisibility InVisibility);
+	SLATE_API void SetTextBlockVisibility(EVisibility InVisibility);
 
 public:
 
-	virtual FReply OnPreviewKeyDown( const FGeometry& MyGeometry, const FKeyEvent& InKeyEvent ) override;
-	virtual FReply OnKeyDown(const FGeometry& MyGeometry, const FKeyEvent& InKeyEvent) override;
-	virtual FReply OnKeyUp( const FGeometry& MyGeometry, const FKeyEvent& InKeyEvent ) override;
-	virtual FReply OnPreviewMouseButtonDown( const FGeometry& MyGeometry, const FPointerEvent& MouseEvent ) override;
-	virtual FReply OnMouseButtonDown(const FGeometry& MyGeometry, const FPointerEvent& MouseEvent) override;
-	virtual void OnFocusLost( const FFocusEvent& InFocusEvent ) override;
-	virtual FNavigationReply OnNavigation(const FGeometry& MyGeometry, const FNavigationEvent& InNavigationEvent) override;
+	SLATE_API virtual FReply OnPreviewKeyDown( const FGeometry& MyGeometry, const FKeyEvent& InKeyEvent ) override;
+	SLATE_API virtual FReply OnKeyDown(const FGeometry& MyGeometry, const FKeyEvent& InKeyEvent) override;
+	SLATE_API virtual FReply OnKeyUp( const FGeometry& MyGeometry, const FKeyEvent& InKeyEvent ) override;
+	SLATE_API virtual FReply OnPreviewMouseButtonDown( const FGeometry& MyGeometry, const FPointerEvent& MouseEvent ) override;
+	SLATE_API virtual FReply OnMouseButtonDown(const FGeometry& MyGeometry, const FPointerEvent& MouseEvent) override;
+	SLATE_API virtual void OnFocusLost( const FFocusEvent& InFocusEvent ) override;
+	SLATE_API virtual FNavigationReply OnNavigation(const FGeometry& MyGeometry, const FNavigationEvent& InNavigationEvent) override;
 	virtual bool SupportsKeyboardFocus() const override { return true; }
 
 private:
 
 	/** Gets the display text for the currently selected key. */
-	FText GetSelectedKeyText() const;
+	SLATE_API FText GetSelectedKeyText() const;
 
 	/**  Gets the margin around the text used to display the currently selected key. */
-	FMargin GetMargin() const;
+	SLATE_API FMargin GetMargin() const;
 
 	/** Handles the OnClicked event from the button which enables key selection mode. */
-	FReply OnClicked();
+	SLATE_API FReply OnClicked();
 
 	/** Sets the currently selected key and invokes the associated events. */
-	void SelectKey( FKey Key, bool bShiftDown, bool bControllDown, bool bAltDown, bool bCommandDown );
+	SLATE_API void SelectKey( FKey Key, bool bShiftDown, bool bControllDown, bool bAltDown, bool bCommandDown );
 
 	/** Sets bIsSelectingKey and invokes the associated events. */
-	void SetIsSelectingKey(bool bInIsSelectingKey);
+	SLATE_API void SetIsSelectingKey(bool bInIsSelectingKey);
 
 	/** Returns true, if the key has been specified as an escape key, else false. */
-	bool IsEscapeKey(const FKey& InKey) const;
+	SLATE_API bool IsEscapeKey(const FKey& InKey) const;
 
 private:
 
diff --git a/Engine/Source/Runtime/Slate/Public/Widgets/Input/SMenuAnchor.h b/Engine/Source/Runtime/Slate/Public/Widgets/Input/SMenuAnchor.h
--- a/Engine/Source/Runtime/Slate/Public/Widgets/Input/SMenuAnchor.h
+++ b/Engine/Source/Runtime/Slate/Public/Widgets/Input/SMenuAnchor.h
@@ -28,7 +28,7 @@ DECLARE_DELEGATE_OneParam(FOnIsOpenChanged, bool)
  * A PopupAnchor summons a Popup relative to its content.
  * Summoning a popup relative to the cursor is accomplished via the application.
  */
-class SLATE_API SMenuAnchor : public SPanel, public IMenuHost
+class SMenuAnchor : public SPanel, public IMenuHost
 {
 public:
 	SLATE_BEGIN_ARGS( SMenuAnchor )
@@ -73,22 +73,22 @@ public:
 
 	SLATE_END_ARGS()
 
-	virtual ~SMenuAnchor();
+	SLATE_API virtual ~SMenuAnchor();
 	
 	/**
 	 * Construct this widget
 	 *
 	 * @param	InArgs	The declaration data for this widget
 	 */
-	void Construct( const FArguments& InArgs );
+	SLATE_API void Construct( const FArguments& InArgs );
 	
-	SMenuAnchor();
+	SLATE_API SMenuAnchor();
 
 	/** See Content Slot attribute */
-	void SetContent(TSharedRef<SWidget> InContent);
+	SLATE_API void SetContent(TSharedRef<SWidget> InContent);
 
 	/** See MenuContent attribute */
-	virtual void SetMenuContent(TSharedRef<SWidget> InMenuContent);
+	SLATE_API virtual void SetMenuContent(TSharedRef<SWidget> InMenuContent);
 	
 	/**
 	 * Open or close the popup
@@ -96,64 +96,64 @@ public:
 	 * @param InIsOpen    If true, open the popup. Otherwise close it.
 	 * @param bFocusMenu  Should we focus the popup as soon as it opens?
 	 */
-	virtual void SetIsOpen( bool InIsOpen, const bool bFocusMenu = true, const int32 FocusUserIndex = 0 );
+	SLATE_API virtual void SetIsOpen( bool InIsOpen, const bool bFocusMenu = true, const int32 FocusUserIndex = 0 );
 	
 	/** @return true if the popup is open; false otherwise. */
-	bool IsOpen() const;
+	SLATE_API bool IsOpen() const;
 
 	/** @return true if we should open the menu due to a click. Sometimes we should not, if
 	   the same MouseDownEvent that just closed the menu is about to re-open it because it happens to land on the button.*/
-	bool ShouldOpenDueToClick() const;
+	SLATE_API bool ShouldOpenDueToClick() const;
 
 	/** @return The current menu position */
-	FVector2D GetMenuPosition() const;
+	SLATE_API FVector2D GetMenuPosition() const;
 
-	void SetMenuPlacement(TAttribute<EMenuPlacement> InMenuPlacement);
+	SLATE_API void SetMenuPlacement(TAttribute<EMenuPlacement> InMenuPlacement);
 
-	void SetFitInWindow(bool bFit);
+	SLATE_API void SetFitInWindow(bool bFit);
 
 	/** @return Whether this menu has open submenus */
-	bool HasOpenSubMenus() const;
+	SLATE_API bool HasOpenSubMenus() const;
 
 	// IMenuHost interface
-	virtual TSharedPtr<SWindow> GetMenuWindow() const override;
-	virtual void OnMenuDismissed() override;
-	virtual bool UsingApplicationMenuStack() const override;
+	SLATE_API virtual TSharedPtr<SWindow> GetMenuWindow() const override;
+	SLATE_API virtual void OnMenuDismissed() override;
+	SLATE_API virtual bool UsingApplicationMenuStack() const override;
 	// End of IMenuHost interface
 
-	static void DismissAllApplicationMenus();
+	static SLATE_API void DismissAllApplicationMenus();
 
 protected:
 	// SWidget interface
-	virtual void Tick( const FGeometry& AllottedGeometry, const double InCurrentTime, const float InDeltaTime ) override;
-	virtual bool ComputeVolatility() const override;
-	virtual void OnArrangeChildren( const FGeometry& AllottedGeometry, FArrangedChildren& ArrangedChildren ) const override;
-	virtual FVector2D ComputeDesiredSize(float) const override;
-	virtual FChildren* GetChildren() override;
-	virtual int32 OnPaint( const FPaintArgs& Args, const FGeometry& AllottedGeometry, const FSlateRect& MyCullingRect, FSlateWindowElementList& OutDrawElements, int32 LayerId, const FWidgetStyle& InWidgetStyle, bool bParentEnabled ) const override;
+	SLATE_API virtual void Tick( const FGeometry& AllottedGeometry, const double InCurrentTime, const float InDeltaTime ) override;
+	SLATE_API virtual bool ComputeVolatility() const override;
+	SLATE_API virtual void OnArrangeChildren( const FGeometry& AllottedGeometry, FArrangedChildren& ArrangedChildren ) const override;
+	SLATE_API virtual FVector2D ComputeDesiredSize(float) const override;
+	SLATE_API virtual FChildren* GetChildren() override;
+	SLATE_API virtual int32 OnPaint( const FPaintArgs& Args, const FGeometry& AllottedGeometry, const FSlateRect& MyCullingRect, FSlateWindowElementList& OutDrawElements, int32 LayerId, const FWidgetStyle& InWidgetStyle, bool bParentEnabled ) const override;
 	// End of SWidget interface
 
 	/** @return true if the popup is currently open and reusing an existing window */
-	bool IsOpenAndReusingWindow() const;
+	SLATE_API bool IsOpenAndReusingWindow() const;
 
 	/** @return true if the popup is currently open and we created a dedicated window for it */
-	bool IsOpenViaCreatedWindow() const;
+	SLATE_API bool IsOpenViaCreatedWindow() const;
 
 	/** Handler/callback called by menus created by this anchor, when they are dismissed */
-	void OnMenuClosed(TSharedRef<IMenu> InMenu);
+	SLATE_API void OnMenuClosed(TSharedRef<IMenu> InMenu);
 
 	/** Reset the popup state to prepopup. */
-	void ResetPopupMenuContent();
+	SLATE_API void ResetPopupMenuContent();
 
 	/** Computes the placement geometry for menus displayed in a separately created window */
-	FGeometry ComputeNewWindowMenuPlacement(const FGeometry& AllottedGeometry, const FVector2D& PopupDesiredSize, EMenuPlacement PlacementMode) const;
+	SLATE_API FGeometry ComputeNewWindowMenuPlacement(const FGeometry& AllottedGeometry, const FVector2D& PopupDesiredSize, EMenuPlacement PlacementMode) const;
 
-	static TArray<TWeakPtr<IMenu>> OpenApplicationMenus;
+	static SLATE_API TArray<TWeakPtr<IMenu>> OpenApplicationMenus;
 
 protected:
-	struct SLATE_API FPopupPlacement
+	struct FPopupPlacement
 	{
-		FPopupPlacement(const FGeometry& PlacementGeometry, const FVector2D& PopupDesiredSize, EMenuPlacement PlacementMode);
+		SLATE_API FPopupPlacement(const FGeometry& PlacementGeometry, const FVector2D& PopupDesiredSize, EMenuPlacement PlacementMode);
 		
 		FVector2D LocalPopupSize;
 		FVector2D LocalPopupOffset;
diff --git a/Engine/Source/Runtime/Slate/Public/Widgets/Input/SMultiLineEditableTextBox.h b/Engine/Source/Runtime/Slate/Public/Widgets/Input/SMultiLineEditableTextBox.h
--- a/Engine/Source/Runtime/Slate/Public/Widgets/Input/SMultiLineEditableTextBox.h
+++ b/Engine/Source/Runtime/Slate/Public/Widgets/Input/SMultiLineEditableTextBox.h
@@ -29,7 +29,7 @@ enum class ETextShapingMethod : uint8;
 /**
  * Editable text box widget
  */
-class SLATE_API SMultiLineEditableTextBox : public SBorder
+class SMultiLineEditableTextBox : public SBorder
 {
 
 public:
@@ -248,7 +248,7 @@ public:
 	 *
 	 * @param	InArgs	The declaration data for this widget
 	 */
-	void Construct( const FArguments& InArgs );
+	SLATE_API void Construct( const FArguments& InArgs );
 	
 	/**
 	 * Returns the text string
@@ -275,20 +275,20 @@ public:
 	 *
 	 * @param  OutTextLine	Text string
 	 */	
-	void GetCurrentTextLine(FString& OutTextLine) const;
+	SLATE_API void GetCurrentTextLine(FString& OutTextLine) const;
 
 	/** See attribute Style */
-	void SetStyle(const FEditableTextBoxStyle* InStyle);
+	SLATE_API void SetStyle(const FEditableTextBoxStyle* InStyle);
 
 	/** See attribute TextStyle */
-	void SetTextStyle(const FTextBlockStyle* InTextStyle);	
+	SLATE_API void SetTextStyle(const FTextBlockStyle* InTextStyle);	
 
 	/**
 	 * Sets the text string currently being edited 
 	 *
 	 * @param  InNewText  The new text string
 	 */
-	void SetText( const TAttribute< FText >& InNewText );
+	SLATE_API void SetText( const TAttribute< FText >& InNewText );
 
 	/**
 	 * Returns the hint text string
@@ -305,113 +305,113 @@ public:
 	 *
 	 * @param  InHintText The hint text string
 	 */
-	void SetHintText( const TAttribute< FText >& InHintText );
+	SLATE_API void SetHintText( const TAttribute< FText >& InHintText );
 
 	/** Set the text that is currently being searched for (if any) */
-	void SetSearchText(const TAttribute<FText>& InSearchText);
+	SLATE_API void SetSearchText(const TAttribute<FText>& InSearchText);
 
 	/** Get the text that is currently being searched for (if any) */
-	FText GetSearchText() const;
+	SLATE_API FText GetSearchText() const;
 
 	/**
 	 * Sets the text color and opacity (overrides Style)
 	 *
 	 * @param  InForegroundColor 	The text color and opacity
 	 */
-	void SetTextBoxForegroundColor(const TAttribute<FSlateColor>& InForegroundColor);
+	SLATE_API void SetTextBoxForegroundColor(const TAttribute<FSlateColor>& InForegroundColor);
 
 	/**
 	 * Sets the color of the background/border around the editable text (overrides Style) 
 	 *
 	 * @param  InBackgroundColor 	The background/border color
 	 */
-	void SetTextBoxBackgroundColor(const TAttribute<FSlateColor>& InBackgroundColor);
+	SLATE_API void SetTextBoxBackgroundColor(const TAttribute<FSlateColor>& InBackgroundColor);
 
 	/**
 	 * Sets the text color and opacity when read-only (overrides Style) 
 	 *
 	 * @param  InReadOnlyForegroundColor 	The read-only text color and opacity
 	 */
-	void SetReadOnlyForegroundColor(const TAttribute<FSlateColor>& InReadOnlyForegroundColor);
+	SLATE_API void SetReadOnlyForegroundColor(const TAttribute<FSlateColor>& InReadOnlyForegroundColor);
 
 	/**
 	 * Sets whether to select word on the mouse double click
 	 *
 	 * @param  InSelectWordOnMouseDoubleClick		Select word on the mouse double click
 	 */
-	void SetSelectWordOnMouseDoubleClick(const TAttribute<bool>& InSelectWordOnMouseDoubleClick);
+	SLATE_API void SetSelectWordOnMouseDoubleClick(const TAttribute<bool>& InSelectWordOnMouseDoubleClick);
 
 	/** See TextShapingMethod attribute */
-	void SetTextShapingMethod(const TOptional<ETextShapingMethod>& InTextShapingMethod);
+	SLATE_API void SetTextShapingMethod(const TOptional<ETextShapingMethod>& InTextShapingMethod);
 
 	/** See TextFlowDirection attribute */
-	void SetTextFlowDirection(const TOptional<ETextFlowDirection>& InTextFlowDirection);
+	SLATE_API void SetTextFlowDirection(const TOptional<ETextFlowDirection>& InTextFlowDirection);
 
 	/** See WrapTextAt attribute */
-	void SetWrapTextAt(const TAttribute<float>& InWrapTextAt);
+	SLATE_API void SetWrapTextAt(const TAttribute<float>& InWrapTextAt);
 
 	/** See AutoWrapText attribute */
-	void SetAutoWrapText(const TAttribute<bool>& InAutoWrapText);
+	SLATE_API void SetAutoWrapText(const TAttribute<bool>& InAutoWrapText);
 
 	/** Set WrappingPolicy attribute */
-	void SetWrappingPolicy(const TAttribute<ETextWrappingPolicy>& InWrappingPolicy);
+	SLATE_API void SetWrappingPolicy(const TAttribute<ETextWrappingPolicy>& InWrappingPolicy);
 
 	/** See LineHeightPercentage attribute */
-	void SetLineHeightPercentage(const TAttribute<float>& InLineHeightPercentage);
+	SLATE_API void SetLineHeightPercentage(const TAttribute<float>& InLineHeightPercentage);
 
 	/** See Margin attribute */
-	void SetMargin(const TAttribute<FMargin>& InMargin);
+	SLATE_API void SetMargin(const TAttribute<FMargin>& InMargin);
 
 	/** See Justification attribute */
-	void SetJustification(const TAttribute<ETextJustify::Type>& InJustification);
+	SLATE_API void SetJustification(const TAttribute<ETextJustify::Type>& InJustification);
 
 	/** Sets the overflow policy for this text block */
-	void SetOverflowPolicy(TOptional<ETextOverflowPolicy> InOverflowPolicy);
+	SLATE_API void SetOverflowPolicy(TOptional<ETextOverflowPolicy> InOverflowPolicy);
 
 	/** See the AllowContextMenu attribute */
-	void SetAllowContextMenu(const TAttribute< bool >& InAllowContextMenu);
+	SLATE_API void SetAllowContextMenu(const TAttribute< bool >& InAllowContextMenu);
 
 	/** Set the VirtualKeyboardDismissAction attribute */
-	void SetVirtualKeyboardDismissAction(TAttribute< EVirtualKeyboardDismissAction > InVirtualKeyboardDismissAction);
+	SLATE_API void SetVirtualKeyboardDismissAction(TAttribute< EVirtualKeyboardDismissAction > InVirtualKeyboardDismissAction);
 	
 	/** Set the ReadOnly attribute */
-	void SetIsReadOnly(const TAttribute< bool >& InIsReadOnly);
+	SLATE_API void SetIsReadOnly(const TAttribute< bool >& InIsReadOnly);
 
 	/**
 	 * If InError is a non-empty string the TextBox will the ErrorReporting provided during construction
 	 * If no error reporting was provided, the TextBox will create a default error reporter.
 	 */
-	void SetError( const FText& InError );
-	void SetError( const FString& InError );
+	SLATE_API void SetError( const FText& InError );
+	SLATE_API void SetError( const FString& InError );
 
 	// SWidget overrides
-	virtual bool SupportsKeyboardFocus() const override;
-	virtual bool HasKeyboardFocus() const override;
-	virtual FReply OnFocusReceived( const FGeometry& MyGeometry, const FFocusEvent& InFocusEvent ) override;
-	virtual FReply OnKeyDown(const FGeometry& MyGeometry, const FKeyEvent& InKeyEvent) override;
+	SLATE_API virtual bool SupportsKeyboardFocus() const override;
+	SLATE_API virtual bool HasKeyboardFocus() const override;
+	SLATE_API virtual FReply OnFocusReceived( const FGeometry& MyGeometry, const FFocusEvent& InFocusEvent ) override;
+	SLATE_API virtual FReply OnKeyDown(const FGeometry& MyGeometry, const FKeyEvent& InKeyEvent) override;
 
 
 	/** Query to see if any text is selected within the document */
-	bool AnyTextSelected() const;
+	SLATE_API bool AnyTextSelected() const;
 
 	/** Select all the text in the document */
-	void SelectAllText();
+	SLATE_API void SelectAllText();
 
 	/** Clear the active text selection */
-	void ClearSelection();
+	SLATE_API void ClearSelection();
 
 	/** Get the currently selected text */
-	FText GetSelectedText() const;
+	SLATE_API FText GetSelectedText() const;
 
 	/** Insert the given text at the current cursor position, correctly taking into account new line characters */
-	void InsertTextAtCursor(const FText& InText);
-	void InsertTextAtCursor(const FString& InString);
+	SLATE_API void InsertTextAtCursor(const FText& InText);
+	SLATE_API void InsertTextAtCursor(const FString& InString);
 
 	/** Insert the given run at the current cursor position */
-	void InsertRunAtCursor(TSharedRef<IRun> InRun);
+	SLATE_API void InsertRunAtCursor(TSharedRef<IRun> InRun);
 
 	/** Move the cursor to the given location in the document */
-	void GoTo(const FTextLocation& NewLocation);
+	SLATE_API void GoTo(const FTextLocation& NewLocation);
 
 	/** Move the cursor to the specified location */
 	void GoTo(const ETextLocation NewLocation)
@@ -420,7 +420,7 @@ public:
 	}
 
 	/** Scroll to the given location in the document (without moving the cursor) */
-	void ScrollTo(const FTextLocation& NewLocation);
+	SLATE_API void ScrollTo(const FTextLocation& NewLocation);
 
 	/** Scroll to the given location in the document (without moving the cursor) */
 	void ScrollTo(const ETextLocation NewLocation)
@@ -429,58 +429,58 @@ public:
 	}
 
 	/** Apply the given style to the currently selected text (or insert a new run at the current cursor position if no text is selected) */
-	void ApplyToSelection(const FRunInfo& InRunInfo, const FTextBlockStyle& InStyle);
+	SLATE_API void ApplyToSelection(const FRunInfo& InRunInfo, const FTextBlockStyle& InStyle);
 
 	/** Begin a new text search (this is called automatically when the bound search text changes) */
-	void BeginSearch(const FText& InSearchText, const ESearchCase::Type InSearchCase = ESearchCase::IgnoreCase, const bool InReverse = false);
+	SLATE_API void BeginSearch(const FText& InSearchText, const ESearchCase::Type InSearchCase = ESearchCase::IgnoreCase, const bool InReverse = false);
 
 	/** Advance the current search to the next match (does nothing if not currently searching) */
-	void AdvanceSearch(const bool InReverse = false);
+	SLATE_API void AdvanceSearch(const bool InReverse = false);
 
 	/** Get the run currently under the cursor, or null if there is no run currently under the cursor */
-	TSharedPtr<const IRun> GetRunUnderCursor() const;
+	SLATE_API TSharedPtr<const IRun> GetRunUnderCursor() const;
 
 	/** Get the runs currently that are current selected, some of which may be only partially selected */
-	TArray<TSharedRef<const IRun>> GetSelectedRuns() const;
+	SLATE_API TArray<TSharedRef<const IRun>> GetSelectedRuns() const;
 
 	/** Get the interaction position of the cursor (where to insert, delete, etc, text from/to) */
-	FTextLocation GetCursorLocation() const;
+	SLATE_API FTextLocation GetCursorLocation() const;
 
 	/** Get the horizontal scroll bar widget */
-	TSharedPtr<const SScrollBar> GetHScrollBar() const;
+	SLATE_API TSharedPtr<const SScrollBar> GetHScrollBar() const;
 
 	/** Get the vertical scroll bar widget */
-	TSharedPtr<const SScrollBar> GetVScrollBar() const;
+	SLATE_API TSharedPtr<const SScrollBar> GetVScrollBar() const;
 
 	/** Refresh this text box immediately, rather than wait for the usual caching mechanisms to take affect on the text Tick */
-	void Refresh();
+	SLATE_API void Refresh();
 
 	/**
 	 * Sets the OnKeyCharHandler to provide first chance handling of the SMultiLineEditableText's OnKeyChar event
 	 *
 	 * @param InOnKeyCharHandler			Delegate to call during OnKeyChar event
 	 */
-	void SetOnKeyCharHandler(FOnKeyChar InOnKeyCharHandler);
+	SLATE_API void SetOnKeyCharHandler(FOnKeyChar InOnKeyCharHandler);
 
 	/**
 	 * Sets the OnKeyDownHandler to provide first chance handling of the SMultiLineEditableText's OnKeyDown event
 	 *
 	 * @param InOnKeyDownHandler			Delegate to call during OnKeyDown event
 	 */
-	void SetOnKeyDownHandler(FOnKeyDown InOnKeyDownHandler);
+	SLATE_API void SetOnKeyDownHandler(FOnKeyDown InOnKeyDownHandler);
 
 
 	/**
 	 * 
 	 */
-	void ForceScroll(int32 UserIndex, float ScrollAxisMagnitude);
+	SLATE_API void ForceScroll(int32 UserIndex, float ScrollAxisMagnitude);
 
 protected:
 	/** Callback for the editable text's OnTextChanged event */
-	void OnEditableTextChanged(const FText& InText);
+	SLATE_API void OnEditableTextChanged(const FText& InText);
 
 	/** Callback when the editable text is committed. */
-	void OnEditableTextCommitted(const FText& InText, ETextCommit::Type InCommitType);
+	SLATE_API void OnEditableTextCommitted(const FText& InText, ETextCommit::Type InCommitType);
 
 protected:
 
@@ -554,12 +554,12 @@ protected:
 
 private:
 
-	FMargin DeterminePadding() const;
-	FMargin DetermineHScrollBarPadding() const;
-	FMargin DetermineVScrollBarPadding() const;
-	FSlateFontInfo DetermineFont() const;
-	FSlateColor DetermineBackgroundColor() const;
-	FSlateColor DetermineForegroundColor() const;
+	SLATE_API FMargin DeterminePadding() const;
+	SLATE_API FMargin DetermineHScrollBarPadding() const;
+	SLATE_API FMargin DetermineVScrollBarPadding() const;
+	SLATE_API FSlateFontInfo DetermineFont() const;
+	SLATE_API FSlateColor DetermineBackgroundColor() const;
+	SLATE_API FSlateColor DetermineForegroundColor() const;
 
 	/** Styling: border image to draw when not hovered or focused */
 	const FSlateBrush* BorderImageNormal;
@@ -571,7 +571,7 @@ private:
 	const FSlateBrush* BorderImageReadOnly;
 
 	/** @return Border image for the text box based on the hovered and focused state */
-	const FSlateBrush* GetBorderImage() const;
+	SLATE_API const FSlateBrush* GetBorderImage() const;
 
 };
 
diff --git a/Engine/Source/Runtime/Slate/Public/Widgets/Input/SSearchBox.h b/Engine/Source/Runtime/Slate/Public/Widgets/Input/SSearchBox.h
--- a/Engine/Source/Runtime/Slate/Public/Widgets/Input/SSearchBox.h
+++ b/Engine/Source/Runtime/Slate/Public/Widgets/Input/SSearchBox.h
@@ -17,7 +17,7 @@
 class FActiveTimerHandle;
 
 /** A text box that is used for searching. Meant to be as easy to use as possible with as few options as possible. */
-class SLATE_API SSearchBox : public SEditableTextBox
+class SSearchBox : public SEditableTextBox
 {
 
 public:
@@ -100,7 +100,7 @@ public:
 		SLATE_EVENT(FOnClicked, OnSearchGlassClicked)
 	SLATE_END_ARGS()
 
-	void Construct( const FArguments& InArgs );
+	SLATE_API void Construct( const FArguments& InArgs );
 
 private:
 
diff --git a/Engine/Source/Runtime/Slate/Public/Widgets/Input/SSlider.h b/Engine/Source/Runtime/Slate/Public/Widgets/Input/SSlider.h
--- a/Engine/Source/Runtime/Slate/Public/Widgets/Input/SSlider.h
+++ b/Engine/Source/Runtime/Slate/Public/Widgets/Input/SSlider.h
@@ -19,9 +19,9 @@ class FSlateWindowElementList;
 /**
  * A Slate slider control is a linear scale and draggable handle.
  */
-class SLATE_API SSlider : public SLeafWidget
+class SSlider : public SLeafWidget
 {
-	SLATE_DECLARE_WIDGET(SSlider, SLeafWidget)
+	SLATE_DECLARE_WIDGET_API(SSlider, SLeafWidget, SLATE_API)
 
 public:
 
@@ -100,17 +100,17 @@ public:
 
 	SLATE_END_ARGS()
 
-	SSlider();
+	SLATE_API SSlider();
 
 	/**
 	 * Construct the widget.
 	 * 
 	 * @param InDeclaration A declaration from which to construct the widget.
 	 */
-	void Construct( const SSlider::FArguments& InDeclaration );
+	SLATE_API void Construct( const SSlider::FArguments& InDeclaration );
 
 	/** Set the widget style. */
-	void SetStyle(const FSliderStyle* InStyle);
+	SLATE_API void SetStyle(const FSliderStyle* InStyle);
 
 	/** Get the MinValue attribute */
 	float GetMinValue() const { return MinValue; }
@@ -119,70 +119,70 @@ public:
 	float GetMaxValue() const { return MaxValue; }
 
 	/** Get the Value attribute */
-	float GetValue() const;
+	SLATE_API float GetValue() const;
 
 	/** Get the Value attribute scaled from 0 to 1 */
-	float GetNormalizedValue() const;
+	SLATE_API float GetNormalizedValue() const;
 
 	/** Set the Value attribute */
-	void SetValue(TAttribute<float> InValueAttribute);
+	SLATE_API void SetValue(TAttribute<float> InValueAttribute);
 
 	/** Set the MinValue and MaxValue attributes. If the new MinValue is more than the new MaxValue, MaxValue will be changed to equal MinValue. */
-	void SetMinAndMaxValues(float InMinValue, float InMaxValue);
+	SLATE_API void SetMinAndMaxValues(float InMinValue, float InMaxValue);
 	
 	/** Set the IndentHandle attribute */
-	void SetIndentHandle(TAttribute<bool> InIndentHandle);
+	SLATE_API void SetIndentHandle(TAttribute<bool> InIndentHandle);
 	
 	/** Set the Locked attribute */
-	void SetLocked(TAttribute<bool> InLocked);
+	SLATE_API void SetLocked(TAttribute<bool> InLocked);
 
 	/** Set the Orientation attribute */
-	void SetOrientation(EOrientation InOrientation);
+	SLATE_API void SetOrientation(EOrientation InOrientation);
 	
 	/** Set the SliderBarColor attribute */
-	void SetSliderBarColor(TAttribute<FSlateColor> InSliderBarColor);
+	SLATE_API void SetSliderBarColor(TAttribute<FSlateColor> InSliderBarColor);
 	
 	/** Set the SliderHandleColor attribute */
-	void SetSliderHandleColor(TAttribute<FSlateColor> InSliderHandleColor);
+	SLATE_API void SetSliderHandleColor(TAttribute<FSlateColor> InSliderHandleColor);
 
 	/** Get the StepSize attribute */
-	float GetStepSize() const;
+	SLATE_API float GetStepSize() const;
 
 	/** Set the StepSize attribute */
-	void SetStepSize(TAttribute<float> InStepSize);
+	SLATE_API void SetStepSize(TAttribute<float> InStepSize);
 
 	/** Set the MouseUsesStep attribute */
-	void SetMouseUsesStep(bool MouseUsesStep);
+	SLATE_API void SetMouseUsesStep(bool MouseUsesStep);
 
 	/** Set the RequiresControllerLock attribute */
-	void SetRequiresControllerLock(bool RequiresControllerLock);
+	SLATE_API void SetRequiresControllerLock(bool RequiresControllerLock);
 
 public:
 
 	// SWidget overrides
 
-	virtual int32 OnPaint( const FPaintArgs& Args, const FGeometry& AllottedGeometry, const FSlateRect& MyCullingRect, FSlateWindowElementList& OutDrawElements, int32 LayerId, const FWidgetStyle& InWidgetStyle, bool bParentEnabled ) const override;
-	virtual FVector2D ComputeDesiredSize(float) const override;
-	virtual FReply OnMouseButtonDown( const FGeometry& MyGeometry, const FPointerEvent& MouseEvent ) override;
-	virtual FReply OnMouseButtonUp( const FGeometry& MyGeometry, const FPointerEvent& MouseEvent ) override;
-	virtual void OnMouseCaptureLost(const FCaptureLostEvent& CaptureLostEvent) override;
-	virtual FReply OnMouseMove( const FGeometry& MyGeometry, const FPointerEvent& MouseEvent ) override;
-	virtual FReply OnTouchStarted(const FGeometry& MyGeometry, const FPointerEvent& InTouchEvent) override;
-	virtual FReply OnTouchMoved(const FGeometry& MyGeometry, const FPointerEvent& InTouchEvent) override;
-	virtual FReply OnTouchEnded(const FGeometry& MyGeometry, const FPointerEvent& InTouchEvent) override;
-	virtual FReply OnKeyDown(const FGeometry& MyGeometry, const FKeyEvent& InKeyEvent) override;
-	virtual FReply OnKeyUp(const FGeometry& MyGeometry, const FKeyEvent& InKeyEvent) override;
-	virtual void OnFocusLost(const FFocusEvent& InFocusEvent) override;
-	virtual FNavigationReply OnNavigation(const FGeometry& MyGeometry, const FNavigationEvent& InNavigationEvent) override;
-
-	virtual bool SupportsKeyboardFocus() const override;
-	virtual bool IsInteractable() const override;
+	SLATE_API virtual int32 OnPaint( const FPaintArgs& Args, const FGeometry& AllottedGeometry, const FSlateRect& MyCullingRect, FSlateWindowElementList& OutDrawElements, int32 LayerId, const FWidgetStyle& InWidgetStyle, bool bParentEnabled ) const override;
+	SLATE_API virtual FVector2D ComputeDesiredSize(float) const override;
+	SLATE_API virtual FReply OnMouseButtonDown( const FGeometry& MyGeometry, const FPointerEvent& MouseEvent ) override;
+	SLATE_API virtual FReply OnMouseButtonUp( const FGeometry& MyGeometry, const FPointerEvent& MouseEvent ) override;
+	SLATE_API virtual void OnMouseCaptureLost(const FCaptureLostEvent& CaptureLostEvent) override;
+	SLATE_API virtual FReply OnMouseMove( const FGeometry& MyGeometry, const FPointerEvent& MouseEvent ) override;
+	SLATE_API virtual FReply OnTouchStarted(const FGeometry& MyGeometry, const FPointerEvent& InTouchEvent) override;
+	SLATE_API virtual FReply OnTouchMoved(const FGeometry& MyGeometry, const FPointerEvent& InTouchEvent) override;
+	SLATE_API virtual FReply OnTouchEnded(const FGeometry& MyGeometry, const FPointerEvent& InTouchEvent) override;
+	SLATE_API virtual FReply OnKeyDown(const FGeometry& MyGeometry, const FKeyEvent& InKeyEvent) override;
+	SLATE_API virtual FReply OnKeyUp(const FGeometry& MyGeometry, const FKeyEvent& InKeyEvent) override;
+	SLATE_API virtual void OnFocusLost(const FFocusEvent& InFocusEvent) override;
+	SLATE_API virtual FNavigationReply OnNavigation(const FGeometry& MyGeometry, const FNavigationEvent& InNavigationEvent) override;
+
+	SLATE_API virtual bool SupportsKeyboardFocus() const override;
+	SLATE_API virtual bool IsInteractable() const override;
 #if WITH_ACCESSIBILITY
-	virtual TSharedRef<FSlateAccessibleWidget> CreateAccessibleWidget() override;
+	SLATE_API virtual TSharedRef<FSlateAccessibleWidget> CreateAccessibleWidget() override;
 #endif
 
 	/** @return Is the handle locked or not? Defaults to false */
-	bool IsLocked() const;
+	SLATE_API bool IsLocked() const;
 
 protected:
 
@@ -191,7 +191,7 @@ protected:
 	 *
 	 * @param NewValue The value to commit.
 	 */
-	virtual void CommitValue(float NewValue);
+	SLATE_API virtual void CommitValue(float NewValue);
 
 	/**
 	 * Calculates the new value based on the given absolute coordinates.
@@ -200,10 +200,10 @@ protected:
 	 * @param AbsolutePosition The absolute position of the slider.
 	 * @return The new value.
 	 */
-	float PositionToValue( const FGeometry& MyGeometry, const UE::Slate::FDeprecateVector2DParameter& AbsolutePosition );
+	SLATE_API float PositionToValue( const FGeometry& MyGeometry, const UE::Slate::FDeprecateVector2DParameter& AbsolutePosition );
 
-	const FSlateBrush* GetBarImage() const;
-	const FSlateBrush* GetThumbImage() const;
+	SLATE_API const FSlateBrush* GetBarImage() const;
+	SLATE_API const FSlateBrush* GetThumbImage() const;
 
 protected:
 #if WITH_EDITORONLY_DATA
@@ -283,7 +283,7 @@ protected:
 private:
 
 	// Resets controller input state. Fires delegates.
-	void ResetControllerState();
+	SLATE_API void ResetControllerState();
 
 	// Holds the slider's current value.
 	TSlateAttribute<float> ValueSlateAttribute;
diff --git a/Engine/Source/Runtime/Slate/Public/Widgets/Input/SSubMenuHandler.h b/Engine/Source/Runtime/Slate/Public/Widgets/Input/SSubMenuHandler.h
--- a/Engine/Source/Runtime/Slate/Public/Widgets/Input/SSubMenuHandler.h
+++ b/Engine/Source/Runtime/Slate/Public/Widgets/Input/SSubMenuHandler.h
@@ -12,7 +12,7 @@ class FActiveTimerHandle;
 /**
  * Wrapper for any widget that is used in a table view that wants to handle sub-menus with the same functionality as a normal multibox menu
  */
-class SLATE_API SSubMenuHandler : public SCompoundWidget
+class SSubMenuHandler : public SCompoundWidget
 {
 public:
 	SLATE_BEGIN_ARGS( SSubMenuHandler )
@@ -25,19 +25,19 @@ public:
 		SLATE_EVENT( FOnGetContent, OnGetMenuContent )
 	SLATE_END_ARGS()
 
-	void Construct(const FArguments& InArgs, TWeakPtr<SMenuOwner> InMenuOwner);
+	SLATE_API void Construct(const FArguments& InArgs, TWeakPtr<SMenuOwner> InMenuOwner);
 
 public:
 	// SWidget interface
-	virtual void OnMouseEnter( const FGeometry& MyGeometry, const FPointerEvent& MouseEvent ) override;
-	virtual void OnMouseLeave( const FPointerEvent& MouseEvent ) override;
+	SLATE_API virtual void OnMouseEnter( const FGeometry& MyGeometry, const FPointerEvent& MouseEvent ) override;
+	SLATE_API virtual void OnMouseLeave( const FPointerEvent& MouseEvent ) override;
 	// End of SWidget interface
 	
 	/**
 	 * Returns whether or the sub-menu entry should appear hovered.  If the sub-menu is open we will always show the menu as hovered to indicate which sub-menu is open
 	 * In the case that the user is interacting in this menu we do not show the menu as hovered because we need to show what the user is actually selecting
 	 */
-	bool ShouldSubMenuAppearHovered() const;
+	SLATE_API bool ShouldSubMenuAppearHovered() const;
 
 	/**
 	 * Requests that the sub-menu associated with this widget be toggled on or off.
@@ -47,12 +47,12 @@ public:
 	 * @param bClobber					true if we want to open a menu when another menu is already open
 	 * @param const bool bImmediate		Default false, when true the menu will open immediately
 	 */
-	void RequestSubMenuToggle( bool bOpenMenu, const bool bClobber, const bool bImmediate = false);
+	SLATE_API void RequestSubMenuToggle( bool bOpenMenu, const bool bClobber, const bool bImmediate = false);
 	
 	/**
 	 * Cancels any open requests to toggle a sub-menu       
 	 */
-	void CancelPendingSubMenu();
+	SLATE_API void CancelPendingSubMenu();
 
 	/** Returns TRUE if there is a sub-menu available for this item */
 	bool HasSubMenu() const
@@ -61,11 +61,11 @@ public:
 	}
 
 	/** Returns TRUE if the sub-menu is open */
-	bool IsSubMenuOpen() const;
+	SLATE_API bool IsSubMenuOpen() const;
 
 protected:
 	/** One-off delayed active timer to update the open/closed state of the sub menu. */
-	EActiveTimerReturnType UpdateSubMenuState(double InCurrentTime, float InDeltaTime, bool bWantsOpen);
+	SLATE_API EActiveTimerReturnType UpdateSubMenuState(double InCurrentTime, float InDeltaTime, bool bWantsOpen);
 
 	/** The handle to the active timer to update the sub-menu state */
 	TWeakPtr<FActiveTimerHandle> ActiveTimerHandle;
diff --git a/Engine/Source/Runtime/Slate/Public/Widgets/Input/SSuggestionTextBox.h b/Engine/Source/Runtime/Slate/Public/Widgets/Input/SSuggestionTextBox.h
--- a/Engine/Source/Runtime/Slate/Public/Widgets/Input/SSuggestionTextBox.h
+++ b/Engine/Source/Runtime/Slate/Public/Widgets/Input/SSuggestionTextBox.h
@@ -39,7 +39,7 @@ DECLARE_DELEGATE_TwoParams(FOnShowingSuggestions, const FString&, OUT TArray<FSt
 /**
  * Implements an editable text box that can show auto-complete histories and suggestions lists.
  */
-class SLATE_API SSuggestionTextBox
+class SSuggestionTextBox
 	: public SCompoundWidget
 {
 
@@ -129,7 +129,7 @@ public:
 public:
 
 	/** Default constructor. */
-	SSuggestionTextBox( );
+	SLATE_API SSuggestionTextBox( );
 
 public:
 
@@ -138,7 +138,7 @@ public:
 	 *
 	 * @param InArgs Declaration used by the SNew() macro to construct this widget.
 	 */
-	void Construct( const FArguments& InArgs );
+	SLATE_API void Construct( const FArguments& InArgs );
 
 	/**
 	 * Returns the text string.
@@ -168,25 +168,25 @@ public:
 	 *
 	 * @param InNewText The new text string.
 	 */
-	void SetText( const TAttribute< FText >& InNewText );
+	SLATE_API void SetText( const TAttribute< FText >& InNewText );
 
 protected:
 	
 	/** Clears the list of suggestions and hides the suggestions list. */
-	void ClearSuggestions( );
+	SLATE_API void ClearSuggestions( );
 
 	/** Sets the keyboard focus to the text box. */
-	void FocusTextBox( );
+	SLATE_API void FocusTextBox( );
 
 	/**
 	 * Gets the string value of the currently selected suggestion.
 	 *
 	 * @return Suggestion string, or empty if no suggestion is selected.
 	 */
-	FString GetSelectedSuggestionString( ) const;
+	SLATE_API FString GetSelectedSuggestionString( ) const;
 
 	/** Highlights the selected suggestion in the suggestion list. */
-	void MarkActiveSuggestion( );
+	SLATE_API void MarkActiveSuggestion( );
 
 	/**
 	 * Sets the list of suggestions.
@@ -194,35 +194,35 @@ protected:
 	 * @param SuggestionStrings The suggestion strings.
 	 * @param InHistoryMode Whether the suggestions represent the input history.
 	 */
-	void SetSuggestions( TArray<FString>& SuggestionStrings, bool InHistoryMode );
+	SLATE_API void SetSuggestions( TArray<FString>& SuggestionStrings, bool InHistoryMode );
 
 protected:
 
 	// SWidget overrides
 
-	virtual void OnFocusLost( const FFocusEvent& InFocusEvent ) override;
-	virtual FReply OnKeyDown( const FGeometry& MyGeometry, const FKeyEvent& KeyEvent ) override;
-	virtual bool SupportsKeyboardFocus( ) const override;
+	SLATE_API virtual void OnFocusLost( const FFocusEvent& InFocusEvent ) override;
+	SLATE_API virtual FReply OnKeyDown( const FGeometry& MyGeometry, const FKeyEvent& KeyEvent ) override;
+	SLATE_API virtual bool SupportsKeyboardFocus( ) const override;
 
 private:
 
 	// Callback for generating a row widget for the suggestion list view.
-	TSharedRef<ITableRow> HandleSuggestionListViewGenerateRow( TSharedPtr<FString> Message, const TSharedRef<STableViewBase>& OwnerTable );
+	SLATE_API TSharedRef<ITableRow> HandleSuggestionListViewGenerateRow( TSharedPtr<FString> Message, const TSharedRef<STableViewBase>& OwnerTable );
 
 	// Callback for changing the selected suggestion.
-	void HandleSuggestionListViewSelectionChanged( TSharedPtr<FString> NewValue, ESelectInfo::Type SelectInfo );
+	SLATE_API void HandleSuggestionListViewSelectionChanged( TSharedPtr<FString> NewValue, ESelectInfo::Type SelectInfo );
 
 	// Callback for getting the highlight string in a suggestion list row widget.
-	FText HandleSuggestionListWidgetHighlightText( ) const;
+	SLATE_API FText HandleSuggestionListWidgetHighlightText( ) const;
 
 	// Callback for changing the text box input.
-	void HandleTextBoxTextChanged( const FText& InText );
+	SLATE_API void HandleTextBoxTextChanged( const FText& InText );
 
 	// Callback committing the text box input.
-	void HandleTextBoxTextCommitted( const FText& InText, ETextCommit::Type CommitInfo );
+	SLATE_API void HandleTextBoxTextCommitted( const FText& InText, ETextCommit::Type CommitInfo );
 
 	// Called when global focus changes so we can detect when to close the SuggestionList.
-	void OnGlobalFocusChanging(const FFocusEvent& FocusEvent, const FWeakWidgetPath& OldFocusedWidgetPath, const TSharedPtr<SWidget>& OldFocusedWidget, const FWidgetPath& NewFocusedWidgetPath, const TSharedPtr<SWidget>& NewFocusedWidget);
+	SLATE_API void OnGlobalFocusChanging(const FFocusEvent& FocusEvent, const FWeakWidgetPath& OldFocusedWidgetPath, const TSharedPtr<SWidget>& OldFocusedWidget, const FWidgetPath& NewFocusedWidgetPath, const TSharedPtr<SWidget>& NewFocusedWidget);
 
 private:
 
diff --git a/Engine/Source/Runtime/Slate/Public/Widgets/Input/STextComboBox.h b/Engine/Source/Runtime/Slate/Public/Widgets/Input/STextComboBox.h
--- a/Engine/Source/Runtime/Slate/Public/Widgets/Input/STextComboBox.h
+++ b/Engine/Source/Runtime/Slate/Public/Widgets/Input/STextComboBox.h
@@ -18,7 +18,7 @@
 /**
  * A combo box that shows text content.
  */
-class SLATE_API STextComboBox : public SCompoundWidget
+class STextComboBox : public SCompoundWidget
 {
 public:
 
@@ -63,12 +63,12 @@ public:
 		SLATE_EVENT( FGetTextComboLabel, OnGetTextLabelForItem ) 
 	SLATE_END_ARGS()
 
-	void Construct( const FArguments& InArgs );
+	SLATE_API void Construct( const FArguments& InArgs );
 
 	/** Called to create a widget for each string */
-	TSharedRef<SWidget> MakeItemWidget( TSharedPtr<FString> StringItem );
+	SLATE_API TSharedRef<SWidget> MakeItemWidget( TSharedPtr<FString> StringItem );
 
-	void SetSelectedItem (TSharedPtr<FString> NewSelection);
+	SLATE_API void SetSelectedItem (TSharedPtr<FString> NewSelection);
 
 	/** Returns the currently selected text string */
 	TSharedPtr<FString> GetSelectedItem()
@@ -77,21 +77,21 @@ public:
 	}
 
 	/** Request to reload the text options in the combobox from the OptionsSource attribute */
-	void RefreshOptions();
+	SLATE_API void RefreshOptions();
 
 	/** Clears the selected item in the text combo */
-	void ClearSelection();
+	SLATE_API void ClearSelection();
 
 private:
 	TSharedPtr<FString> OnGetSelection() const {return SelectedItem;}
 
 	/** Called when selection changes in the combo pop-up */
-	void OnSelectionChanged (TSharedPtr<FString> Selection, ESelectInfo::Type SelectInfo);
+	SLATE_API void OnSelectionChanged (TSharedPtr<FString> Selection, ESelectInfo::Type SelectInfo);
 
 	/** Helper method to get the text for a given item in the combo box */
-	FText GetSelectedTextLabel() const;
+	SLATE_API FText GetSelectedTextLabel() const;
 
-	FText GetItemTextLabel(TSharedPtr<FString> StringItem) const;
+	SLATE_API FText GetItemTextLabel(TSharedPtr<FString> StringItem) const;
 
 private:
 
diff --git a/Engine/Source/Runtime/Slate/Public/Widgets/Input/STextComboPopup.h b/Engine/Source/Runtime/Slate/Public/Widgets/Input/STextComboPopup.h
--- a/Engine/Source/Runtime/Slate/Public/Widgets/Input/STextComboPopup.h
+++ b/Engine/Source/Runtime/Slate/Public/Widgets/Input/STextComboPopup.h
@@ -14,7 +14,7 @@ class SButton;
 DECLARE_DELEGATE_OneParam( FOnTextChosen, const FString& );
 
 /** Simple box containing a combobox of strings to pick from, usually used within a MenuStack */
-class SLATE_API STextComboPopup : public SCompoundWidget
+class STextComboPopup : public SCompoundWidget
 {
 public:
 	/** Array of shared pointers to strings so combo widget can work on them */
@@ -44,18 +44,18 @@ public:
 
 	SLATE_END_ARGS()
 
-	void Construct( const FArguments& InArgs );
+	SLATE_API void Construct( const FArguments& InArgs );
 
 	/** SWidget interface */
-	virtual FReply OnKeyDown( const FGeometry& MyGeometry, const FKeyEvent& InKeyEvent ) override;
+	SLATE_API virtual FReply OnKeyDown( const FGeometry& MyGeometry, const FKeyEvent& InKeyEvent ) override;
 
 	/** Called to create a widget for each string */
-	TSharedRef<SWidget> MakeItemWidget( TSharedPtr<FString> StringItem );
+	SLATE_API TSharedRef<SWidget> MakeItemWidget( TSharedPtr<FString> StringItem );
 	/** Called when 'OK' button is pressed */
-	FReply OnOK();
+	SLATE_API FReply OnOK();
 
 	/** focuses the control on the OK button */
-	void FocusDefaultWidget();
+	SLATE_API void FocusDefaultWidget();
 
 private:
 	TSharedPtr<FString> OnGetSelection() const {return SelectedItem;}
diff --git a/Engine/Source/Runtime/Slate/Public/Widgets/Input/STextEntryPopup.h b/Engine/Source/Runtime/Slate/Public/Widgets/Input/STextEntryPopup.h
--- a/Engine/Source/Runtime/Slate/Public/Widgets/Input/STextEntryPopup.h
+++ b/Engine/Source/Runtime/Slate/Public/Widgets/Input/STextEntryPopup.h
@@ -10,7 +10,7 @@
 class SVerticalBox;
 
 /** Simple text entry popup, usually used within a MenuStack */
-class SLATE_API STextEntryPopup : public SCompoundWidget
+class STextEntryPopup : public SCompoundWidget
 {
 public:
 
@@ -55,18 +55,18 @@ public:
 
 	SLATE_END_ARGS()
 
-	void Construct( const FArguments& InArgs );
+	SLATE_API void Construct( const FArguments& InArgs );
 
 	/** @return the widget that we want to be focused when the popup is shown  */
-	void FocusDefaultWidget();
+	SLATE_API void FocusDefaultWidget();
 
 	/**
 	 * If InError is a non-empty string the TextBox will use the ErrorReporting widget provided during construction
 	 * If no error reporting was provided, the TextBox will create a default error reporter.
 	 * @param InError An error string used to give extra information about an error
 	 */
-	void SetError( const FText& InError );
-	void SetError( const FString& InError );
+	SLATE_API void SetError( const FText& InError );
+	SLATE_API void SetError( const FString& InError );
 
 protected:
 	/** Allows for inserting additional widgets that extend the functionality of the Popup */
diff --git a/Engine/Source/Runtime/Slate/Public/Widgets/Input/SVirtualJoystick.h b/Engine/Source/Runtime/Slate/Public/Widgets/Input/SVirtualJoystick.h
--- a/Engine/Source/Runtime/Slate/Public/Widgets/Input/SVirtualJoystick.h
+++ b/Engine/Source/Runtime/Slate/Public/Widgets/Input/SVirtualJoystick.h
@@ -15,12 +15,12 @@ class ISlateBrushSource;
 /**
  * A virtual joystsick
  */
-class SLATE_API SVirtualJoystick : public SLeafWidget
+class SVirtualJoystick : public SLeafWidget
 {
 
 public:
 	/** The settings of each zone we render */ 
-	struct SLATE_API FControlInfo
+	struct FControlInfo
 	{
 		// Set by the game
 
@@ -53,7 +53,7 @@ public:
 	};
 
 	/** The settings and current state of each zone we render */ 
-	struct SLATE_API FControlData
+	struct FControlData
 	{
 		/** Control settings */
 		FControlInfo Info;
@@ -61,7 +61,7 @@ public:
 		/**
 		 * Reset the control to a centered/inactive state
 		 */
-		void Reset();
+		SLATE_API void Reset();
 
 		// Current state
 
@@ -109,59 +109,59 @@ public:
 		{}
 	SLATE_END_ARGS()
 
-	void Construct(const FArguments& InArgs);
+	SLATE_API void Construct(const FArguments& InArgs);
 
 	/**
 	 * Static function to return if external users should create/activate/etc a touch interface
 	 * Note that this function is also used internally, so even if this returns false but an SVirtualJoystick
 	 * is created, it won't actually show any controls
 	 */
-	static bool ShouldDisplayTouchInterface();
+	static SLATE_API bool ShouldDisplayTouchInterface();
 
 	/**
 	 * Shows or hides the controls (for instance during cinematics)
 	 */
-	void SetJoystickVisibility(const bool bVisible, const bool bFade);
+	SLATE_API void SetJoystickVisibility(const bool bVisible, const bool bFade);
 
-	void AddControl(const FControlInfo& Control);
-	void ClearControls();
-	void SetControls(const TArray<FControlInfo>& InControls);
+	SLATE_API void AddControl(const FControlInfo& Control);
+	SLATE_API void ClearControls();
+	SLATE_API void SetControls(const TArray<FControlInfo>& InControls);
 
 	/**
 	 * Sets parameters that control all controls
 	 */
-	void SetGlobalParameters(float InActiveOpacity, float InInactiveOpacity, float InTimeUntilDeactive, float InTimeUntilReset, float InActivationDelay, bool InbPreventReCenter, float InStartupDelay);
+	SLATE_API void SetGlobalParameters(float InActiveOpacity, float InInactiveOpacity, float InTimeUntilDeactive, float InTimeUntilReset, float InActivationDelay, bool InbPreventReCenter, float InStartupDelay);
 
-	virtual int32 OnPaint(const FPaintArgs& Args, const FGeometry& AllottedGeometry, const FSlateRect& MyCullingRect, FSlateWindowElementList& OutDrawElements, int32 LayerId, const FWidgetStyle& InWidgetStyle, bool bParentEnabled) const override;
+	SLATE_API virtual int32 OnPaint(const FPaintArgs& Args, const FGeometry& AllottedGeometry, const FSlateRect& MyCullingRect, FSlateWindowElementList& OutDrawElements, int32 LayerId, const FWidgetStyle& InWidgetStyle, bool bParentEnabled) const override;
 	
-	virtual FVector2D ComputeDesiredSize(float) const override;
+	SLATE_API virtual FVector2D ComputeDesiredSize(float) const override;
 
-	virtual FReply OnTouchStarted(const FGeometry& MyGeometry, const FPointerEvent& Event) override;
-	virtual FReply OnTouchMoved(const FGeometry& MyGeometry, const FPointerEvent& Event) override;
-	virtual FReply OnTouchEnded(const FGeometry& MyGeometry, const FPointerEvent& Event) override;
+	SLATE_API virtual FReply OnTouchStarted(const FGeometry& MyGeometry, const FPointerEvent& Event) override;
+	SLATE_API virtual FReply OnTouchMoved(const FGeometry& MyGeometry, const FPointerEvent& Event) override;
+	SLATE_API virtual FReply OnTouchEnded(const FGeometry& MyGeometry, const FPointerEvent& Event) override;
 
-	virtual void Tick(const FGeometry& AllottedGeometry, const double InCurrentTime, const float InDeltaTime) override;
+	SLATE_API virtual void Tick(const FGeometry& AllottedGeometry, const double InCurrentTime, const float InDeltaTime) override;
 
-	virtual bool SupportsKeyboardFocus() const override;
+	SLATE_API virtual bool SupportsKeyboardFocus() const override;
 
 protected:
 	/** Callback for handling display metrics changes. */
-	virtual void HandleDisplayMetricsChanged(const FDisplayMetrics& NewDisplayMetric);
+	SLATE_API virtual void HandleDisplayMetricsChanged(const FDisplayMetrics& NewDisplayMetric);
 
-	static void AlignBoxIntoScreen(FVector2D& Position, const FVector2D& Size, const FVector2D& ScreenSize);
-	FVector2D ComputeThumbPosition(int32 ControlIndex, const FVector2D& LocalCoord, float* OutDistanceToTouchSqr = nullptr, float* OutDistanceToEdgeSqr = nullptr);
+	static SLATE_API void AlignBoxIntoScreen(FVector2D& Position, const FVector2D& Size, const FVector2D& ScreenSize);
+	SLATE_API FVector2D ComputeThumbPosition(int32 ControlIndex, const FVector2D& LocalCoord, float* OutDistanceToTouchSqr = nullptr, float* OutDistanceToEdgeSqr = nullptr);
 
 	/**
 	 * Process a touch event (on movement and possibly on initial touch)
 	 *
 	 * @return true if the touch was successful
 	 */
-	virtual bool HandleTouch(int32 ControlIndex, const FVector2D& LocalCoord, const FVector2D& ScreenSize);
+	SLATE_API virtual bool HandleTouch(int32 ControlIndex, const FVector2D& LocalCoord, const FVector2D& ScreenSize);
 
 	/** 
 	 * Return the target opacity to lerp to given the current state
 	 */
-	FORCEINLINE float GetBaseOpacity();
+	SLATE_API FORCEINLINE float GetBaseOpacity();
 
 	/**
 	 * TArray specialization for controls. In the game only few joysticks presented
diff --git a/Engine/Source/Runtime/Slate/Public/Widgets/Input/SVirtualKeyboardEntry.h b/Engine/Source/Runtime/Slate/Public/Widgets/Input/SVirtualKeyboardEntry.h
--- a/Engine/Source/Runtime/Slate/Public/Widgets/Input/SVirtualKeyboardEntry.h
+++ b/Engine/Source/Runtime/Slate/Public/Widgets/Input/SVirtualKeyboardEntry.h
@@ -17,7 +17,7 @@
 class FPaintArgs;
 class FSlateWindowElementList;
 
-class SLATE_API SVirtualKeyboardEntry : public SLeafWidget, public IVirtualKeyboardEntry
+class SVirtualKeyboardEntry : public SLeafWidget, public IVirtualKeyboardEntry
 {
 
 public:
@@ -71,42 +71,42 @@ public:
 
 
 	/** Constructor */
-	SVirtualKeyboardEntry();
+	SLATE_API SVirtualKeyboardEntry();
 	
 	/**
 	 * Construct this widget
 	 *
 	 * @param	InArgs	The declaration data for this widget
 	 */
-	void Construct( const FArguments& InArgs );
+	SLATE_API void Construct( const FArguments& InArgs );
 
 	/** Restores the text to the original state */
-	void RestoreOriginalText();
+	SLATE_API void RestoreOriginalText();
 
 	/** @return whether the current text varies from the original */
-	bool HasTextChangedFromOriginal() const;
+	SLATE_API bool HasTextChangedFromOriginal() const;
 
 	/**
 	* Sets the text currently being edited
 	*
 	* @param  InNewText  The new text
 	*/
-	void SetText(const TAttribute< FText >& InNewText);
+	SLATE_API void SetText(const TAttribute< FText >& InNewText);
 
 	/**
 	 * Sets the font used to draw the text
 	 *
 	 * @param  InNewFont	The new font to use
 	 */
-	void SetFont( const TAttribute< FSlateFontInfo >& InNewFont );
+	SLATE_API void SetFont( const TAttribute< FSlateFontInfo >& InNewFont );
 
 	/** @return Whether the user can edit the text. */
-	bool GetIsReadOnly() const;
+	SLATE_API bool GetIsReadOnly() const;
 
 public:
 	//~ Begin IVirtualKeyboardEntry Interface
-	virtual void SetTextFromVirtualKeyboard(const FText& InNewText, ETextEntryType TextEntryType) override;
-	virtual void SetSelectionFromVirtualKeyboard(int InSelStart, int InSelEnd) override;
+	SLATE_API virtual void SetTextFromVirtualKeyboard(const FText& InNewText, ETextEntryType TextEntryType) override;
+	SLATE_API virtual void SetSelectionFromVirtualKeyboard(int InSelStart, int InSelEnd) override;
 
 	virtual bool GetSelection(int& OutSelStart, int& OutSelEnd) override
 	{
@@ -152,18 +152,18 @@ public:
 protected:
 
 	//~ Begin SWidget Interface
-	virtual void Tick( const FGeometry& AllottedGeometry, const double InCurrentTime, const float InDeltaTime ) override;
-	virtual FVector2D ComputeDesiredSize(float) const override;
-	virtual int32 OnPaint( const FPaintArgs& Args, const FGeometry& AllottedGeometry, const FSlateRect& MyCullingRect, FSlateWindowElementList& OutDrawElements, int32 LayerId, const FWidgetStyle& InWidgetStyle, bool bParentEnabled ) const override;
-	virtual bool SupportsKeyboardFocus() const override;
-	virtual FReply OnFocusReceived( const FGeometry& MyGeometry, const FFocusEvent& InFocusEvent ) override;
-	virtual void OnFocusLost( const FFocusEvent& InFocusEvent ) override;
+	SLATE_API virtual void Tick( const FGeometry& AllottedGeometry, const double InCurrentTime, const float InDeltaTime ) override;
+	SLATE_API virtual FVector2D ComputeDesiredSize(float) const override;
+	SLATE_API virtual int32 OnPaint( const FPaintArgs& Args, const FGeometry& AllottedGeometry, const FSlateRect& MyCullingRect, FSlateWindowElementList& OutDrawElements, int32 LayerId, const FWidgetStyle& InWidgetStyle, bool bParentEnabled ) const override;
+	SLATE_API virtual bool SupportsKeyboardFocus() const override;
+	SLATE_API virtual FReply OnFocusReceived( const FGeometry& MyGeometry, const FFocusEvent& InFocusEvent ) override;
+	SLATE_API virtual void OnFocusLost( const FFocusEvent& InFocusEvent ) override;
 	//~ End SWidget Interface
 
 private:
 
 	/** @return The string that needs to be rendered */
-	FString GetStringToRender() const;
+	SLATE_API FString GetStringToRender() const;
 
 private:
 
diff --git a/Engine/Source/Runtime/Slate/Public/Widgets/Input/SVolumeControl.h b/Engine/Source/Runtime/Slate/Public/Widgets/Input/SVolumeControl.h
--- a/Engine/Source/Runtime/Slate/Public/Widgets/Input/SVolumeControl.h
+++ b/Engine/Source/Runtime/Slate/Public/Widgets/Input/SVolumeControl.h
@@ -17,7 +17,7 @@ class SSlider;
 /**
  * A Slate VolumeControl is a mute icon/button and volume slider.
  */
-class SLATE_API SVolumeControl : public SCompoundWidget
+class SVolumeControl : public SCompoundWidget
 {
 public:
 	DECLARE_DELEGATE_OneParam( FOnMuted, bool );
@@ -48,7 +48,7 @@ public:
 	 * 
 	 * @param InDeclaration   A declaration from which to construct the widget
 	 */
-	void Construct( const SVolumeControl::FArguments& InDeclaration );
+	SLATE_API void Construct( const SVolumeControl::FArguments& InDeclaration );
 
 	float GetVolume() const { return VolumeAttribute.Get(); }
 	bool IsMuted() const { return MutedAttribute.Get(); }
@@ -76,7 +76,7 @@ private:
 	const FSlateBrush* SpeakerIcons[ES_MAX];
 
 	float GetSliderPosition() const { return IsMuted() ? 0.f : GetVolume(); }
-	const FSlateBrush* GetSpeakerImage() const;
-	FReply OnMuteClicked();
-	void OnWriteValue(float NewValue);
+	SLATE_API const FSlateBrush* GetSpeakerImage() const;
+	SLATE_API FReply OnMuteClicked();
+	SLATE_API void OnWriteValue(float NewValue);
 };
diff --git a/Engine/Source/Runtime/Slate/Public/Widgets/LayerManager/STooltipPresenter.h b/Engine/Source/Runtime/Slate/Public/Widgets/LayerManager/STooltipPresenter.h
--- a/Engine/Source/Runtime/Slate/Public/Widgets/LayerManager/STooltipPresenter.h
+++ b/Engine/Source/Runtime/Slate/Public/Widgets/LayerManager/STooltipPresenter.h
@@ -12,7 +12,7 @@
 class FArrangedChildren;
 
 /** Presents its content at the cursor's position. Tooltip avoids widget edges. */
-class SLATE_API STooltipPresenter : public SPanel
+class STooltipPresenter : public SPanel
 {
 public:
 	SLATE_BEGIN_ARGS(STooltipPresenter)
@@ -31,17 +31,17 @@ public:
 		bCanSupportFocus = false;
 	}
 
-	void Construct(const FArguments& InArgs);
+	SLATE_API void Construct(const FArguments& InArgs);
 
-	void SetContent(TSharedPtr<SWidget> InWidget);
+	SLATE_API void SetContent(TSharedPtr<SWidget> InWidget);
 
 private:
-	virtual bool ComputeVolatility() const override;
-	virtual void OnArrangeChildren(const FGeometry& AllottedGeometry, FArrangedChildren& ArrangedChildren) const override;
+	SLATE_API virtual bool ComputeVolatility() const override;
+	SLATE_API virtual void OnArrangeChildren(const FGeometry& AllottedGeometry, FArrangedChildren& ArrangedChildren) const override;
 
-	virtual FVector2D ComputeDesiredSize(float) const override;
+	SLATE_API virtual FVector2D ComputeDesiredSize(float) const override;
 
-	virtual FChildren* GetChildren() override;
+	SLATE_API virtual FChildren* GetChildren() override;
 
 	TWeakChild<SWidget> ChildSlot;
 };
diff --git a/Engine/Source/Runtime/Slate/Public/Widgets/Layout/LinkableScrollBar.h b/Engine/Source/Runtime/Slate/Public/Widgets/Layout/LinkableScrollBar.h
--- a/Engine/Source/Runtime/Slate/Public/Widgets/Layout/LinkableScrollBar.h
+++ b/Engine/Source/Runtime/Slate/Public/Widgets/Layout/LinkableScrollBar.h
@@ -5,13 +5,13 @@
 #include "CoreMinimal.h"
 #include "Widgets/Layout/SScrollBar.h"
 
-class SLATE_API SLinkableScrollBar : public SScrollBar
+class SLinkableScrollBar : public SScrollBar
 {
 public:
 	
-	virtual void SetState(float InOffsetFraction, float InThumbSizeFraction, bool bCallOnUserScrolled = false) override;
+	SLATE_API virtual void SetState(float InOffsetFraction, float InThumbSizeFraction, bool bCallOnUserScrolled = false) override;
 
-	static void LinkScrollBars(TSharedRef<SLinkableScrollBar> Left, TSharedRef<SLinkableScrollBar> Right, TAttribute<TArray<FVector2f>> ScrollSyncRate);
+	static SLATE_API void LinkScrollBars(TSharedRef<SLinkableScrollBar> Left, TSharedRef<SLinkableScrollBar> Right, TAttribute<TArray<FVector2f>> ScrollSyncRate);
 
 private:
 	TWeakPtr<SLinkableScrollBar> LinkedScrollBarRight;
diff --git a/Engine/Source/Runtime/Slate/Public/Widgets/Layout/SBackgroundBlur.h b/Engine/Source/Runtime/Slate/Public/Widgets/Layout/SBackgroundBlur.h
--- a/Engine/Source/Runtime/Slate/Public/Widgets/Layout/SBackgroundBlur.h
+++ b/Engine/Source/Runtime/Slate/Public/Widgets/Layout/SBackgroundBlur.h
@@ -15,9 +15,9 @@
 * Note: For low-spec machines where the blur effect is too computationally expensive, 
 * a user-specified fallback image is used instead (effectively turning this into a Border)
 */
-class SLATE_API SBackgroundBlur : public SCompoundWidget
+class SBackgroundBlur : public SCompoundWidget
 {
-	SLATE_DECLARE_WIDGET(SBackgroundBlur, SCompoundWidget)
+	SLATE_DECLARE_WIDGET_API(SBackgroundBlur, SCompoundWidget, SLATE_API)
 
 public:
 	SLATE_BEGIN_ARGS(SBackgroundBlur)
@@ -46,25 +46,25 @@ public:
 	SLATE_END_ARGS()
 
 public:
-	SBackgroundBlur();
-	void Construct(const FArguments& InArgs);
-	virtual int32 OnPaint(const FPaintArgs& Args, const FGeometry& AllottedGeometry, const FSlateRect& MyCullingRect, FSlateWindowElementList& OutDrawElements, int32 LayerId, const FWidgetStyle& InWidgetStyle, bool bParentEnabled) const override;
+	SLATE_API SBackgroundBlur();
+	SLATE_API void Construct(const FArguments& InArgs);
+	SLATE_API virtual int32 OnPaint(const FPaintArgs& Args, const FGeometry& AllottedGeometry, const FSlateRect& MyCullingRect, FSlateWindowElementList& OutDrawElements, int32 LayerId, const FWidgetStyle& InWidgetStyle, bool bParentEnabled) const override;
 
-	void SetContent(const TSharedRef<SWidget>& InContent);
-	void SetApplyAlphaToBlur(bool bInApplyAlphaToBlur);
-	void SetBlurRadius(TAttribute<TOptional<int32>> InBlurRadius);
-	void SetBlurStrength(TAttribute<float> InStrength);
-	void SetCornerRadius(TAttribute<FVector4> InCornerRadius);
-	void SetLowQualityBackgroundBrush(const FSlateBrush* InBrush);
+	SLATE_API void SetContent(const TSharedRef<SWidget>& InContent);
+	SLATE_API void SetApplyAlphaToBlur(bool bInApplyAlphaToBlur);
+	SLATE_API void SetBlurRadius(TAttribute<TOptional<int32>> InBlurRadius);
+	SLATE_API void SetBlurStrength(TAttribute<float> InStrength);
+	SLATE_API void SetCornerRadius(TAttribute<FVector4> InCornerRadius);
+	SLATE_API void SetLowQualityBackgroundBrush(const FSlateBrush* InBrush);
 	
-	void SetHAlign(EHorizontalAlignment HAlign);
-	void SetVAlign(EVerticalAlignment VAlign);
-	void SetPadding(TAttribute<FMargin> InPadding);
+	SLATE_API void SetHAlign(EHorizontalAlignment HAlign);
+	SLATE_API void SetVAlign(EVerticalAlignment VAlign);
+	SLATE_API void SetPadding(TAttribute<FMargin> InPadding);
 
-	bool IsUsingLowQualityFallbackBrush() const;
+	SLATE_API bool IsUsingLowQualityFallbackBrush() const;
 
 protected:
-	void ComputeEffectiveKernelSize(float Strength, int32& OutKernelSize, int32& OutDownsampleAmount) const;
+	SLATE_API void ComputeEffectiveKernelSize(float Strength, int32& OutKernelSize, int32& OutDownsampleAmount) const;
 
 
 	/** @return an attribute reference of ColorAndOpacity */
diff --git a/Engine/Source/Runtime/Slate/Public/Widgets/Layout/SBorder.h b/Engine/Source/Runtime/Slate/Public/Widgets/Layout/SBorder.h
--- a/Engine/Source/Runtime/Slate/Public/Widgets/Layout/SBorder.h
+++ b/Engine/Source/Runtime/Slate/Public/Widgets/Layout/SBorder.h
@@ -21,9 +21,9 @@ class FSlateWindowElementList;
  * Border also has a Content() slot as well as some parameters controlling the
  * arrangement of said content.
  */
-class SLATE_API SBorder : public SCompoundWidget
+class SBorder : public SCompoundWidget
 {
-	SLATE_DECLARE_WIDGET(SBorder, SCompoundWidget)
+	SLATE_DECLARE_WIDGET_API(SBorder, SCompoundWidget, SLATE_API)
 
 public:
 
@@ -98,67 +98,67 @@ public:
 	/**
 	 * Default constructor.
 	 */
-	SBorder();
+	SLATE_API SBorder();
 
 	/**
 	 * Construct this widget
 	 *
 	 * @param	InArgs	The declaration data for this widget
 	 */
-	void Construct( const FArguments& InArgs );
+	SLATE_API void Construct( const FArguments& InArgs );
 
 	/**
 	 * Sets the content for this border
 	 *
 	 * @param	InContent	The widget to use as content for the border
 	 */
-	virtual void SetContent( TSharedRef< SWidget > InContent );
+	SLATE_API virtual void SetContent( TSharedRef< SWidget > InContent );
 
 	/**
 	 * Gets the content for this border
 	 *
 	 * @return The widget used as content for the border
 	 */
-	const TSharedRef< SWidget >& GetContent() const;
+	SLATE_API const TSharedRef< SWidget >& GetContent() const;
 
 	/** Clears out the content for the border */
-	void ClearContent();
+	SLATE_API void ClearContent();
 
 	/** Sets the color and opacity of the background image of this border. */
-	void SetBorderBackgroundColor(TAttribute<FSlateColor> InColorAndOpacity);
+	SLATE_API void SetBorderBackgroundColor(TAttribute<FSlateColor> InColorAndOpacity);
 
 	/** Gets the color and opacity of the background image of this border. */
 	FSlateColor GetBorderBackgroundColor() const { return BorderBackgroundColorAttribute.Get(); }
 
 	/** Set the desired size scale multiplier */
-	void SetDesiredSizeScale(TAttribute<FVector2D> InDesiredSizeScale);
+	SLATE_API void SetDesiredSizeScale(TAttribute<FVector2D> InDesiredSizeScale);
 	
 	/** See HAlign argument */
-	void SetHAlign(EHorizontalAlignment HAlign);
+	SLATE_API void SetHAlign(EHorizontalAlignment HAlign);
 
 	/** See VAlign argument */
-	void SetVAlign(EVerticalAlignment VAlign);
+	SLATE_API void SetVAlign(EVerticalAlignment VAlign);
 
 	/** See Padding attribute */
-	void SetPadding(TAttribute<FMargin> InPadding);
+	SLATE_API void SetPadding(TAttribute<FMargin> InPadding);
 
 	/** Set whether or not to show the disabled effect when this border is disabled */
-	void SetShowEffectWhenDisabled(TAttribute<bool> InShowEffectWhenDisabled);
+	SLATE_API void SetShowEffectWhenDisabled(TAttribute<bool> InShowEffectWhenDisabled);
 
 	/** Set the image to draw for this border. */
-	void SetBorderImage(TAttribute<const FSlateBrush*> InBorderImage);
+	SLATE_API void SetBorderImage(TAttribute<const FSlateBrush*> InBorderImage);
 
 	/** Get the image to draw for this border. */
 	const FSlateBrush* GetBorderImage() const { return BorderImageAttribute.Get(); }
 
 public:
 	// SWidget interface
-	virtual int32 OnPaint( const FPaintArgs& Args, const FGeometry& AllottedGeometry, const FSlateRect& MyCullingRect, FSlateWindowElementList& OutDrawElements, int32 LayerId, const FWidgetStyle& InWidgetStyle, bool bParentEnabled ) const override;
+	SLATE_API virtual int32 OnPaint( const FPaintArgs& Args, const FGeometry& AllottedGeometry, const FSlateRect& MyCullingRect, FSlateWindowElementList& OutDrawElements, int32 LayerId, const FWidgetStyle& InWidgetStyle, bool bParentEnabled ) const override;
 	// End of SWidget interface
 
 protected:
 	//~Begin SWidget overrides.
-	virtual FVector2D ComputeDesiredSize(float) const override;
+	SLATE_API virtual FVector2D ComputeDesiredSize(float) const override;
 	//~End SWidget overrides.
 
 	/** Get whether or not to show the disabled effect when this border is disabled */
