[Slate]

diff --git a/Engine/Source/Runtime/Slate/Public/Widgets/Layout/SBox.h b/Engine/Source/Runtime/Slate/Public/Widgets/Layout/SBox.h
--- a/Engine/Source/Runtime/Slate/Public/Widgets/Layout/SBox.h
+++ b/Engine/Source/Runtime/Slate/Public/Widgets/Layout/SBox.h
@@ -19,9 +19,9 @@ class FSlateWindowElementList;
 /**
  * SBox is the simplest layout element.
  */
-class SLATE_API SBox : public SPanel
+class SBox : public SPanel
 {
-	SLATE_DECLARE_WIDGET(SBox, SPanel)
+	SLATE_DECLARE_WIDGET_API(SBox, SPanel, SLATE_API)
 
 public:
 	class UE_DEPRECATED(5.0, "FBoxSlot is deprecated. Use FSingleWidgetChildrenWithBasicLayoutSlot or FOneSimpleMemberChild")
@@ -99,55 +99,55 @@ public:
 
 	SLATE_END_ARGS()
 
-	SBox();
+	SLATE_API SBox();
 
-	void Construct(const FArguments& InArgs);
+	SLATE_API void Construct(const FArguments& InArgs);
 
-	virtual void OnArrangeChildren(const FGeometry& AllottedGeometry, FArrangedChildren& ArrangedChildren) const override;
-	virtual FChildren* GetChildren() override;
-	virtual int32 OnPaint(const FPaintArgs& Args, const FGeometry& AllottedGeometry, const FSlateRect& MyCullingRect, FSlateWindowElementList& OutDrawElements, int32 LayerId, const FWidgetStyle& InWidgetStyle, bool bParentEnabled) const override;
+	SLATE_API virtual void OnArrangeChildren(const FGeometry& AllottedGeometry, FArrangedChildren& ArrangedChildren) const override;
+	SLATE_API virtual FChildren* GetChildren() override;
+	SLATE_API virtual int32 OnPaint(const FPaintArgs& Args, const FGeometry& AllottedGeometry, const FSlateRect& MyCullingRect, FSlateWindowElementList& OutDrawElements, int32 LayerId, const FWidgetStyle& InWidgetStyle, bool bParentEnabled) const override;
 
 	/**
 		* See the Content slot.
 		*/
-	void SetContent(const TSharedRef< SWidget >& InContent);
+	SLATE_API void SetContent(const TSharedRef< SWidget >& InContent);
 
 	/** See HAlign argument */
-	void SetHAlign(EHorizontalAlignment HAlign);
+	SLATE_API void SetHAlign(EHorizontalAlignment HAlign);
 
 	/** See VAlign argument */
-	void SetVAlign(EVerticalAlignment VAlign);
+	SLATE_API void SetVAlign(EVerticalAlignment VAlign);
 
 	/** See Padding attribute */
-	void SetPadding(TAttribute<FMargin> InPadding);
+	SLATE_API void SetPadding(TAttribute<FMargin> InPadding);
 
 	/** See WidthOverride attribute */
-	void SetWidthOverride(TAttribute<FOptionalSize> InWidthOverride);
+	SLATE_API void SetWidthOverride(TAttribute<FOptionalSize> InWidthOverride);
 
 	/** See HeightOverride attribute */
-	void SetHeightOverride(TAttribute<FOptionalSize> InHeightOverride);
+	SLATE_API void SetHeightOverride(TAttribute<FOptionalSize> InHeightOverride);
 
 	/** See MinDesiredWidth attribute */
-	void SetMinDesiredWidth(TAttribute<FOptionalSize> InMinDesiredWidth);
+	SLATE_API void SetMinDesiredWidth(TAttribute<FOptionalSize> InMinDesiredWidth);
 
 	/** See MinDesiredHeight attribute */
-	void SetMinDesiredHeight(TAttribute<FOptionalSize> InMinDesiredHeight);
+	SLATE_API void SetMinDesiredHeight(TAttribute<FOptionalSize> InMinDesiredHeight);
 
 	/** See MaxDesiredWidth attribute */
-	void SetMaxDesiredWidth(TAttribute<FOptionalSize> InMaxDesiredWidth);
+	SLATE_API void SetMaxDesiredWidth(TAttribute<FOptionalSize> InMaxDesiredWidth);
 
 	/** See MaxDesiredHeight attribute */
-	void SetMaxDesiredHeight(TAttribute<FOptionalSize> InMaxDesiredHeight);
+	SLATE_API void SetMaxDesiredHeight(TAttribute<FOptionalSize> InMaxDesiredHeight);
 
-	void SetMinAspectRatio(TAttribute<FOptionalSize> InMinAspectRatio);
+	SLATE_API void SetMinAspectRatio(TAttribute<FOptionalSize> InMinAspectRatio);
 
-	void SetMaxAspectRatio(TAttribute<FOptionalSize> InMaxAspectRatio);
+	SLATE_API void SetMaxAspectRatio(TAttribute<FOptionalSize> InMaxAspectRatio);
 
 protected:
 	// Begin SWidget overrides.
-	virtual FVector2D ComputeDesiredSize(float) const override;
-	float ComputeDesiredWidth() const;
-	float ComputeDesiredHeight() const;
+	SLATE_API virtual FVector2D ComputeDesiredSize(float) const override;
+	SLATE_API float ComputeDesiredWidth() const;
+	SLATE_API float ComputeDesiredHeight() const;
 	// End SWidget overrides.
 
 protected:
diff --git a/Engine/Source/Runtime/Slate/Public/Widgets/Layout/SConstraintCanvas.h b/Engine/Source/Runtime/Slate/Public/Widgets/Layout/SConstraintCanvas.h
--- a/Engine/Source/Runtime/Slate/Public/Widgets/Layout/SConstraintCanvas.h
+++ b/Engine/Source/Runtime/Slate/Public/Widgets/Layout/SConstraintCanvas.h
@@ -22,7 +22,7 @@ class FSlateWindowElementList;
  * ConstraintCanvas is a layout widget that allows you to arbitrary position and size child widgets in a 
  * relative coordinate space.  Additionally it permits anchoring widgets.
  */
-class SLATE_API SConstraintCanvas : public SPanel
+class SConstraintCanvas : public SPanel
 {
 public:
 
@@ -137,16 +137,16 @@ public:
 
 	SLATE_END_ARGS()
 
-	SConstraintCanvas();
+	SLATE_API SConstraintCanvas();
 
 	/**
 	 * Construct this widget
 	 *
 	 * @param	InArgs	The declaration data for this widget
 	 */
-	void Construct( const FArguments& InArgs );
+	SLATE_API void Construct( const FArguments& InArgs );
 
-	static FSlot::FSlotArguments Slot();
+	static SLATE_API FSlot::FSlotArguments Slot();
 
 	using FScopedWidgetSlotArguments = TPanelChildren<FSlot>::FScopedWidgetSlotArguments;
 	/**
@@ -154,31 +154,31 @@ public:
 	 *
 	 * @return The added slot.
 	 */
-	FScopedWidgetSlotArguments AddSlot();
+	SLATE_API FScopedWidgetSlotArguments AddSlot();
 
 	/**
 	 * Removes a particular content slot.
 	 *
 	 * @param SlotWidget The widget in the slot to remove.
 	 */
-	int32 RemoveSlot( const TSharedRef<SWidget>& SlotWidget );
+	SLATE_API int32 RemoveSlot( const TSharedRef<SWidget>& SlotWidget );
 
 	/**
 	 * Removes all slots from the panel.
 	 */
-	void ClearChildren();
+	SLATE_API void ClearChildren();
 
 public:
 
 	// Begin SWidget overrides
-	virtual void OnArrangeChildren( const FGeometry& AllottedGeometry, FArrangedChildren& ArrangedChildren ) const override;
-	virtual int32 OnPaint( const FPaintArgs& Args, const FGeometry& AllottedGeometry, const FSlateRect& MyCullingRect, FSlateWindowElementList& OutDrawElements, int32 LayerId, const FWidgetStyle& InWidgetStyle, bool bParentEnabled ) const override;
-	virtual FChildren* GetChildren() override;
+	SLATE_API virtual void OnArrangeChildren( const FGeometry& AllottedGeometry, FArrangedChildren& ArrangedChildren ) const override;
+	SLATE_API virtual int32 OnPaint( const FPaintArgs& Args, const FGeometry& AllottedGeometry, const FSlateRect& MyCullingRect, FSlateWindowElementList& OutDrawElements, int32 LayerId, const FWidgetStyle& InWidgetStyle, bool bParentEnabled ) const override;
+	SLATE_API virtual FChildren* GetChildren() override;
 	// End SWidget overrides
 
 protected:
 	// Begin SWidget overrides.
-	virtual FVector2D ComputeDesiredSize(float) const override;
+	SLATE_API virtual FVector2D ComputeDesiredSize(float) const override;
 	// End SWidget overrides.
 
 private:
@@ -187,7 +187,7 @@ private:
 	typedef TArray<bool, TInlineAllocator<16>> FArrangedChildLayers;
 
 	/** Like ArrangeChildren but also generates an array of layering information (see FArrangedChildLayers). */
-	void ArrangeLayeredChildren(const FGeometry& AllottedGeometry, FArrangedChildren& ArrangedChildren, FArrangedChildLayers& ArrangedChildLayers) const;
+	SLATE_API void ArrangeLayeredChildren(const FGeometry& AllottedGeometry, FArrangedChildren& ArrangedChildren, FArrangedChildLayers& ArrangedChildLayers) const;
 
 protected:
 
diff --git a/Engine/Source/Runtime/Slate/Public/Widgets/Layout/SDPIScaler.h b/Engine/Source/Runtime/Slate/Public/Widgets/Layout/SDPIScaler.h
--- a/Engine/Source/Runtime/Slate/Public/Widgets/Layout/SDPIScaler.h
+++ b/Engine/Source/Runtime/Slate/Public/Widgets/Layout/SDPIScaler.h
@@ -15,9 +15,9 @@ class FArrangedChildren;
 /**
  * Controls the DPI scale of its content. Can be used for zooming or shrinking arbitrary widget content.
  */
-class SLATE_API SDPIScaler : public SPanel
+class SDPIScaler : public SPanel
 {
-	SLATE_DECLARE_WIDGET(SDPIScaler, SPanel)
+	SLATE_DECLARE_WIDGET_API(SDPIScaler, SPanel, SLATE_API)
 
 public:
 	SLATE_BEGIN_ARGS( SDPIScaler )
@@ -38,25 +38,25 @@ public:
 
 	SLATE_END_ARGS()
 
-	SDPIScaler();
+	SLATE_API SDPIScaler();
 
-	void Construct( const FArguments& InArgs );
+	SLATE_API void Construct( const FArguments& InArgs );
 
-	virtual void OnArrangeChildren( const FGeometry& AllottedGeometry, FArrangedChildren& ArrangedChildren ) const override;
+	SLATE_API virtual void OnArrangeChildren( const FGeometry& AllottedGeometry, FArrangedChildren& ArrangedChildren ) const override;
 	
-	virtual FVector2D ComputeDesiredSize(float) const override;
+	SLATE_API virtual FVector2D ComputeDesiredSize(float) const override;
 
-	virtual FChildren* GetChildren() override;
+	SLATE_API virtual FChildren* GetChildren() override;
 
 	/** See the Content attribute */
-	void SetContent(TSharedRef<SWidget> InContent);
+	SLATE_API void SetContent(TSharedRef<SWidget> InContent);
 
 	/** See the DPIScale attribute */
-	void SetDPIScale(TAttribute<float> InDPIScale);
+	SLATE_API void SetDPIScale(TAttribute<float> InDPIScale);
 
 protected:
 
-	virtual float GetRelativeLayoutScale(int32 ChildIndex, float LayoutScaleMultiplier) const override;
+	SLATE_API virtual float GetRelativeLayoutScale(int32 ChildIndex, float LayoutScaleMultiplier) const override;
 
 	TSlateAttributeRef<float> GetDPIScaleAttribute() const { return TSlateAttributeRef<float>(SharedThis(this), DPIScaleAttribute); }
 
diff --git a/Engine/Source/Runtime/Slate/Public/Widgets/Layout/SExpandableArea.h b/Engine/Source/Runtime/Slate/Public/Widgets/Layout/SExpandableArea.h
--- a/Engine/Source/Runtime/Slate/Public/Widgets/Layout/SExpandableArea.h
+++ b/Engine/Source/Runtime/Slate/Public/Widgets/Layout/SExpandableArea.h
@@ -23,7 +23,7 @@ class SButton;
 /**
  * Represents an expandable area of content            
  */
-class SLATE_API SExpandableArea : public SCompoundWidget
+class SExpandableArea : public SCompoundWidget
 {
 public:
 
@@ -92,7 +92,7 @@ public:
 	 *
 	 * @param InArgs Construction arguments.
 	 */
-	void Construct( const FArguments& InArgs );
+	SLATE_API void Construct( const FArguments& InArgs );
 
 	/**
 	 * @return true if the area is currently expanded
@@ -102,54 +102,54 @@ public:
 	/**
 	 * Instantly sets the expanded state of the area
 	 */
-	virtual void SetExpanded( bool bExpanded );
+	SLATE_API virtual void SetExpanded( bool bExpanded );
 
-	virtual void SetExpanded_Animated( bool bExpanded );
+	SLATE_API virtual void SetExpanded_Animated( bool bExpanded );
 
 	/**
 	 * Returns true if area title is hovered, false otherwise.
 	 */
-	virtual bool IsTitleHovered() const;
+	SLATE_API virtual bool IsTitleHovered() const;
 
 	/**
 	 * Sets the max height of the area
 	 */
-	void SetMaxHeight(float InMaxHeight);
+	SLATE_API void SetMaxHeight(float InMaxHeight);
 
 	/**
 	 * Sets the color of all borders
 	 */
-	void SetBorderBackgroundColor(const FSlateColor& InBorderColor);
+	SLATE_API void SetBorderBackgroundColor(const FSlateColor& InBorderColor);
 
 	/**
 	 * Sets the header padding
 	 */
-	void SetHeaderPadding(FMargin InHeaderPadding);
+	SLATE_API void SetHeaderPadding(FMargin InHeaderPadding);
 
 	/**
 	 * Sets the body padding
 	 */
-	void SetAreaPadding(FMargin InAreaPadding);
+	SLATE_API void SetAreaPadding(FMargin InAreaPadding);
 
 	/**
 	 * Sets the style of the widget
 	 */
-	void SetStyle(const FExpandableAreaStyle* InStyle);
+	SLATE_API void SetStyle(const FExpandableAreaStyle* InStyle);
 
 	/**
 	 * Invalidates the style
 	 */
-	void InvalidateStyle();
+	SLATE_API void InvalidateStyle();
 
 	/**
 	 * Invalidates the border brush
 	 */
-	void InvalidateBorderBrush();
+	SLATE_API void InvalidateBorderBrush();
 
 	/**
 	 * Sets the brush of all borders
 	 */
-	void SetBorderBrush(const FSlateBrush* InBorderBrush);
+	SLATE_API void SetBorderBrush(const FSlateBrush* InBorderBrush);
 
 protected:
 
@@ -159,32 +159,32 @@ protected:
 	 * @param InArgs	Construction arguments
 	 * @param HeaderContent	User specified header content to display
 	 */
-	virtual TSharedRef<SWidget> ConstructHeaderWidget( const FArguments& InArgs, TSharedRef<SWidget> HeaderContent );
+	SLATE_API virtual TSharedRef<SWidget> ConstructHeaderWidget( const FArguments& InArgs, TSharedRef<SWidget> HeaderContent );
 
 	/**
 	 * @return The visibility of this section                   
 	 */
-	virtual EVisibility OnGetContentVisibility() const;
+	SLATE_API virtual EVisibility OnGetContentVisibility() const;
 
-	FReply OnHeaderClicked();
+	SLATE_API FReply OnHeaderClicked();
 
 	/** Toggles selection visibility when the panel is clicked. */
-	virtual void OnToggleContentVisibility();
+	SLATE_API virtual void OnToggleContentVisibility();
 
 	/**
 	 * @return The collapsed/expanded image we should show.                     
 	 */
-	virtual const FSlateBrush* OnGetCollapseImage() const;
+	SLATE_API virtual const FSlateBrush* OnGetCollapseImage() const;
 
 	/**
 	 * @return The scale of the content inside this section.  Animated 
 	 */
-	virtual FVector2D GetSectionScale() const;
+	SLATE_API virtual FVector2D GetSectionScale() const;
 
 	/**
 	 * Computes the desired size of this area. Optionally clamping to Max height
 	 */
-	virtual FVector2D ComputeDesiredSize(float) const override;
+	SLATE_API virtual FVector2D ComputeDesiredSize(float) const override;
 
 protected:
 
diff --git a/Engine/Source/Runtime/Slate/Public/Widgets/Layout/SFxWidget.h b/Engine/Source/Runtime/Slate/Public/Widgets/Layout/SFxWidget.h
--- a/Engine/Source/Runtime/Slate/Public/Widgets/Layout/SFxWidget.h
+++ b/Engine/Source/Runtime/Slate/Public/Widgets/Layout/SFxWidget.h
@@ -16,7 +16,7 @@ class FSlateWindowElementList;
  * 
  * Unless specified properties do not affect layout.
  */
-class SLATE_API SFxWidget : public SCompoundWidget
+class SFxWidget : public SCompoundWidget
 {
 public:
 	SLATE_BEGIN_ARGS( SFxWidget )
@@ -60,28 +60,28 @@ public:
 
 	SLATE_END_ARGS()
 
-	SFxWidget();
+	SLATE_API SFxWidget();
 
-	void Construct( const FArguments& InArgs );
+	SLATE_API void Construct( const FArguments& InArgs );
 
 	/** @see VisualOffset */
-	void SetVisualOffset( TAttribute<FVector2D> InOffset );
+	SLATE_API void SetVisualOffset( TAttribute<FVector2D> InOffset );
 
 	/** @see VisualOffset */
-	void SetVisualOffset( FVector InOffset );
+	SLATE_API void SetVisualOffset( FVector InOffset );
 
 	/** @see RenderScale */
-	void SetRenderScale( TAttribute<float> InScale );
+	SLATE_API void SetRenderScale( TAttribute<float> InScale );
 
 	/** @see RenderScale */
-	void SetRenderScale( float InScale );
+	SLATE_API void SetRenderScale( float InScale );
 
 	
 
 protected:
-	virtual int32 OnPaint( const FPaintArgs& Args, const FGeometry& AllottedGeometry, const FSlateRect& MyCullingRect, FSlateWindowElementList& OutDrawElements, int32 LayerId, const FWidgetStyle& InWidgetStyle, bool bParentEnabled ) const override;
-	virtual FVector2D ComputeDesiredSize(float) const override;
-	virtual void OnArrangeChildren( const FGeometry& AllottedGeometry, FArrangedChildren& ArrangedChildren ) const override;
+	SLATE_API virtual int32 OnPaint( const FPaintArgs& Args, const FGeometry& AllottedGeometry, const FSlateRect& MyCullingRect, FSlateWindowElementList& OutDrawElements, int32 LayerId, const FWidgetStyle& InWidgetStyle, bool bParentEnabled ) const override;
+	SLATE_API virtual FVector2D ComputeDesiredSize(float) const override;
+	SLATE_API virtual void OnArrangeChildren( const FGeometry& AllottedGeometry, FArrangedChildren& ArrangedChildren ) const override;
 
 	TSlateAttribute<float, EInvalidateWidgetReason::Paint> RenderScale;
 	TSlateAttribute<FVector2D, EInvalidateWidgetReason::Paint> RenderScaleOrigin;
diff --git a/Engine/Source/Runtime/Slate/Public/Widgets/Layout/SGridPanel.h b/Engine/Source/Runtime/Slate/Public/Widgets/Layout/SGridPanel.h
--- a/Engine/Source/Runtime/Slate/Public/Widgets/Layout/SGridPanel.h
+++ b/Engine/Source/Runtime/Slate/Public/Widgets/Layout/SGridPanel.h
@@ -15,9 +15,9 @@ class FArrangedChildren;
 class FPaintArgs;
 class FSlateWindowElementList;
 
-class SLATE_API SGridPanel : public SPanel
+class SGridPanel : public SPanel
 {
-	SLATE_DECLARE_WIDGET(SGridPanel, SPanel)
+	SLATE_DECLARE_WIDGET_API(SGridPanel, SPanel, SLATE_API)
 
 public:
 	// Used by the mandatory named parameter in FSlot
@@ -33,7 +33,7 @@ public:
 		int32 TheLayer;
 	};
 
-	class SLATE_API FSlot : public TBasicLayoutWidgetSlot<FSlot>
+	class FSlot : public TBasicLayoutWidgetSlot<FSlot>
 	{
 		friend SGridPanel;
 
@@ -65,7 +65,7 @@ public:
 				SLATE_ARGUMENT(TOptional<FVector2D>, Nudge)
 			SLATE_SLOT_END_ARGS()
 
-			void Construct(const FChildren& SlotOwner, FSlotArguments&& InArgs);
+			SLATE_API void Construct(const FChildren& SlotOwner, FSlotArguments&& InArgs);
 
 		public:
 			/** Which column in the grid this cell belongs to */
@@ -186,7 +186,7 @@ public:
 	/**
 	 * Used by declarative syntax to create a Slot in the specified Column, Row and Layer.
 	 */
-	static FSlot::FSlotArguments Slot( int32 Column, int32 Row, Layer InLayer = Layer(0) );
+	static SLATE_API FSlot::FSlotArguments Slot( int32 Column, int32 Row, Layer InLayer = Layer(0) );
 
 	using FScopedWidgetSlotArguments = TPanelChildren<FSlot>::FScopedWidgetSlotArguments;
 	/**
@@ -194,7 +194,7 @@ public:
 	 *
 	 * @return A reference to the newly-added slot
 	 */
-	FScopedWidgetSlotArguments AddSlot( int32 Column, int32 Row, Layer InLayer = Layer(0) );
+	SLATE_API FScopedWidgetSlotArguments AddSlot( int32 Column, int32 Row, Layer InLayer = Layer(0) );
 
 	/**
 	* Removes a slot from this panel which contains the specified SWidget
@@ -202,7 +202,7 @@ public:
 	* @param SlotWidget The widget to match when searching through the slots
 	* @returns The true if the slot was removed and false if no slot was found matching the widget
 	*/
-	bool RemoveSlot(const TSharedRef<SWidget>& SlotWidget);
+	SLATE_API bool RemoveSlot(const TSharedRef<SWidget>& SlotWidget);
 
 	SLATE_BEGIN_ARGS( SGridPanel )
 		{
@@ -241,12 +241,12 @@ public:
 		
 	SLATE_END_ARGS()
 
-	SGridPanel();
+	SLATE_API SGridPanel();
 
 	/** Removes all slots from the panel */
-	void ClearChildren();
+	SLATE_API void ClearChildren();
 	
-	void Construct( const FArguments& InArgs );
+	SLATE_API void Construct( const FArguments& InArgs );
 
 	/**
 	 * GetDesiredSize of a subregion in the graph.
@@ -256,26 +256,26 @@ public:
 	 *
 	 * @return FVector2D  The desired size of the region of cells specified.
 	 */
-	FVector2D GetDesiredRegionSize( const FIntPoint& StartCell, int32 Width, int32 Height ) const;
+	SLATE_API FVector2D GetDesiredRegionSize( const FIntPoint& StartCell, int32 Width, int32 Height ) const;
 
 	/** Specify a column to stretch instead of sizing to content. */
-	void SetColumnFill( int32 ColumnId, const TAttribute<float>& Coefficient );
+	SLATE_API void SetColumnFill( int32 ColumnId, const TAttribute<float>& Coefficient );
 
 	/** Specify a row to stretch instead of sizing to content. */
-	void SetRowFill( int32 RowId, const TAttribute<float>& Coefficient );
+	SLATE_API void SetRowFill( int32 RowId, const TAttribute<float>& Coefficient );
 
 	/** Clear the row and column fill rules. */
-	void ClearFill();
+	SLATE_API void ClearFill();
 
 public:
 
 	// SWidget interface
 
-	virtual int32 OnPaint( const FPaintArgs& Args, const FGeometry& AllottedGeometry, const FSlateRect& MyCullingRect, FSlateWindowElementList& OutDrawElements, int32 LayerId, const FWidgetStyle& InWidgetStyle, bool bParentEnabled ) const override;
-	virtual void OnArrangeChildren( const FGeometry& AllottedGeometry, FArrangedChildren& ArrangedChildren ) const override;
-	virtual void CacheDesiredSize(float) override;
-	virtual FVector2D ComputeDesiredSize(float) const override;
-	virtual FChildren* GetChildren() override;
+	SLATE_API virtual int32 OnPaint( const FPaintArgs& Args, const FGeometry& AllottedGeometry, const FSlateRect& MyCullingRect, FSlateWindowElementList& OutDrawElements, int32 LayerId, const FWidgetStyle& InWidgetStyle, bool bParentEnabled ) const override;
+	SLATE_API virtual void OnArrangeChildren( const FGeometry& AllottedGeometry, FArrangedChildren& ArrangedChildren ) const override;
+	SLATE_API virtual void CacheDesiredSize(float) override;
+	SLATE_API virtual FVector2D ComputeDesiredSize(float) const override;
+	SLATE_API virtual FChildren* GetChildren() override;
 
 private:
 
@@ -285,10 +285,10 @@ private:
 	 *
 	 * The resulting array is 1-element longer.
 	 */
-	static void ComputePartialSums( TArray<float>& TurnMeIntoPartialSums );
+	static SLATE_API void ComputePartialSums( TArray<float>& TurnMeIntoPartialSums );
 	
 	/** Given a SizeContribution, distribute it to the elements in DistributeOverMe at indexes from [StartIndex .. UpperBound) */
-	static void DistributeSizeContributions( float SizeContribution, TArray<float>& DistributeOverMe, int32 StartIndex, int32 UpperBound );
+	static SLATE_API void DistributeSizeContributions( float SizeContribution, TArray<float>& DistributeOverMe, int32 StartIndex, int32 UpperBound );
 
 	/**
 	 * Find the index where the given slot should be inserted into the list of Slots based on its LayerParam, such that Slots are sorter by layer.
@@ -296,13 +296,13 @@ private:
 	 * @param The newly-allocated slot to insert.
 	 * @return The index where the slot should be inserted.
 	 */
-	int32 FindInsertSlotLocation( const FSlot* InSlot );
+	SLATE_API int32 FindInsertSlotLocation( const FSlot* InSlot );
 
 	/** Compute the sizes of columns and rows needed to fit all the slots in this grid. */
-	void ComputeDesiredCellSizes( TArray<float>& OutColumns, TArray<float>& OutRows ) const;
+	SLATE_API void ComputeDesiredCellSizes( TArray<float>& OutColumns, TArray<float>& OutRows ) const;
 
 	/** Draw the debug grid of rows and colummns; useful for inspecting the GridPanel's logic. See OnPaint() for parameter meaning */
-	int32 LayoutDebugPaint(const FGeometry& AllottedGeometry, const FSlateRect& MyCullingRect, FSlateWindowElementList& OutDrawElements, int32 LayerId ) const;
+	SLATE_API int32 LayoutDebugPaint(const FGeometry& AllottedGeometry, const FSlateRect& MyCullingRect, FSlateWindowElementList& OutDrawElements, int32 LayerId ) const;
 	
 	/** 
 	 * Callback used to resize our internal arrays if a slot (or slot span) is changed. 
@@ -310,7 +310,7 @@ private:
 	 * @param InSlot The slot that has just changed.
 	 * @param bSlotLayerChanged Whether the slot layer changed.
 	 */
-	void NotifySlotChanged(const FSlot* InSlot, bool bSlotLayerChanged = false);
+	SLATE_API void NotifySlotChanged(const FSlot* InSlot, bool bSlotLayerChanged = false);
 
 private:
 
diff --git a/Engine/Source/Runtime/Slate/Public/Widgets/Layout/SHeader.h b/Engine/Source/Runtime/Slate/Public/Widgets/Layout/SHeader.h
--- a/Engine/Source/Runtime/Slate/Public/Widgets/Layout/SHeader.h
+++ b/Engine/Source/Runtime/Slate/Public/Widgets/Layout/SHeader.h
@@ -10,7 +10,7 @@
  * A heading for a section of content.
  * Looks like this: ------- Content --------
  */
-class SLATE_API SHeader : public SHorizontalBox
+class SHeader : public SHorizontalBox
 {
 public:
 	SLATE_BEGIN_ARGS(SHeader)
@@ -34,5 +34,5 @@ public:
 
 	SLATE_END_ARGS()
 
-	void Construct( const FArguments& InArgs );
+	SLATE_API void Construct( const FArguments& InArgs );
  };
diff --git a/Engine/Source/Runtime/Slate/Public/Widgets/Layout/SLinkedBox.h b/Engine/Source/Runtime/Slate/Public/Widgets/Layout/SLinkedBox.h
--- a/Engine/Source/Runtime/Slate/Public/Widgets/Layout/SLinkedBox.h
+++ b/Engine/Source/Runtime/Slate/Public/Widgets/Layout/SLinkedBox.h
@@ -24,22 +24,22 @@ class SLinkedBox;
  *  
  *  FLinkedBoxManager needs to be created and passed to each LinkedBox during construction.
  */
-class SLATE_API FLinkedBoxManager : public TSharedFromThis<FLinkedBoxManager> 
+class FLinkedBoxManager : public TSharedFromThis<FLinkedBoxManager> 
 {
 
 public: 
 
-	FLinkedBoxManager();
-	~FLinkedBoxManager();
+	SLATE_API FLinkedBoxManager();
+	SLATE_API ~FLinkedBoxManager();
 
 	/* Add an SLinkedBox */
-	void RegisterLinkedBox(SLinkedBox* InBox);
+	SLATE_API void RegisterLinkedBox(SLinkedBox* InBox);
 
 	/* Remove an SLinkedBox */
-	void UnregisterLinkedBox(SLinkedBox* InBox);
+	SLATE_API void UnregisterLinkedBox(SLinkedBox* InBox);
 
 	/* Used by the individual SLinkedBoxes to acquire the computed uniform size */
-	FVector2D GetUniformCellSize() const;
+	SLATE_API FVector2D GetUniformCellSize() const;
 
 protected:
 	
@@ -53,7 +53,7 @@ private:
 };
 
 /** A panel that */
-class SLATE_API SLinkedBox: public SBox
+class SLinkedBox: public SBox
 {
 
 public:
@@ -109,18 +109,18 @@ public:
 
 	SLATE_END_ARGS()
 
-	FVector2D GetChildrensDesiredSize() const;
+	SLATE_API FVector2D GetChildrensDesiredSize() const;
 
-	virtual FVector2D ComputeDesiredSize(float) const override;
+	SLATE_API virtual FVector2D ComputeDesiredSize(float) const override;
 
 	/**
 	 * Most panels do not create widgets as part of their implementation, so
 	 * they do not need to implement a Construct()
 	 */
-	void Construct( const FArguments& InArgs, TSharedRef<FLinkedBoxManager> InManager ); 
+	SLATE_API void Construct( const FArguments& InArgs, TSharedRef<FLinkedBoxManager> InManager ); 
 
-	SLinkedBox();
-	~SLinkedBox();
+	SLATE_API SLinkedBox();
+	SLATE_API ~SLinkedBox();
 	
 protected:
 
@@ -130,9 +130,9 @@ protected:
 	 *  the Manager can call a prepass on all of the sibling 
 	 *  LinkedBoxes at once.  
 	*/
-	virtual bool CustomPrepass(float LayoutScaleMultiplier) override;
+	SLATE_API virtual bool CustomPrepass(float LayoutScaleMultiplier) override;
 
-	void CustomChildPrepass();
+	SLATE_API void CustomChildPrepass();
 
 	friend class FLinkedBoxManager;
 
diff --git a/Engine/Source/Runtime/Slate/Public/Widgets/Layout/SMenuOwner.h b/Engine/Source/Runtime/Slate/Public/Widgets/Layout/SMenuOwner.h
--- a/Engine/Source/Runtime/Slate/Public/Widgets/Layout/SMenuOwner.h
+++ b/Engine/Source/Runtime/Slate/Public/Widgets/Layout/SMenuOwner.h
@@ -10,7 +10,7 @@
 /**
  * Any widget that wants to own lists of items that may have sub-menus that should be handled in the same way as multibox menus (mouse can cross other items without sub-menu closing)
  */
-class SLATE_API SMenuOwner : public SCompoundWidget
+class SMenuOwner : public SCompoundWidget
 {
 public:
 	SLATE_BEGIN_ARGS( SMenuOwner )
@@ -18,7 +18,7 @@ public:
 		SLATE_DEFAULT_SLOT( FArguments, Content )
 	SLATE_END_ARGS()
 
-	void Construct(const FArguments& InArgs);
+	SLATE_API void Construct(const FArguments& InArgs);
 
 public:
 	/**
@@ -26,7 +26,7 @@ public:
 	 *
 	 * @param	InMenuAnchor	Menu anchor for active pull-down menu or sub-menu
 	 */
-	void SetSummonedMenu( TSharedRef< SMenuAnchor > InMenuAnchor );
+	SLATE_API void SetSummonedMenu( TSharedRef< SMenuAnchor > InMenuAnchor );
 	
 
 	/**
@@ -34,12 +34,12 @@ public:
 	 *
 	 * @return	Menu anchor, or null pointer
 	 */
-	TSharedPtr< const SMenuAnchor > GetOpenMenu() const;
+	SLATE_API TSharedPtr< const SMenuAnchor > GetOpenMenu() const;
 
 	/**
 	 * For menu bar widget, closes any open pull-down or sub menus
 	 */
-	void CloseSummonedMenus();
+	SLATE_API void CloseSummonedMenus();
 
 protected:
 	/** For menu bar widgets, this stores a weak reference to the last pull-down or sub-menu that was summoned. */
diff --git a/Engine/Source/Runtime/Slate/Public/Widgets/Layout/SMissingWidget.h b/Engine/Source/Runtime/Slate/Public/Widgets/Layout/SMissingWidget.h
--- a/Engine/Source/Runtime/Slate/Public/Widgets/Layout/SMissingWidget.h
+++ b/Engine/Source/Runtime/Slate/Public/Widgets/Layout/SMissingWidget.h
@@ -7,7 +7,7 @@
 /**
  * Like a null widget, but visualizes itself as being explicitly missing.
  */
-class SLATE_API SMissingWidget
+class SMissingWidget
 {
 public:
 
@@ -16,5 +16,5 @@ public:
 	 *
 	 * @return The widget.
 	 */
-	static TSharedRef<class SWidget> MakeMissingWidget( );
+	static SLATE_API TSharedRef<class SWidget> MakeMissingWidget( );
 };
diff --git a/Engine/Source/Runtime/Slate/Public/Widgets/Layout/SPopup.h b/Engine/Source/Runtime/Slate/Public/Widgets/Layout/SPopup.h
--- a/Engine/Source/Runtime/Slate/Public/Widgets/Layout/SPopup.h
+++ b/Engine/Source/Runtime/Slate/Public/Widgets/Layout/SPopup.h
@@ -10,7 +10,7 @@ class FPaintArgs;
 class FSlateWindowElementList;
 
 /** A popup's contents show up on top of other things. */
-class SLATE_API SPopup : public SCompoundWidget
+class SPopup : public SCompoundWidget
 {
 public:
 
@@ -29,10 +29,10 @@ public:
 		bCanSupportFocus = false;
 	}
 
-	void Construct(const FArguments& InArgs);
+	SLATE_API void Construct(const FArguments& InArgs);
 
 private:
 
-	virtual int32 OnPaint(const FPaintArgs& Args, const FGeometry& AllottedGeometry, const FSlateRect& MyCullingRect, FSlateWindowElementList& OutDrawElements, int32 LayerId, const FWidgetStyle& InWidgetStyle, bool bParentEnabled) const override;
+	SLATE_API virtual int32 OnPaint(const FPaintArgs& Args, const FGeometry& AllottedGeometry, const FSlateRect& MyCullingRect, FSlateWindowElementList& OutDrawElements, int32 LayerId, const FWidgetStyle& InWidgetStyle, bool bParentEnabled) const override;
 
 };
diff --git a/Engine/Source/Runtime/Slate/Public/Widgets/Layout/SRadialBox.h b/Engine/Source/Runtime/Slate/Public/Widgets/Layout/SRadialBox.h
--- a/Engine/Source/Runtime/Slate/Public/Widgets/Layout/SRadialBox.h
+++ b/Engine/Source/Runtime/Slate/Public/Widgets/Layout/SRadialBox.h
@@ -16,7 +16,7 @@ class FArrangedChildren;
 /**
  * Arranges widgets in a circular fashion
  */
-class SLATE_API SRadialBox : public SPanel
+class SRadialBox : public SPanel
 {
 
 public:
@@ -67,47 +67,47 @@ public:
 
 	SLATE_END_ARGS()
 
-	SRadialBox();
+	SLATE_API SRadialBox();
 
-	static FSlot::FSlotArguments Slot();
+	static SLATE_API FSlot::FSlotArguments Slot();
 
 	using FScopedWidgetSlotArguments = TPanelChildren<FSlot>::FScopedWidgetSlotArguments;
-	FScopedWidgetSlotArguments AddSlot();
+	SLATE_API FScopedWidgetSlotArguments AddSlot();
 
 	/** Removes a slot from this radial box which contains the specified SWidget
 	 *
 	 * @param SlotWidget The widget to match when searching through the slots
 	 * @returns The index in the children array where the slot was removed and -1 if no slot was found matching the widget
 	 */
-	int32 RemoveSlot( const TSharedRef<SWidget>& SlotWidget );
+	SLATE_API int32 RemoveSlot( const TSharedRef<SWidget>& SlotWidget );
 
-	void Construct( const FArguments& InArgs );
+	SLATE_API void Construct( const FArguments& InArgs );
 
-	virtual void Tick( const FGeometry& AllottedGeometry, const double InCurrentTime, const float InDeltaTime ) override;
+	SLATE_API virtual void Tick( const FGeometry& AllottedGeometry, const double InCurrentTime, const float InDeltaTime ) override;
 
-	virtual void OnArrangeChildren( const FGeometry& AllottedGeometry, FArrangedChildren& ArrangedChildren ) const override;
+	SLATE_API virtual void OnArrangeChildren( const FGeometry& AllottedGeometry, FArrangedChildren& ArrangedChildren ) const override;
 
-	void ClearChildren();
+	SLATE_API void ClearChildren();
 
-	virtual FVector2D ComputeDesiredSize(float) const override;
+	SLATE_API virtual FVector2D ComputeDesiredSize(float) const override;
 
-	virtual FChildren* GetChildren() override;
+	SLATE_API virtual FChildren* GetChildren() override;
 
 	void SetStartingAngle(float InStartingAngle) { StartingAngle = InStartingAngle; }
 	void SetAngleBetweenItems(float InAngleBetweenItems) { AngleBetweenItems = InAngleBetweenItems; }
 	void SetDistributeItemsEvenly(bool bInDistributeItemsEvenly) { bDistributeItemsEvenly = bInDistributeItemsEvenly; }
 	void SetSectorCentralAngle(float InSectorCentralAngle) { SectorCentralAngle = InSectorCentralAngle; }
 
-	void SetUseAllottedWidth(bool bInUseAllottedWidth);
+	SLATE_API void SetUseAllottedWidth(bool bInUseAllottedWidth);
 
 	/** Mods the angle so it's between 0-360 */
 	UE_DEPRECATED(5.0, "NormalizeAngle is deprecated. You should use the FRotator::NormalizeAxis.")
-	int32 NormalizeAngle(int32 Angle) const;
+	SLATE_API int32 NormalizeAngle(int32 Angle) const;
 
 private:
 
 	/** Mods the angle so it's between 0-360 */
-	float InternalNormalizeAngle(float Angle) const;
+	SLATE_API float InternalNormalizeAngle(float Angle) const;
 
 	/** The slots that contain this panel's children. */
 	TPanelChildren<FSlot> Slots;
diff --git a/Engine/Source/Runtime/Slate/Public/Widgets/Layout/SResponsiveGridPanel.h b/Engine/Source/Runtime/Slate/Public/Widgets/Layout/SResponsiveGridPanel.h
--- a/Engine/Source/Runtime/Slate/Public/Widgets/Layout/SResponsiveGridPanel.h
+++ b/Engine/Source/Runtime/Slate/Public/Widgets/Layout/SResponsiveGridPanel.h
@@ -33,9 +33,9 @@ namespace SResponsiveGridSize
  * The SResponsiveGridPanel is still in development and the API may change drastically in the future
  * or maybe removed entirely.
  */
-class SLATE_API SResponsiveGridPanel : public SPanel
+class SResponsiveGridPanel : public SPanel
 {
-	SLATE_DECLARE_WIDGET(SResponsiveGridPanel, SPanel)
+	SLATE_DECLARE_WIDGET_API(SResponsiveGridPanel, SPanel, SLATE_API)
 public:
 
 	class FSlot : public TBasicLayoutWidgetSlot< FSlot >
@@ -143,32 +143,32 @@ public:
 
 	SLATE_END_ARGS()
 
-	SResponsiveGridPanel();
+	SLATE_API SResponsiveGridPanel();
 
 	/** Removes all rows from the panel */
-	void ClearChildren();
+	SLATE_API void ClearChildren();
 	
-	void Construct( const FArguments& InArgs, int32 TotalColumns );
+	SLATE_API void Construct( const FArguments& InArgs, int32 TotalColumns );
 
 	using FScopedWidgetSlotArguments = TPanelChildren<FSlot>::FScopedWidgetSlotArguments;
-	FScopedWidgetSlotArguments AddSlot(int32 Row);
+	SLATE_API FScopedWidgetSlotArguments AddSlot(int32 Row);
 
-	bool RemoveSlot(const TSharedRef<SWidget>& SlotWidget);
+	SLATE_API bool RemoveSlot(const TSharedRef<SWidget>& SlotWidget);
 
-	void SetRowFill(int32 RowId, float Coefficient);
+	SLATE_API void SetRowFill(int32 RowId, float Coefficient);
 
 public:
 
 	//~ SWidget interface
 
-	virtual int32 OnPaint( const FPaintArgs& Args, const FGeometry& AllottedGeometry, const FSlateRect& MyCullingRect, FSlateWindowElementList& OutDrawElements, int32 LayerId, const FWidgetStyle& InWidgetStyle, bool bParentEnabled ) const override;
-	virtual void OnArrangeChildren( const FGeometry& AllottedGeometry, FArrangedChildren& ArrangedChildren ) const override;
-	virtual void CacheDesiredSize(float) override;
-	virtual FChildren* GetChildren() override;
+	SLATE_API virtual int32 OnPaint( const FPaintArgs& Args, const FGeometry& AllottedGeometry, const FSlateRect& MyCullingRect, FSlateWindowElementList& OutDrawElements, int32 LayerId, const FWidgetStyle& InWidgetStyle, bool bParentEnabled ) const override;
+	SLATE_API virtual void OnArrangeChildren( const FGeometry& AllottedGeometry, FArrangedChildren& ArrangedChildren ) const override;
+	SLATE_API virtual void CacheDesiredSize(float) override;
+	SLATE_API virtual FChildren* GetChildren() override;
 
 protected:
 	//~ Begin SWidget overrides.
-	virtual FVector2D ComputeDesiredSize(float) const override;
+	SLATE_API virtual FVector2D ComputeDesiredSize(float) const override;
 	//~ End SWidget overrides.
 
 private:
@@ -179,21 +179,21 @@ private:
 	 *
 	 * The resulting array is 1-element longer.
 	 */
-	static void ComputePartialSums( TArray<float>& TurnMeIntoPartialSums );
+	static SLATE_API void ComputePartialSums( TArray<float>& TurnMeIntoPartialSums );
 	
 	/** Given a SizeContribution, distribute it to the elements in DistributeOverMe at indexes from [StartIndex .. UpperBound) */
-	static void DistributeSizeContributions( float SizeContribution, TArray<float>& DistributeOverMe, int32 StartIndex, int32 UpperBound );
+	static SLATE_API void DistributeSizeContributions( float SizeContribution, TArray<float>& DistributeOverMe, int32 StartIndex, int32 UpperBound );
 
 	/**
 	 * Return the location where to insert the slot. INDEX_NONE if we insert it at the end.
 	 */
-	int32 FindInsertSlotLocation(SResponsiveGridPanel::FSlot* InSlot) const;
+	SLATE_API int32 FindInsertSlotLocation(SResponsiveGridPanel::FSlot* InSlot) const;
 
 	/** Compute the sizes of columns and rows needed to fit all the slots in this grid. */
-	void ComputeDesiredCellSizes(float AvailableWidth, TArray<float>& OutColumns, TArray<float>& OutRows, TArray<float>& OutRowToSlot) const;
+	SLATE_API void ComputeDesiredCellSizes(float AvailableWidth, TArray<float>& OutColumns, TArray<float>& OutRows, TArray<float>& OutRowToSlot) const;
 
 	/** Draw the debug grid of rows and colummns; useful for inspecting the GridPanel's logic. See OnPaint() for parameter meaning */
-	int32 LayoutDebugPaint(const FGeometry& AllottedGeometry, const FSlateRect& MyCullingRect, FSlateWindowElementList& OutDrawElements, int32 LayerId ) const;
+	SLATE_API int32 LayoutDebugPaint(const FGeometry& AllottedGeometry, const FSlateRect& MyCullingRect, FSlateWindowElementList& OutDrawElements, int32 LayerId ) const;
 
 private:
 
diff --git a/Engine/Source/Runtime/Slate/Public/Widgets/Layout/SSafeZone.h b/Engine/Source/Runtime/Slate/Public/Widgets/Layout/SSafeZone.h
--- a/Engine/Source/Runtime/Slate/Public/Widgets/Layout/SSafeZone.h
+++ b/Engine/Source/Runtime/Slate/Public/Widgets/Layout/SSafeZone.h
@@ -26,7 +26,7 @@ class FArrangedChildren;
 		] 
 */
 
-class SLATE_API SSafeZone : public SBox
+class SSafeZone : public SBox
 {
 	SLATE_BEGIN_ARGS(SSafeZone)
 		: _HAlign( HAlign_Fill )
@@ -90,28 +90,28 @@ class SLATE_API SSafeZone : public SBox
 
 public:
 
-	SSafeZone();
-	virtual ~SSafeZone();
+	SLATE_API SSafeZone();
+	SLATE_API virtual ~SSafeZone();
 
-	void Construct( const FArguments& InArgs );
+	SLATE_API void Construct( const FArguments& InArgs );
 	
-	void SetTitleSafe( bool bIsTitleSafe );
-	void SetSafeAreaScale(FMargin InSafeAreaScale);
+	SLATE_API void SetTitleSafe( bool bIsTitleSafe );
+	SLATE_API void SetSafeAreaScale(FMargin InSafeAreaScale);
 
-	void SetSidesToPad( bool InPadLeft, bool InPadRight, bool InPadTop, bool InPadBottom );
+	SLATE_API void SetSidesToPad( bool InPadLeft, bool InPadRight, bool InPadTop, bool InPadBottom );
 
-	FMargin GetSafeMargin(float InLayoutScale) const;
+	SLATE_API FMargin GetSafeMargin(float InLayoutScale) const;
 
 #if WITH_EDITOR
-	void SetOverrideScreenInformation(TOptional<FVector2D> InScreenSize, TOptional<float> InOverrideDpiScale);
-	void DebugSafeAreaUpdated(const FMargin& NewSafeZone, bool bShouldRecacheMetrics);
+	SLATE_API void SetOverrideScreenInformation(TOptional<FVector2D> InScreenSize, TOptional<float> InOverrideDpiScale);
+	SLATE_API void DebugSafeAreaUpdated(const FMargin& NewSafeZone, bool bShouldRecacheMetrics);
 #endif
 
-	virtual void OnArrangeChildren( const FGeometry& AllottedGeometry, FArrangedChildren& ArrangedChildren ) const override;
-	virtual FVector2D ComputeDesiredSize(float LayoutScale) const override;
+	SLATE_API virtual void OnArrangeChildren( const FGeometry& AllottedGeometry, FArrangedChildren& ArrangedChildren ) const override;
+	SLATE_API virtual FVector2D ComputeDesiredSize(float LayoutScale) const override;
 
-	static void SetGlobalSafeZoneScale(TOptional<float> InScale);
-	static TOptional<float> GetGlobalSafeZoneScale();
+	static SLATE_API void SetGlobalSafeZoneScale(TOptional<float> InScale);
+	static SLATE_API TOptional<float> GetGlobalSafeZoneScale();
 
 private:
 
diff --git a/Engine/Source/Runtime/Slate/Public/Widgets/Layout/SScaleBox.h b/Engine/Source/Runtime/Slate/Public/Widgets/Layout/SScaleBox.h
--- a/Engine/Source/Runtime/Slate/Public/Widgets/Layout/SScaleBox.h
+++ b/Engine/Source/Runtime/Slate/Public/Widgets/Layout/SScaleBox.h
@@ -73,9 +73,9 @@ namespace EStretch
  * you needed to have a background image scale to fill an area but not become distorted with different aspect ratios, or if you need
  * to auto fit some text to an area, this is the control for you.
  */
-class SLATE_API SScaleBox : public SCompoundWidget
+class SScaleBox : public SCompoundWidget
 {
-	SLATE_DECLARE_WIDGET(SScaleBox, SCompoundWidget)
+	SLATE_DECLARE_WIDGET_API(SScaleBox, SCompoundWidget, SLATE_API)
 
 public:
 	SLATE_BEGIN_ARGS(SScaleBox)
@@ -117,59 +117,59 @@ public:
 
 protected:
 	/** Constructor */
-	SScaleBox();
+	SLATE_API SScaleBox();
 
 public:
-	virtual ~SScaleBox();
+	SLATE_API virtual ~SScaleBox();
 
-	void Construct(const FArguments& InArgs);
+	SLATE_API void Construct(const FArguments& InArgs);
 	
 	// SWidget interface
-	virtual void OnArrangeChildren(const FGeometry& AllottedGeometry, FArrangedChildren& ArrangedChildren) const override;
-	virtual int32 OnPaint(const FPaintArgs& Args, const FGeometry& AllottedGeometry, const FSlateRect& MyCullingRect, FSlateWindowElementList& OutDrawElements, int32 LayerId, const FWidgetStyle& InWidgetStyle, bool bParentEnabled) const override;
+	SLATE_API virtual void OnArrangeChildren(const FGeometry& AllottedGeometry, FArrangedChildren& ArrangedChildren) const override;
+	SLATE_API virtual int32 OnPaint(const FPaintArgs& Args, const FGeometry& AllottedGeometry, const FSlateRect& MyCullingRect, FSlateWindowElementList& OutDrawElements, int32 LayerId, const FWidgetStyle& InWidgetStyle, bool bParentEnabled) const override;
 	// End SWidget of interface
 
 	/** See Content slot */
-	void SetContent(TSharedRef<SWidget> InContent);
+	SLATE_API void SetContent(TSharedRef<SWidget> InContent);
 
 	/** See HAlign argument */
-	void SetHAlign(EHorizontalAlignment HAlign);
+	SLATE_API void SetHAlign(EHorizontalAlignment HAlign);
 
 	/** See VAlign argument */
-	void SetVAlign(EVerticalAlignment VAlign);
+	SLATE_API void SetVAlign(EVerticalAlignment VAlign);
 
 	/** See StretchDirection argument */
-	void SetStretchDirection(EStretchDirection::Type InStretchDirection);
+	SLATE_API void SetStretchDirection(EStretchDirection::Type InStretchDirection);
 
 	/** See Stretch argument */
-	void SetStretch(EStretch::Type InStretch);
+	SLATE_API void SetStretch(EStretch::Type InStretch);
 
 	/** See UserSpecifiedScale argument */
-	void SetUserSpecifiedScale(float InUserSpecifiedScale);
+	SLATE_API void SetUserSpecifiedScale(float InUserSpecifiedScale);
 
 	/** Set IgnoreInheritedScale argument */
-	void SetIgnoreInheritedScale(bool InIgnoreInheritedScale);
+	SLATE_API void SetIgnoreInheritedScale(bool InIgnoreInheritedScale);
 
 #if WITH_EDITOR
-	void SetOverrideScreenInformation(TOptional<FVector2D> InScreenSize);
+	SLATE_API void SetOverrideScreenInformation(TOptional<FVector2D> InScreenSize);
 #endif
 	
 protected:
 	// Begin SWidget overrides.
-	virtual bool CustomPrepass(float LayoutScaleMultiplier) override;
-	virtual FVector2D ComputeDesiredSize(float InScale) const override;
-	virtual float GetRelativeLayoutScale(int32 ChildIndex, float LayoutScaleMultiplier) const override;
+	SLATE_API virtual bool CustomPrepass(float LayoutScaleMultiplier) override;
+	SLATE_API virtual FVector2D ComputeDesiredSize(float InScale) const override;
+	SLATE_API virtual float GetRelativeLayoutScale(int32 ChildIndex, float LayoutScaleMultiplier) const override;
 	// End SWidget overrides.
 
-	bool DoesScaleRequireNormalizingPrepassOrLocalGeometry() const;
-	bool IsDesiredSizeDependentOnAreaAndScale() const;
-	float ComputeContentScale(const FGeometry& PaintGeometry) const;
+	SLATE_API bool DoesScaleRequireNormalizingPrepassOrLocalGeometry() const;
+	SLATE_API bool IsDesiredSizeDependentOnAreaAndScale() const;
+	SLATE_API float ComputeContentScale(const FGeometry& PaintGeometry) const;
 
-	void RefreshSafeZoneScale();
-	void HandleSafeFrameChangedEvent();
+	SLATE_API void RefreshSafeZoneScale();
+	SLATE_API void HandleSafeFrameChangedEvent();
 
 #if WITH_EDITOR
-	void DebugSafeAreaUpdated(const FMargin& NewSafeZone, bool bShouldRecacheMetrics);
+	SLATE_API void DebugSafeAreaUpdated(const FMargin& NewSafeZone, bool bShouldRecacheMetrics);
 #endif
 
 protected:
diff --git a/Engine/Source/Runtime/Slate/Public/Widgets/Layout/SScissorRectBox.h b/Engine/Source/Runtime/Slate/Public/Widgets/Layout/SScissorRectBox.h
--- a/Engine/Source/Runtime/Slate/Public/Widgets/Layout/SScissorRectBox.h
+++ b/Engine/Source/Runtime/Slate/Public/Widgets/Layout/SScissorRectBox.h
@@ -27,7 +27,7 @@ class FSlateWindowElementList;
  * NOTE: This class is provided as a STOPGAP for design-time widgets (like the UMG designer itself)
  *       until a more robust clipping solution can be put in place.
  */
-class SLATE_API SScissorRectBox : public SPanel
+class SScissorRectBox : public SPanel
 {
 
 public:
@@ -41,20 +41,20 @@ public:
 		SLATE_DEFAULT_SLOT(FArguments, Content)
 	SLATE_END_ARGS()
 	
-	SScissorRectBox();
+	SLATE_API SScissorRectBox();
 	
-	void Construct( const FArguments& InArgs );
+	SLATE_API void Construct( const FArguments& InArgs );
 
 	/**
 	 * See the Content slot.
      */
-	void SetContent(const TSharedRef< SWidget >& InContent);
+	SLATE_API void SetContent(const TSharedRef< SWidget >& InContent);
 
 private:
-	virtual FVector2D ComputeDesiredSize(float) const override;
-	virtual void OnArrangeChildren( const FGeometry& AllottedGeometry, FArrangedChildren& ArrangedChildren ) const override;
-	virtual FChildren* GetChildren() override;
-	virtual int32 OnPaint( const FPaintArgs& Args, const FGeometry& AllottedGeometry, const FSlateRect& MyCullingRect, FSlateWindowElementList& OutDrawElements, int32 LayerId, const FWidgetStyle& InWidgetStyle, bool bParentEnabled ) const override;
+	SLATE_API virtual FVector2D ComputeDesiredSize(float) const override;
+	SLATE_API virtual void OnArrangeChildren( const FGeometry& AllottedGeometry, FArrangedChildren& ArrangedChildren ) const override;
+	SLATE_API virtual FChildren* GetChildren() override;
+	SLATE_API virtual int32 OnPaint( const FPaintArgs& Args, const FGeometry& AllottedGeometry, const FSlateRect& MyCullingRect, FSlateWindowElementList& OutDrawElements, int32 LayerId, const FWidgetStyle& InWidgetStyle, bool bParentEnabled ) const override;
 
 	FScissorRectSlot ChildSlot;
 };
diff --git a/Engine/Source/Runtime/Slate/Public/Widgets/Layout/SScrollBar.h b/Engine/Source/Runtime/Slate/Public/Widgets/Layout/SScrollBar.h
--- a/Engine/Source/Runtime/Slate/Public/Widgets/Layout/SScrollBar.h
+++ b/Engine/Source/Runtime/Slate/Public/Widgets/Layout/SScrollBar.h
@@ -24,7 +24,7 @@ DECLARE_DELEGATE_OneParam(
 
 class SScrollBarTrack;
 
-class SLATE_API SScrollBar : public SBorder
+class SScrollBar : public SBorder
 {
 public:
 
@@ -65,14 +65,14 @@ public:
 	 *
 	 * @param	InArgs	The declaration data for this widget
 	 */
-	void Construct(const FArguments& InArgs);
+	SLATE_API void Construct(const FArguments& InArgs);
 
 	/**
 	 * Set the handler to be invoked when the user scrolls.
 	 *
 	 * @param InHandler   Method to execute when the user scrolls the scrollbar
 	 */
-	void SetOnUserScrolled( const FOnUserScrolled& InHandler );
+	SLATE_API void SetOnUserScrolled( const FOnUserScrolled& InHandler );
 
 	/**
 	 * Set the offset and size of the track's thumb.
@@ -83,83 +83,83 @@ public:
 	 * @param InThumbSizeFraction  Size of thumbnail as a fraction of the total available scroll space.
 	 * @param bCallOnUserScrolled  If true, OnUserScrolled will be called with InOffsetFraction
 	 */
-	virtual void SetState( float InOffsetFraction, float InThumbSizeFraction, bool bCallOnUserScrolled = false);
+	SLATE_API virtual void SetState( float InOffsetFraction, float InThumbSizeFraction, bool bCallOnUserScrolled = false);
 
 	// SWidget
-	virtual void Tick(const FGeometry& AllottedGeometry, const double InCurrentTime, const float InDeltaTime) override;
-	virtual FReply OnMouseButtonDown( const FGeometry& MyGeometry, const FPointerEvent& MouseEvent ) override;
-	virtual FReply OnMouseButtonUp( const FGeometry& MyGeometry, const FPointerEvent& MouseEvent ) override;
-	virtual FReply OnMouseMove( const FGeometry& MyGeometry, const FPointerEvent& MouseEvent ) override;
-	virtual void OnMouseEnter(const FGeometry& MyGeometry, const FPointerEvent& MouseEvent) override;
-	virtual void OnMouseLeave(const FPointerEvent& MouseEvent) override;
+	SLATE_API virtual void Tick(const FGeometry& AllottedGeometry, const double InCurrentTime, const float InDeltaTime) override;
+	SLATE_API virtual FReply OnMouseButtonDown( const FGeometry& MyGeometry, const FPointerEvent& MouseEvent ) override;
+	SLATE_API virtual FReply OnMouseButtonUp( const FGeometry& MyGeometry, const FPointerEvent& MouseEvent ) override;
+	SLATE_API virtual FReply OnMouseMove( const FGeometry& MyGeometry, const FPointerEvent& MouseEvent ) override;
+	SLATE_API virtual void OnMouseEnter(const FGeometry& MyGeometry, const FPointerEvent& MouseEvent) override;
+	SLATE_API virtual void OnMouseLeave(const FPointerEvent& MouseEvent) override;
 	// SWidget
 
 	/** @return true if scrolling is possible; false if the view is big enough to fit all the content */
-	bool IsNeeded() const;
+	SLATE_API bool IsNeeded() const;
 
 	/** @return normalized distance from top */
-	float DistanceFromTop() const;
+	SLATE_API float DistanceFromTop() const;
 
 	/** @return normalized distance from bottom */
-	float DistanceFromBottom() const;
+	SLATE_API float DistanceFromBottom() const;
 
 	/** @return normalized percentage of track covered by thumb bar */
-	float ThumbSizeFraction() const;
+	SLATE_API float ThumbSizeFraction() const;
 
 	/** @return the scrollbar's visibility as a product of internal rules and user-specified visibility */
-	EVisibility ShouldBeVisible() const;
+	SLATE_API EVisibility ShouldBeVisible() const;
 
 	/** @return True if the user is scrolling by dragging the scroll bar thumb. */
-	bool IsScrolling() const;
+	SLATE_API bool IsScrolling() const;
 
 	/** @return the orientation in which the scrollbar is scrolling. */
-	EOrientation GetOrientation() const; 
+	SLATE_API EOrientation GetOrientation() const; 
 
 	/** Set argument Style */
-	void SetStyle(const FScrollBarStyle* InStyle);
+	SLATE_API void SetStyle(const FScrollBarStyle* InStyle);
 
 	/** Invalidate the style */
-	void InvalidateStyle();
+	SLATE_API void InvalidateStyle();
 
 	/** Set UserVisibility attribute */
 	void SetUserVisibility(TAttribute<EVisibility> InUserVisibility) { UserVisibility = InUserVisibility; }
 
 	/** Set DragFocusCause attribute */
-	void SetDragFocusCause(EFocusCause InDragFocusCause);
+	SLATE_API void SetDragFocusCause(EFocusCause InDragFocusCause);
 
 	/** Set Thickness attribute */
-	void SetThickness(TAttribute<FVector2D> InThickness);
+	SLATE_API void SetThickness(TAttribute<FVector2D> InThickness);
 
 	/** Set ScrollBarAlwaysVisible attribute */
-	void SetScrollBarAlwaysVisible(bool InAlwaysVisible);
+	SLATE_API void SetScrollBarAlwaysVisible(bool InAlwaysVisible);
 
 	/** Set ScrollBarTrackAlwaysVisible attribute */
-	void SetScrollBarTrackAlwaysVisible(bool InAlwaysVisible);
+	SLATE_API void SetScrollBarTrackAlwaysVisible(bool InAlwaysVisible);
 
 	/** Returns True when the scrollbar should always be shown, else False */
-	bool AlwaysShowScrollbar() const;
+	SLATE_API bool AlwaysShowScrollbar() const;
 
 	/** Allows external scrolling panels to notify the scrollbar when scrolling begins. */
-	virtual void BeginScrolling();
+	SLATE_API virtual void BeginScrolling();
 
 	/** Allows external scrolling panels to notify the scrollbar when scrolling ends. */
-	virtual void EndScrolling();
+	SLATE_API virtual void EndScrolling();
 
-	SScrollBar();
+	SLATE_API SScrollBar();
 
 protected:
 	
 	/** Execute the on user scrolled delegate */
-	void ExecuteOnUserScrolled( const FGeometry& MyGeometry, const FPointerEvent& MouseEvent );
+	SLATE_API void ExecuteOnUserScrolled( const FGeometry& MyGeometry, const FPointerEvent& MouseEvent );
 
 	/** We fade the scroll track unless it is being hovered*/
-	FSlateColor GetTrackOpacity() const;
+	SLATE_API FSlateColor GetTrackOpacity() const;
 
 	/** We always show a subtle scroll thumb, but highlight it extra when dragging */
-	FLinearColor GetThumbOpacity() const;
+	SLATE_API FLinearColor GetThumbOpacity() const;
 
 	/** @return the name of an image for the scrollbar thumb based on whether the user is dragging or hovering it. */
-	const FSlateBrush* GetDragThumbImage() const;
+	SLATE_API const FSlateBrush* GetDragThumbImage() const;
 
 	/** The scrollbar's visibility as specified by the user. Will be compounded with internal visibility rules. */
 	TAttribute<EVisibility> UserVisibility;
diff --git a/Engine/Source/Runtime/Slate/Public/Widgets/Layout/SScrollBarTrack.h b/Engine/Source/Runtime/Slate/Public/Widgets/Layout/SScrollBarTrack.h
--- a/Engine/Source/Runtime/Slate/Public/Widgets/Layout/SScrollBarTrack.h
+++ b/Engine/Source/Runtime/Slate/Public/Widgets/Layout/SScrollBarTrack.h
@@ -12,7 +12,7 @@
 class FArrangedChildren;
 
 /** Arrange 3 widgets: the top track, bottom track, and thumb. */
-class SLATE_API SScrollBarTrack : public SPanel
+class SScrollBarTrack : public SPanel
 {
 public:
 	/** A ListPanel slot is very simple - it just stores a widget. */
@@ -50,7 +50,7 @@ public:
 	 *
 	 * @param InArgs  Declaration from which to construct the widget.
 	 */
-	void Construct(const FArguments& InArgs);
+	SLATE_API void Construct(const FArguments& InArgs);
 
 	struct FTrackSizeInfo
 	{
@@ -72,7 +72,7 @@ public:
 		}
 	};
 
-	FTrackSizeInfo GetTrackSizeInfo(const FGeometry& InTrackGeometry) const;
+	SLATE_API FTrackSizeInfo GetTrackSizeInfo(const FGeometry& InTrackGeometry) const;
 
 	/**
 	 * Panels arrange their children in a space described by the AllottedGeometry parameter. The results of the arrangement
@@ -81,7 +81,7 @@ public:
 	 * @param AllottedGeometry    The geometry allotted for this widget by its parent.
 	 * @param ArrangedChildren    The array to which to add the WidgetGeometries that represent the arranged children.
 	 */
-	virtual void OnArrangeChildren(const FGeometry& AllottedGeometry, FArrangedChildren& ArrangedChildren) const;
+	SLATE_API virtual void OnArrangeChildren(const FGeometry& AllottedGeometry, FArrangedChildren& ArrangedChildren) const;
 
 	/**
 	 * A Panel's desired size in the space required to arrange of its children on the screen while respecting all of
@@ -89,24 +89,24 @@ public:
 	 *
 	 * @return The desired size.
 	 */
-	virtual FVector2D ComputeDesiredSize(float) const;
+	SLATE_API virtual FVector2D ComputeDesiredSize(float) const;
 
 	/** @return  The children of a panel in a slot-agnostic way. */
-	virtual FChildren* GetChildren();
+	SLATE_API virtual FChildren* GetChildren();
 
-	void SetSizes(float InThumbOffsetFraction, float InThumbSizeFraction);
+	SLATE_API void SetSizes(float InThumbOffsetFraction, float InThumbSizeFraction);
 
-	bool IsNeeded() const;
+	SLATE_API bool IsNeeded() const;
 
-	float DistanceFromTop() const;
+	SLATE_API float DistanceFromTop() const;
 
-	float DistanceFromBottom() const;
+	SLATE_API float DistanceFromBottom() const;
 
-	float GetMinThumbSize() const;
+	SLATE_API float GetMinThumbSize() const;
 
-	float GetThumbSizeFraction() const;
+	SLATE_API float GetThumbSizeFraction() const;
 
-	void SetIsAlwaysVisible(bool InIsAlwaysVisible);
+	SLATE_API void SetIsAlwaysVisible(bool InIsAlwaysVisible);
 
 protected:
 
diff --git a/Engine/Source/Runtime/Slate/Public/Widgets/Layout/SScrollBorder.h b/Engine/Source/Runtime/Slate/Public/Widgets/Layout/SScrollBorder.h
--- a/Engine/Source/Runtime/Slate/Public/Widgets/Layout/SScrollBorder.h
+++ b/Engine/Source/Runtime/Slate/Public/Widgets/Layout/SScrollBorder.h
@@ -16,7 +16,7 @@ class IScrollableWidget;
 /**
  * Shows a border above and below a scrollable area.
  */
-class SLATE_API SScrollBorder : public SCompoundWidget
+class SScrollBorder : public SCompoundWidget
 {
 public:
 	SLATE_BEGIN_ARGS(SScrollBorder)
@@ -34,20 +34,20 @@ public:
 		SLATE_DEFAULT_SLOT( FArguments, Content )
 	SLATE_END_ARGS()
 
-	virtual ~SScrollBorder();
+	SLATE_API virtual ~SScrollBorder();
 
 	/**
 	  * Constructs a scrollable border overlay for the specified ScrollableWidget
 	  */
-	void Construct(const FArguments& InArgs, TSharedRef<IScrollableWidget> InScrollableWidget);
+	SLATE_API void Construct(const FArguments& InArgs, TSharedRef<IScrollableWidget> InScrollableWidget);
 	
 protected:
 
 	/** Gets the top border opacity. */
-	FSlateColor GetTopBorderOpacity() const;
+	SLATE_API FSlateColor GetTopBorderOpacity() const;
 
 	/** Gets the bottom border opacity. */
-	FSlateColor GetBottomBorderOpacity() const;
+	SLATE_API FSlateColor GetBottomBorderOpacity() const;
 
 	/** The border fade distance in normalized coordinates. */
 	TAttribute<FVector2D> BorderFadeDistance;
diff --git a/Engine/Source/Runtime/Slate/Public/Widgets/Layout/SScrollBox.h b/Engine/Source/Runtime/Slate/Public/Widgets/Layout/SScrollBox.h
--- a/Engine/Source/Runtime/Slate/Public/Widgets/Layout/SScrollBox.h
+++ b/Engine/Source/Runtime/Slate/Public/Widgets/Layout/SScrollBox.h
@@ -68,11 +68,11 @@ enum class EScrollWhenFocusChanges : uint8
 };
 
 /** SScrollBox can scroll through an arbitrary number of widgets. */
-class SLATE_API SScrollBox : public SCompoundWidget
+class SScrollBox : public SCompoundWidget
 {
 public:
 	/** A Slot that provides layout options for the contents of a scrollable box. */
-	class SLATE_API FSlot : public TBasicLayoutWidgetSlot<FSlot>
+	class FSlot : public TBasicLayoutWidgetSlot<FSlot>
 	{
 	public:
 		SLATE_SLOT_BEGIN_ARGS(FSlot, TBasicLayoutWidgetSlot<FSlot>)
@@ -107,8 +107,8 @@ public:
 			, MaxSize(*this, 0.0f)
 		{ }
 
-		void Construct(const FChildren& SlotOwner, FSlotArguments&& InArgs);
-		static void RegisterAttributes(FSlateWidgetSlotAttributeInitializer& AttributeInitializer);
+		SLATE_API void Construct(const FChildren& SlotOwner, FSlotArguments&& InArgs);
+		static SLATE_API void RegisterAttributes(FSlateWidgetSlotAttributeInitializer& AttributeInitializer);
 
 		/** Get the space rule this slot should occupy along scrollbox's direction. */
 		FSizeParam::ESizeRule GetSizeRule() const
@@ -249,53 +249,53 @@ public:
 
 	SLATE_END_ARGS()
 
-	SScrollBox();
+	SLATE_API SScrollBox();
 
 	/** @return a new slot. Slots contain children for SScrollBox */
-	static FSlot::FSlotArguments Slot();
+	static SLATE_API FSlot::FSlotArguments Slot();
 
-	void Construct( const FArguments& InArgs );
+	SLATE_API void Construct( const FArguments& InArgs );
 
 	using FScopedWidgetSlotArguments = TPanelChildren<FSlot>::FScopedWidgetSlotArguments;
 	/** Adds a slot to SScrollBox */
-	FScopedWidgetSlotArguments AddSlot();
+	SLATE_API FScopedWidgetSlotArguments AddSlot();
 
 	/** Removes a slot at the specified location */
-	void RemoveSlot( const TSharedRef<SWidget>& WidgetToRemove );
+	SLATE_API void RemoveSlot( const TSharedRef<SWidget>& WidgetToRemove );
 
 	/** Removes all children from the box */
-	void ClearChildren();
+	SLATE_API void ClearChildren();
 
 	/** @return Returns true if the user is currently interactively scrolling the view by holding
 		        the right mouse button and dragging. */
-	bool IsRightClickScrolling() const;
+	SLATE_API bool IsRightClickScrolling() const;
 
-	EAllowOverscroll GetAllowOverscroll() const;
+	SLATE_API EAllowOverscroll GetAllowOverscroll() const;
 
-	void SetAllowOverscroll( EAllowOverscroll NewAllowOverscroll );
+	SLATE_API void SetAllowOverscroll( EAllowOverscroll NewAllowOverscroll );
 
-	void SetAnimateWheelScrolling(bool bInAnimateWheelScrolling);
+	SLATE_API void SetAnimateWheelScrolling(bool bInAnimateWheelScrolling);
 
-	void SetWheelScrollMultiplier(float NewWheelScrollMultiplier);
+	SLATE_API void SetWheelScrollMultiplier(float NewWheelScrollMultiplier);
 	
-	void SetScrollWhenFocusChanges(EScrollWhenFocusChanges NewScrollWhenFocusChanges);
+	SLATE_API void SetScrollWhenFocusChanges(EScrollWhenFocusChanges NewScrollWhenFocusChanges);
 
-	float GetScrollOffset() const;
+	SLATE_API float GetScrollOffset() const;
 
-	float GetViewFraction() const;
+	SLATE_API float GetViewFraction() const;
 
-	float GetViewOffsetFraction() const;
+	SLATE_API float GetViewOffsetFraction() const;
 
 	/** Gets the scroll offset of the bottom of the ScrollBox in Slate Units. */
-	float GetScrollOffsetOfEnd() const;
+	SLATE_API float GetScrollOffsetOfEnd() const;
 
-	void SetScrollOffset( float NewScrollOffset );
+	SLATE_API void SetScrollOffset( float NewScrollOffset );
 
-	void ScrollToStart();
+	SLATE_API void ScrollToStart();
 
-	void ScrollToEnd();
+	SLATE_API void ScrollToEnd();
 
-	void EndInertialScrolling();
+	SLATE_API void EndInertialScrolling();
 
 	/** 
 	 * Attempt to scroll a widget into view, will safely handle non-descendant widgets 
@@ -304,60 +304,60 @@ public:
 	 * @param InAnimateScroll	Whether or not to animate the scroll
 	 * @param InDestination		Where we want the child widget to stop.
 	 */
-	void ScrollDescendantIntoView(const TSharedPtr<SWidget>& WidgetToFind, bool InAnimateScroll = true, EDescendantScrollDestination InDestination = EDescendantScrollDestination::IntoView, float Padding = 0);
+	SLATE_API void ScrollDescendantIntoView(const TSharedPtr<SWidget>& WidgetToFind, bool InAnimateScroll = true, EDescendantScrollDestination InDestination = EDescendantScrollDestination::IntoView, float Padding = 0);
 
 	/** Get the current orientation of the scrollbox. */
-	EOrientation GetOrientation();
+	SLATE_API EOrientation GetOrientation();
 
-	void SetNavigationDestination(const EDescendantScrollDestination NewNavigationDestination);
+	SLATE_API void SetNavigationDestination(const EDescendantScrollDestination NewNavigationDestination);
 
-	void SetConsumeMouseWheel(EConsumeMouseWheel NewConsumeMouseWheel);
+	SLATE_API void SetConsumeMouseWheel(EConsumeMouseWheel NewConsumeMouseWheel);
 
 	/** Sets the current orientation of the scrollbox and updates the layout */
-	void SetOrientation(EOrientation InOrientation);
+	SLATE_API void SetOrientation(EOrientation InOrientation);
 
-	void SetScrollBarVisibility(EVisibility InVisibility);
+	SLATE_API void SetScrollBarVisibility(EVisibility InVisibility);
 
-	void SetScrollBarAlwaysVisible(bool InAlwaysVisible);
+	SLATE_API void SetScrollBarAlwaysVisible(bool InAlwaysVisible);
 	
-	void SetScrollBarTrackAlwaysVisible(bool InAlwaysVisible);
+	SLATE_API void SetScrollBarTrackAlwaysVisible(bool InAlwaysVisible);
 
-	void SetScrollBarThickness(UE::Slate::FDeprecateVector2DParameter InThickness);
+	SLATE_API void SetScrollBarThickness(UE::Slate::FDeprecateVector2DParameter InThickness);
 
-	void SetScrollBarPadding(const FMargin& InPadding);
+	SLATE_API void SetScrollBarPadding(const FMargin& InPadding);
 
-	void SetScrollBarRightClickDragAllowed(bool bIsAllowed);
+	SLATE_API void SetScrollBarRightClickDragAllowed(bool bIsAllowed);
 
-	void SetStyle(const FScrollBoxStyle* InStyle);
+	SLATE_API void SetStyle(const FScrollBoxStyle* InStyle);
 	
-	void SetScrollBarStyle(const FScrollBarStyle* InBarStyle);
+	SLATE_API void SetScrollBarStyle(const FScrollBarStyle* InBarStyle);
 
-	void InvalidateStyle();
+	SLATE_API void InvalidateStyle();
 
-	void InvalidateScrollBarStyle();
+	SLATE_API void InvalidateScrollBarStyle();
 
 public:
 
 	// SWidget interface
-	virtual void Tick( const FGeometry& AllottedGeometry, const double InCurrentTime, const float InDeltaTime ) override;
-	virtual bool ComputeVolatility() const override;
-	virtual FReply OnPreviewMouseButtonDown(const FGeometry& MyGeometry, const FPointerEvent& MouseEvent) override;
-	virtual FReply OnMouseButtonDown( const FGeometry& MyGeometry, const FPointerEvent& MouseEvent ) override;
-	virtual FReply OnMouseButtonUp( const FGeometry& MyGeometry, const FPointerEvent& MouseEvent ) override;
-	virtual FReply OnMouseMove( const FGeometry& MyGeometry, const FPointerEvent& MouseEvent ) override;
-	virtual void OnMouseEnter( const FGeometry& MyGeometry, const FPointerEvent& MouseEvent ) override;
-	virtual void OnMouseLeave( const FPointerEvent& MouseEvent ) override;
-	virtual FReply OnMouseWheel( const FGeometry& MyGeometry, const FPointerEvent& MouseEvent ) override;
-	virtual FCursorReply OnCursorQuery( const FGeometry& MyGeometry, const FPointerEvent& CursorEvent ) const override;
-	virtual int32 OnPaint( const FPaintArgs& Args, const FGeometry& AllottedGeometry, const FSlateRect& MyCullingRect, FSlateWindowElementList& OutDrawElements, int32 LayerId, const FWidgetStyle& InWidgetStyle, bool bParentEnabled ) const override;
-	virtual FReply OnTouchEnded(const FGeometry& MyGeometry, const FPointerEvent& InTouchEvent) override;
-	virtual void OnMouseCaptureLost(const FCaptureLostEvent& CaptureLostEvent) override;
-	virtual FNavigationReply OnNavigation(const FGeometry& MyGeometry, const FNavigationEvent& InNavigationEvent) override;
-	virtual void OnFocusChanging(const FWeakWidgetPath& PreviousFocusPath, const FWidgetPath& NewWidgetPath, const FFocusEvent& InFocusEvent) override;
+	SLATE_API virtual void Tick( const FGeometry& AllottedGeometry, const double InCurrentTime, const float InDeltaTime ) override;
+	SLATE_API virtual bool ComputeVolatility() const override;
+	SLATE_API virtual FReply OnPreviewMouseButtonDown(const FGeometry& MyGeometry, const FPointerEvent& MouseEvent) override;
+	SLATE_API virtual FReply OnMouseButtonDown( const FGeometry& MyGeometry, const FPointerEvent& MouseEvent ) override;
+	SLATE_API virtual FReply OnMouseButtonUp( const FGeometry& MyGeometry, const FPointerEvent& MouseEvent ) override;
+	SLATE_API virtual FReply OnMouseMove( const FGeometry& MyGeometry, const FPointerEvent& MouseEvent ) override;
+	SLATE_API virtual void OnMouseEnter( const FGeometry& MyGeometry, const FPointerEvent& MouseEvent ) override;
+	SLATE_API virtual void OnMouseLeave( const FPointerEvent& MouseEvent ) override;
+	SLATE_API virtual FReply OnMouseWheel( const FGeometry& MyGeometry, const FPointerEvent& MouseEvent ) override;
+	SLATE_API virtual FCursorReply OnCursorQuery( const FGeometry& MyGeometry, const FPointerEvent& CursorEvent ) const override;
+	SLATE_API virtual int32 OnPaint( const FPaintArgs& Args, const FGeometry& AllottedGeometry, const FSlateRect& MyCullingRect, FSlateWindowElementList& OutDrawElements, int32 LayerId, const FWidgetStyle& InWidgetStyle, bool bParentEnabled ) const override;
+	SLATE_API virtual FReply OnTouchEnded(const FGeometry& MyGeometry, const FPointerEvent& InTouchEvent) override;
+	SLATE_API virtual void OnMouseCaptureLost(const FCaptureLostEvent& CaptureLostEvent) override;
+	SLATE_API virtual FNavigationReply OnNavigation(const FGeometry& MyGeometry, const FNavigationEvent& InNavigationEvent) override;
+	SLATE_API virtual void OnFocusChanging(const FWeakWidgetPath& PreviousFocusPath, const FWidgetPath& NewWidgetPath, const FFocusEvent& InFocusEvent) override;
 	// End of SWidget interface
 
 protected:
-	void OnClippingChanged();
+	SLATE_API void OnClippingChanged();
 
 private:
 
@@ -431,10 +431,10 @@ private:
 
 protected:
 	/** Scrolls or begins scrolling a widget into view, only valid to call when we have layout geometry. */
-	bool InternalScrollDescendantIntoView(const FGeometry& MyGeometry, const TSharedPtr<SWidget>& WidgetToFind, bool InAnimateScroll = true, EDescendantScrollDestination InDestination = EDescendantScrollDestination::IntoView, float Padding = 0);
+	SLATE_API bool InternalScrollDescendantIntoView(const FGeometry& MyGeometry, const TSharedPtr<SWidget>& WidgetToFind, bool InAnimateScroll = true, EDescendantScrollDestination InDestination = EDescendantScrollDestination::IntoView, float Padding = 0);
 
 	/** returns widget that can receive keyboard focus or nullprt **/
-	TSharedPtr<SWidget> GetKeyboardFocusableWidget(TSharedPtr<SWidget> InWidget);
+	SLATE_API TSharedPtr<SWidget> GetKeyboardFocusableWidget(TSharedPtr<SWidget> InWidget);
 
 	/** The panel which stacks the child slots */
 	TSharedPtr<class SScrollPanel> ScrollPanel;
@@ -601,4 +601,4 @@ private:
 	EOrientation Orientation;
 	bool BackPadScrolling;
 	bool FrontPadScrolling;
-};
\ No newline at end of file
+};
diff --git a/Engine/Source/Runtime/Slate/Public/Widgets/Layout/SSeparator.h b/Engine/Source/Runtime/Slate/Public/Widgets/Layout/SSeparator.h
--- a/Engine/Source/Runtime/Slate/Public/Widgets/Layout/SSeparator.h
+++ b/Engine/Source/Runtime/Slate/Public/Widgets/Layout/SSeparator.h
@@ -8,7 +8,7 @@
 #include "Styling/CoreStyle.h"
 #include "Widgets/Layout/SBorder.h"
 
-class SLATE_API SSeparator : public SBorder
+class SSeparator : public SBorder
 {
 public:
 	SLATE_BEGIN_ARGS(SSeparator)
@@ -33,10 +33,10 @@ public:
 	 *
 	 * @param	InArgs	The declaration data for this widget
 	 */
-	void Construct( const FArguments& InArgs );
+	SLATE_API void Construct( const FArguments& InArgs );
 
 	// SWidget interface
-	virtual FVector2D ComputeDesiredSize(float) const override;
+	SLATE_API virtual FVector2D ComputeDesiredSize(float) const override;
 	// End of SWidget interface
 private:
 	EOrientation Orientation;
diff --git a/Engine/Source/Runtime/Slate/Public/Widgets/Layout/SSpacer.h b/Engine/Source/Runtime/Slate/Public/Widgets/Layout/SSpacer.h
--- a/Engine/Source/Runtime/Slate/Public/Widgets/Layout/SSpacer.h
+++ b/Engine/Source/Runtime/Slate/Public/Widgets/Layout/SSpacer.h
@@ -11,9 +11,9 @@
 class FPaintArgs;
 class FSlateWindowElementList;
 
-class SLATE_API SSpacer : public SLeafWidget
+class SSpacer : public SLeafWidget
 {
-	SLATE_DECLARE_WIDGET(SSpacer, SLeafWidget)
+	SLATE_DECLARE_WIDGET_API(SSpacer, SLeafWidget, SLATE_API)
 
 public:
 
@@ -26,17 +26,17 @@ public:
 		SLATE_ATTRIBUTE( FVector2D, Size )
 	SLATE_END_ARGS()
 
-	SSpacer();
+	SLATE_API SSpacer();
 
 	/**
 	 * Construct this widget
 	 *
 	 * @param	InArgs	The declaration data for this widget
 	 */
-	void Construct( const FArguments& InArgs );
+	SLATE_API void Construct( const FArguments& InArgs );
 
 	// SWidget interface
-	virtual int32 OnPaint( const FPaintArgs& Args, const FGeometry& AllottedGeometry, const FSlateRect& MyCullingRect, FSlateWindowElementList& OutDrawElements, int32 LayerId, const FWidgetStyle& InWidgetStyle, bool bParentEnabled ) const override;
+	SLATE_API virtual int32 OnPaint( const FPaintArgs& Args, const FGeometry& AllottedGeometry, const FSlateRect& MyCullingRect, FSlateWindowElementList& OutDrawElements, int32 LayerId, const FWidgetStyle& InWidgetStyle, bool bParentEnabled ) const override;
 	// End of SWidget interface
 
 	FVector2D GetSize() const
@@ -57,7 +57,7 @@ public:
 
 protected:
 	//~ Begin SWidget overrides.
-	virtual FVector2D ComputeDesiredSize(float) const override;
+	SLATE_API virtual FVector2D ComputeDesiredSize(float) const override;
 	//~ End SWidget overrides.
 
 private:
diff --git a/Engine/Source/Runtime/Slate/Public/Widgets/Layout/SSplitter.h b/Engine/Source/Runtime/Slate/Public/Widgets/Layout/SSplitter.h
--- a/Engine/Source/Runtime/Slate/Public/Widgets/Layout/SSplitter.h
+++ b/Engine/Source/Runtime/Slate/Public/Widgets/Layout/SSplitter.h
@@ -39,7 +39,7 @@ class FLayoutGeometry;
  * SSplitter divides its allotted area into N segments, where N is the number of children it has.
  * It allows the users to resize the children along the splitters axis: that is, horizontally or vertically.
  */
-class SLATE_API SSplitter : public SPanel
+class SSplitter : public SPanel
 {
 
 public:
@@ -159,7 +159,7 @@ public:
 		UE_DEPRECATED(5.0, "Direct access to bIsResizable is now deprecated. Use the getter.")
 		TOptional<bool> bIsResizable;
 	};
-	PRAGMA_ENABLE_DEPRECATION_WARNINGS
+	SLATE_API PRAGMA_ENABLE_DEPRECATION_WARNINGS
 
 	/** @return Add a new FSlot() */
 	static FSlot::FSlotArguments Slot();
@@ -175,7 +175,7 @@ public:
 	 *
 	 * @return the new slot.
 	 */
-	FScopedWidgetSlotArguments AddSlot( int32 AtIndex = INDEX_NONE );
+	SLATE_API FScopedWidgetSlotArguments AddSlot( int32 AtIndex = INDEX_NONE );
 
 	DECLARE_DELEGATE_OneParam(FOnHandleHovered, int32);
 
@@ -215,14 +215,14 @@ public:
 
 	SLATE_END_ARGS()
 
-	SSplitter();
+	SLATE_API SSplitter();
 
 	/**
 	 * Construct this widget
 	 *
 	 * @param	InArgs	The declaration data for this widget
 	 */
-	void Construct( const FArguments& InArgs );
+	SLATE_API void Construct( const FArguments& InArgs );
 
 public:
 
@@ -233,7 +233,7 @@ public:
 	 *
 	 * @return Slot at the index specified by SlotIndex
 	 */
-	SSplitter::FSlot& SlotAt( int32 SlotIndex );
+	SLATE_API SSplitter::FSlot& SlotAt( int32 SlotIndex );
 
 
 	/**
@@ -241,21 +241,21 @@ public:
 	 *
 	 * @param IndexToRemove     Remove the slot and child at this index.
 	 */
-	void RemoveAt( int32 IndexToRemove );
+	SLATE_API void RemoveAt( int32 IndexToRemove );
 
 public:
 
-	virtual void OnArrangeChildren( const FGeometry& AllottedGeometry, FArrangedChildren& ArrangedChildren ) const override;
+	SLATE_API virtual void OnArrangeChildren( const FGeometry& AllottedGeometry, FArrangedChildren& ArrangedChildren ) const override;
 
 
-	virtual int32 OnPaint( const FPaintArgs& Args, const FGeometry& AllottedGeometry, const FSlateRect& MyCullingRect, FSlateWindowElementList& OutDrawElements, int32 LayerId, const FWidgetStyle& InWidgetStyle, bool bParentEnabled ) const override;
+	SLATE_API virtual int32 OnPaint( const FPaintArgs& Args, const FGeometry& AllottedGeometry, const FSlateRect& MyCullingRect, FSlateWindowElementList& OutDrawElements, int32 LayerId, const FWidgetStyle& InWidgetStyle, bool bParentEnabled ) const override;
 
 
 	/**
 	 * A Panel's desired size in the space required to arrange of its children on the screen while respecting all of
 	 * the children's desired sizes and any layout-related options specified by the user. See StackPanel for an example.
 	 */
-	virtual FVector2D ComputeDesiredSize(float) const override;
+	SLATE_API virtual FVector2D ComputeDesiredSize(float) const override;
 
 	/**
 	 * All widgets must provide a way to access their children in a layout-agnostic way.
@@ -264,7 +264,7 @@ public:
 	 * provides layout information about the child it stores. In that case
 	 * GetChildren should simply return the TPanelChildren<Slot>. See StackPanel for an example.
 	 */
-	virtual FChildren* GetChildren() override;
+	SLATE_API virtual FChildren* GetChildren() override;
 
 	/**
 	 * The system calls this method to notify the widget that a mouse button was pressed within it. This event is bubbled.
@@ -274,7 +274,7 @@ public:
 	 *
 	 * @return Whether the event was handled along with possible requests for the system to take action.
 	 */
-	virtual FReply OnMouseButtonDown( const FGeometry& MyGeometry, const FPointerEvent& MouseEvent ) override;
+	SLATE_API virtual FReply OnMouseButtonDown( const FGeometry& MyGeometry, const FPointerEvent& MouseEvent ) override;
 	
 	/**
 	 * The system calls this method to notify the widget that a mouse button was release within it. This event is bubbled.
@@ -284,35 +284,35 @@ public:
 	 *
 	 * @return Whether the event was handled along with possible requests for the system to take action.
 	 */
-	virtual FReply OnMouseButtonUp( const FGeometry& MyGeometry, const FPointerEvent& MouseEvent ) override;
+	SLATE_API virtual FReply OnMouseButtonUp( const FGeometry& MyGeometry, const FPointerEvent& MouseEvent ) override;
 
-	virtual FReply OnMouseButtonDoubleClick(const FGeometry& InMyGeometry, const FPointerEvent& InMouseEvent) override;
+	SLATE_API virtual FReply OnMouseButtonDoubleClick(const FGeometry& InMyGeometry, const FPointerEvent& InMouseEvent) override;
 	
-	virtual FReply OnMouseMove( const FGeometry& MyGeometry, const FPointerEvent& MouseEvent ) override;
+	SLATE_API virtual FReply OnMouseMove( const FGeometry& MyGeometry, const FPointerEvent& MouseEvent ) override;
 
-	virtual void OnMouseLeave( const FPointerEvent& MouseEvent ) override;
+	SLATE_API virtual void OnMouseLeave( const FPointerEvent& MouseEvent ) override;
 
 	/**
 	 * The system asks each widget under the mouse to provide a cursor. This event is bubbled.
 	 * 
 	 * @return FCursorReply::Unhandled() if the event is not handled; return FCursorReply::Cursor() otherwise.
 	 */
-	virtual FCursorReply OnCursorQuery( const FGeometry& MyGeometry, const FPointerEvent& CursorEvent ) const override;
+	SLATE_API virtual FCursorReply OnCursorQuery( const FGeometry& MyGeometry, const FPointerEvent& CursorEvent ) const override;
 
 	/**
 	 * Change the orientation of the splitter
 	 *
 	 * @param NewOrientation  Should the splitter be horizontal or vertical
 	 */
-	void SetOrientation( EOrientation NewOrientation );
+	SLATE_API void SetOrientation( EOrientation NewOrientation );
 
 	/**
 	 * @return the current orientation of the splitter.
 	 */
-	EOrientation GetOrientation() const;
+	SLATE_API EOrientation GetOrientation() const;
 
 private:
-	TArray<FLayoutGeometry> ArrangeChildrenForLayout( const FGeometry& AllottedGeometry ) const;
+	SLATE_API TArray<FLayoutGeometry> ArrangeChildrenForLayout( const FGeometry& AllottedGeometry ) const;
 
 protected:
 
@@ -321,16 +321,16 @@ protected:
 	 *
 	 * @return INDEX_NONE if no such child can be found.
 	 */
-	static int32 FindResizeableSlotBeforeHandle( int32 DraggedHandle, const TPanelChildren<FSlot>& Children );
+	static SLATE_API int32 FindResizeableSlotBeforeHandle( int32 DraggedHandle, const TPanelChildren<FSlot>& Children );
 
 	/**
 	 * Given the index of the dragged handle and the children, find a child below/right_of the dragged handle that can be resized
 	 *
 	 * @return Children.Num() if no such child can be found.
 	 */
-	static int32 FindResizeableSlotAfterHandle( int32 DraggedHandle, const TPanelChildren<FSlot>& Children );
+	static SLATE_API int32 FindResizeableSlotAfterHandle( int32 DraggedHandle, const TPanelChildren<FSlot>& Children );
 
-	static void FindAllResizeableSlotsAfterHandle( int32 DraggedHandle, const TPanelChildren<FSlot>& Children, TArray<int32, FConcurrentLinearArrayAllocator>& OutSlotIndicies );
+	static SLATE_API void FindAllResizeableSlotsAfterHandle( int32 DraggedHandle, const TPanelChildren<FSlot>& Children, TArray<int32, FConcurrentLinearArrayAllocator>& OutSlotIndicies );
 
 	/**
 	 * Resizes the children based on user input. The template parameter Orientation corresponds to the splitter being horizontal or vertical.
@@ -340,16 +340,16 @@ protected:
 	 * @param Children         A reference to this splitter's children array; we will modify the children's layout values.
 	 * @param ChildGeometries  The arranged children; we need their sizes and positions so that we can perform a resizing.
 	 */
-	void HandleResizingByMousePosition(EOrientation Orientation, const float PhysicalSplitterHandleSize, const ESplitterResizeMode::Type ResizeMode, int32 DraggedHandle, const FVector2D& LocalMousePos, TPanelChildren<FSlot>& Children, const TArray<FLayoutGeometry>& ChildGeometries );
-	void HandleResizingDelta(EOrientation Orientation, const float PhysicalSplitterHandleSize, const ESplitterResizeMode::Type ResizeMode, int32 DraggedHandle, float Delta, TPanelChildren<FSlot>& Children, const TArray<FLayoutGeometry>& ChildGeometries);
-	void HandleResizingBySize(EOrientation Orientation, const float PhysicalSplitterHandleSize, const ESplitterResizeMode::Type ResizeMode, int32 DraggedHandle, const FVector2D& DesiredSize, TPanelChildren<FSlot>& Children, const TArray<FLayoutGeometry>& ChildGeometries);
+	SLATE_API void HandleResizingByMousePosition(EOrientation Orientation, const float PhysicalSplitterHandleSize, const ESplitterResizeMode::Type ResizeMode, int32 DraggedHandle, const FVector2D& LocalMousePos, TPanelChildren<FSlot>& Children, const TArray<FLayoutGeometry>& ChildGeometries );
+	SLATE_API void HandleResizingDelta(EOrientation Orientation, const float PhysicalSplitterHandleSize, const ESplitterResizeMode::Type ResizeMode, int32 DraggedHandle, float Delta, TPanelChildren<FSlot>& Children, const TArray<FLayoutGeometry>& ChildGeometries);
+	SLATE_API void HandleResizingBySize(EOrientation Orientation, const float PhysicalSplitterHandleSize, const ESplitterResizeMode::Type ResizeMode, int32 DraggedHandle, const FVector2D& DesiredSize, TPanelChildren<FSlot>& Children, const TArray<FLayoutGeometry>& ChildGeometries);
 
 	/**
 	 * @param ProposedSize  A size that a child would like to be
 	 *
 	 * @return A size that is clamped against the minimum size allowed for children.
 	 */
-	float ClampChild(const FSlot& ChildSlot, float ProposedSize) const;
+	SLATE_API float ClampChild(const FSlot& ChildSlot, float ProposedSize) const;
 
 	/**
 	 * Given a mouse position within the splitter, figure out which resize handle we are hovering (if any).
@@ -400,10 +400,10 @@ ENUM_CLASS_FLAGS(EResizingAxis);
  * of all children along an axis as well as resizing all children
  * by dragging the center of the splitter.
  */
-class SLATE_API SSplitter2x2 : public SPanel
+class SSplitter2x2 : public SPanel
 {
 private:
-	class SLATE_API FSlot : public TSlotBase<FSlot>
+	class FSlot : public TSlotBase<FSlot>
 	{
 	public:	
 		SLATE_SLOT_BEGIN_ARGS(FSlot, TSlotBase<FSlot>)
@@ -425,9 +425,9 @@ private:
 			return PercentageAttribute.Get();
 		}
 
-		FSlot(const TSharedRef<SWidget>& InWidget);
+		SLATE_API FSlot(const TSharedRef<SWidget>& InWidget);
 
-		void Construct(const FChildren& SlotOwner, FSlotArguments&& InArg);
+		SLATE_API void Construct(const FChildren& SlotOwner, FSlotArguments&& InArg);
 
 	private:
 		/** The percentage of the alloted space of the splitter that this slot requires */
@@ -446,86 +446,86 @@ private:
 		SLATE_NAMED_SLOT( FArguments, BottomRight )
 	SLATE_END_ARGS()
 
-	SSplitter2x2();
+	SLATE_API SSplitter2x2();
 
-	void Construct( const FArguments& InArgs );
+	SLATE_API void Construct( const FArguments& InArgs );
 
 	/**
 	 * Returns the widget displayed in the splitter top left area
 	 *
 	 * @return	Top left widget
 	 */
-	TSharedRef< SWidget > GetTopLeftContent();
+	SLATE_API TSharedRef< SWidget > GetTopLeftContent();
 
 	/**
 	 * Returns the widget displayed in the splitter bottom left area
 	 *
 	 * @return	Bottom left widget
 	 */
-	TSharedRef< SWidget > GetBottomLeftContent();
+	SLATE_API TSharedRef< SWidget > GetBottomLeftContent();
 
 	/**
 	 * Returns the widget displayed in the splitter top right area
 	 *
 	 * @return	Top right widget
 	 */
-	TSharedRef< SWidget > GetTopRightContent();
+	SLATE_API TSharedRef< SWidget > GetTopRightContent();
 
 	/**
 	 * Returns the widget displayed in the splitter bottom right area
 	 *
 	 * @return	Bottom right widget
 	 */
-	TSharedRef< SWidget > GetBottomRightContent();
+	SLATE_API TSharedRef< SWidget > GetBottomRightContent();
 
 	/**
 	 * Sets the widget to be displayed in the splitter top left area
 	 *
 	 * @param	TopLeftContent	The top left widget
 	 */
-	void SetTopLeftContent( TSharedRef< SWidget > TopLeftContent );
+	SLATE_API void SetTopLeftContent( TSharedRef< SWidget > TopLeftContent );
 
 	/**
 	 * Sets the widget to be displayed in the splitter bottom left area
 	 *
 	 * @param	BottomLeftContent	The bottom left widget
 	 */
-	void SetBottomLeftContent( TSharedRef< SWidget > BottomLeftContent );
+	SLATE_API void SetBottomLeftContent( TSharedRef< SWidget > BottomLeftContent );
 
 	/**
 	 * Sets the widget to be displayed in the splitter top right area
 	 *
 	 * @param	TopRightContent	The top right widget
 	 */
-	void SetTopRightContent( TSharedRef< SWidget > TopRightContent );
+	SLATE_API void SetTopRightContent( TSharedRef< SWidget > TopRightContent );
 
 	/**
 	 * Sets the widget to be displayed in the splitter bottom right area
 	 *
 	 * @param	BottomRightContent	The bottom right widget
 	 */
-	void SetBottomRightContent( TSharedRef< SWidget > BottomRightContent );
+	SLATE_API void SetBottomRightContent( TSharedRef< SWidget > BottomRightContent );
 
 	/** Returns an array of size percentages for the children in this order: TopLeft, BottomLeft, TopRight, BottomRight */
-	void GetSplitterPercentages( TArray< FVector2D >& OutPercentages ) const;
+	SLATE_API void GetSplitterPercentages( TArray< FVector2D >& OutPercentages ) const;
 	
 	/** Sets the size percentages for the children in this order: TopLeft, BottomLeft, TopRight, BottomRight */
-	void SetSplitterPercentages( TArrayView< FVector2D > InPercentages );
+	SLATE_API void SetSplitterPercentages( TArrayView< FVector2D > InPercentages );
 
 
 private:
 
 	TArray<FLayoutGeometry> ArrangeChildrenForLayout( const FGeometry& AllottedGeometry ) const;
 
-	virtual void OnArrangeChildren( const FGeometry& AllottedGeometry, FArrangedChildren& ArrangedChildren ) const override;
+	SLATE_API virtual void OnArrangeChildren( const FGeometry& AllottedGeometry, FArrangedChildren& ArrangedChildren ) const override;
 
-	virtual int32 OnPaint(const FPaintArgs& Args, const FGeometry& AllottedGeometry, const FSlateRect& MyCullingRect, FSlateWindowElementList& OutDrawElements, int32 LayerId, const FWidgetStyle& InWidgetStyle, bool bParentEnabled) const override;
+	SLATE_API virtual int32 OnPaint(const FPaintArgs& Args, const FGeometry& AllottedGeometry, const FSlateRect& MyCullingRect, FSlateWindowElementList& OutDrawElements, int32 LayerId, const FWidgetStyle& InWidgetStyle, bool bParentEnabled) const override;
 
 	/**
 	 * A Panel's desired size in the space required to arrange of its children on the screen while respecting all of
 	 * the children's desired sizes and any layout-related options specified by the user. See StackPanel for an example.
 	 */
-	virtual FVector2D ComputeDesiredSize(float) const override;
+	SLATE_API virtual FVector2D ComputeDesiredSize(float) const override;
 
 	/**
 	 * All widgets must provide a way to access their children in a layout-agnostic way.
@@ -534,7 +534,7 @@ private:
 	 * provides layout information about the child it stores. In that case
 	 * GetChildren should simply return the TPanelChildren<Slot>. See StackPanel for an example.
 	 */
-	virtual FChildren* GetChildren() override;
+	SLATE_API virtual FChildren* GetChildren() override;
 
 	/**
 	 * The system calls this method to notify the widget that a mouse button was pressed within it. This event is bubbled.
@@ -544,7 +544,7 @@ private:
 	 *
 	 * @return Whether the event was handled along with possible requests for the system to take action.
 	 */
- 	virtual FReply OnMouseButtonDown( const FGeometry& MyGeometry, const FPointerEvent& MouseEvent ) override;
+ 	SLATE_API virtual FReply OnMouseButtonDown( const FGeometry& MyGeometry, const FPointerEvent& MouseEvent ) override;
 	
 	/**
 	 * The system calls this method to notify the widget that a mouse button was release within it. This event is bubbled.
@@ -554,7 +554,7 @@ private:
 	 *
 	 * @return Whether the event was handled along with possible requests for the system to take action.
 	 */
- 	virtual FReply OnMouseButtonUp( const FGeometry& MyGeometry, const FPointerEvent& MouseEvent ) override;
+ 	SLATE_API virtual FReply OnMouseButtonUp( const FGeometry& MyGeometry, const FPointerEvent& MouseEvent ) override;
 
 	/**
 	 * The system calls this method to notify the widget that a mouse moved within it. This event is bubbled.
@@ -564,12 +564,12 @@ private:
 	 *
 	 * @return Whether the event was handled along with possible requests for the system to take action.
 	 */
-	virtual FReply OnMouseMove( const FGeometry& MyGeometry, const FPointerEvent& MouseEvent ) override;
+	SLATE_API virtual FReply OnMouseMove( const FGeometry& MyGeometry, const FPointerEvent& MouseEvent ) override;
 
 	/**
 	 * @return The cursor that should be visible
 	 */
-	virtual FCursorReply OnCursorQuery( const FGeometry& MyGeometry, const FPointerEvent& CursorEvent ) const override;
+	SLATE_API virtual FCursorReply OnCursorQuery( const FGeometry& MyGeometry, const FPointerEvent& CursorEvent ) const override;
 
 	/**
 	 * Calculates the axis being resized
diff --git a/Engine/Source/Runtime/Slate/Public/Widgets/Layout/SUniformGridPanel.h b/Engine/Source/Runtime/Slate/Public/Widgets/Layout/SUniformGridPanel.h
--- a/Engine/Source/Runtime/Slate/Public/Widgets/Layout/SUniformGridPanel.h
+++ b/Engine/Source/Runtime/Slate/Public/Widgets/Layout/SUniformGridPanel.h
@@ -15,7 +15,7 @@
 class FArrangedChildren;
 
 /** A panel that evenly divides up available space between all of its children. */
-class SLATE_API SUniformGridPanel : public SPanel
+class SUniformGridPanel : public SPanel
 {
 public:
 	/** Stores the per-child info for this panel type */
@@ -73,14 +73,14 @@ public:
 		UE_DEPRECATED(5.0, "Direct access to Row is now deprecated. Use the getter or setter.")
 		int32 Row;
 	};
-	PRAGMA_ENABLE_DEPRECATION_WARNINGS
+	SLATE_API PRAGMA_ENABLE_DEPRECATION_WARNINGS
 
 	SUniformGridPanel();
 
 	/**
 	 * Used by declarative syntax to create a Slot in the specified Column, Row.
 	 */
-	static FSlot::FSlotArguments Slot( int32 Column, int32 Row );
+	static SLATE_API FSlot::FSlotArguments Slot( int32 Column, int32 Row );
 
 	SLATE_BEGIN_ARGS( SUniformGridPanel )
 		: _SlotPadding( FMargin(0.0f) )
@@ -104,21 +104,21 @@ public:
 
 	SLATE_END_ARGS()
 
-	void Construct( const FArguments& InArgs );
+	SLATE_API void Construct( const FArguments& InArgs );
 
 	//~ Begin SPanel Interface	
-	virtual void OnArrangeChildren( const FGeometry& AllottedGeometry, FArrangedChildren& ArrangedChildren ) const override;
-	virtual FChildren* GetChildren() override;
+	SLATE_API virtual void OnArrangeChildren( const FGeometry& AllottedGeometry, FArrangedChildren& ArrangedChildren ) const override;
+	SLATE_API virtual FChildren* GetChildren() override;
 	//~ End SPanel Interface
 
 	/** See SlotPadding attribute */
-	void SetSlotPadding(TAttribute<FMargin> InSlotPadding);
+	SLATE_API void SetSlotPadding(TAttribute<FMargin> InSlotPadding);
 
 	/** See MinDesiredSlotWidth attribute */
-	void SetMinDesiredSlotWidth(TAttribute<float> InMinDesiredSlotWidth);
+	SLATE_API void SetMinDesiredSlotWidth(TAttribute<float> InMinDesiredSlotWidth);
 
 	/** See MinDesiredSlotHeight attribute */
-	void SetMinDesiredSlotHeight(TAttribute<float> InMinDesiredSlotHeight);
+	SLATE_API void SetMinDesiredSlotHeight(TAttribute<float> InMinDesiredSlotHeight);
 
 	using FScopedWidgetSlotArguments = TPanelChildren<FSlot>::FScopedWidgetSlotArguments;
 	/**
@@ -126,7 +126,7 @@ public:
 	 *
 	 * @return A reference to the newly-added slot
 	 */
-	FScopedWidgetSlotArguments AddSlot( int32 Column, int32 Row );
+	SLATE_API FScopedWidgetSlotArguments AddSlot( int32 Column, int32 Row );
 	
 	/**
 	 * Removes a slot from this panel which contains the specified SWidget
@@ -134,14 +134,14 @@ public:
 	 * @param SlotWidget The widget to match when searching through the slots
 	 * @returns The true if the slot was removed and false if no slot was found matching the widget
 	 */
-	bool RemoveSlot( const TSharedRef<SWidget>& SlotWidget );
+	SLATE_API bool RemoveSlot( const TSharedRef<SWidget>& SlotWidget );
 
 	/** Removes all slots from the panel */
-	void ClearChildren();
+	SLATE_API void ClearChildren();
 
 protected:
 	// Begin SWidget overrides.
-	virtual FVector2D ComputeDesiredSize(float) const override;
+	SLATE_API virtual FVector2D ComputeDesiredSize(float) const override;
 	// End SWidget overrides.
 
 private:
diff --git a/Engine/Source/Runtime/Slate/Public/Widgets/Layout/SUniformWrapPanel.h b/Engine/Source/Runtime/Slate/Public/Widgets/Layout/SUniformWrapPanel.h
--- a/Engine/Source/Runtime/Slate/Public/Widgets/Layout/SUniformWrapPanel.h
+++ b/Engine/Source/Runtime/Slate/Public/Widgets/Layout/SUniformWrapPanel.h
@@ -15,9 +15,9 @@
 class FArrangedChildren;
 
 /** A panel that evenly divides up available space between all of its children. */
-class SLATE_API SUniformWrapPanel : public SPanel
+class SUniformWrapPanel : public SPanel
 {
-	SLATE_DECLARE_WIDGET(SUniformWrapPanel, SPanel)
+	SLATE_DECLARE_WIDGET_API(SUniformWrapPanel, SPanel, SLATE_API)
 
 public:
 	/** Stores the per-child info for this panel type */
@@ -38,7 +38,7 @@ public:
 		}
 	};
 
-	SUniformWrapPanel();
+	SLATE_API SUniformWrapPanel();
 
 	/**
 	 * Used by declarative syntax to create a Slot.
@@ -90,41 +90,41 @@ public:
 
 	SLATE_END_ARGS()
 
-	void Construct( const FArguments& InArgs );
+	SLATE_API void Construct( const FArguments& InArgs );
 
 	//~ Begin SPanel Interface	
-	virtual void OnArrangeChildren( const FGeometry& AllottedGeometry, FArrangedChildren& ArrangedChildren ) const override;
-	virtual FChildren* GetChildren() override;
+	SLATE_API virtual void OnArrangeChildren( const FGeometry& AllottedGeometry, FArrangedChildren& ArrangedChildren ) const override;
+	SLATE_API virtual FChildren* GetChildren() override;
 	//~ End SPanel Interface
 
 	/** See SlotPadding attribute */
-	void SetSlotPadding(TAttribute<FMargin> InSlotPadding);
+	SLATE_API void SetSlotPadding(TAttribute<FMargin> InSlotPadding);
 
 	/** See MinDesiredSlotWidth attribute */
-	void SetMinDesiredSlotWidth(TAttribute<float> InMinDesiredSlotWidth);
+	SLATE_API void SetMinDesiredSlotWidth(TAttribute<float> InMinDesiredSlotWidth);
 
 	/** See MinDesiredSlotHeight attribute */
-	void SetMinDesiredSlotHeight(TAttribute<float> InMinDesiredSlotHeight);
+	SLATE_API void SetMinDesiredSlotHeight(TAttribute<float> InMinDesiredSlotHeight);
 
 	/** See NumColumnsOverride attribute */
-	void SetNumColumnsOverride(TAttribute<int32> InNumColumnsOverride);
+	SLATE_API void SetNumColumnsOverride(TAttribute<int32> InNumColumnsOverride);
 	
 	/** See MinDesiredSlotWidth attribute */
-	void SetMaxDesiredSlotWidth(TAttribute<float> InMaxDesiredSlotWidth);
+	SLATE_API void SetMaxDesiredSlotWidth(TAttribute<float> InMaxDesiredSlotWidth);
 
 	/** See MinDesiredSlotHeight attribute */
-	void SetMaxDesiredSlotHeight(TAttribute<float> InMaxDesiredSlotHeight);
+	SLATE_API void SetMaxDesiredSlotHeight(TAttribute<float> InMaxDesiredSlotHeight);
 
 	/** See HAlign attribute */
-	void SetHorizontalAlignment(TAttribute<EHorizontalAlignment> InHAlignment);
+	SLATE_API void SetHorizontalAlignment(TAttribute<EHorizontalAlignment> InHAlignment);
 	EHorizontalAlignment GetHorizontalAlignment() { return HAlign.Get(); }
 
 	/** See EvenRowDistribution attribute */
-	void SetEvenRowDistribution(TAttribute<bool> InEvenRowDistribution);
+	SLATE_API void SetEvenRowDistribution(TAttribute<bool> InEvenRowDistribution);
 	bool GetEvenRowDistribution() { return EvenRowDistribution.Get(); }
 
 	using FScopedWidgetSlotArguments = TPanelChildren<FSlot>::FScopedWidgetSlotArguments;
-	FScopedWidgetSlotArguments AddSlot();
+	SLATE_API FScopedWidgetSlotArguments AddSlot();
 	
 	/**
 	 * Removes a slot from this panel which contains the specified SWidget
@@ -132,17 +132,17 @@ public:
 	 * @param SlotWidget The widget to match when searching through the slots
 	 * @returns The true if the slot was removed and false if no slot was found matching the widget
 	 */
-	bool RemoveSlot( const TSharedRef<SWidget>& SlotWidget );
+	SLATE_API bool RemoveSlot( const TSharedRef<SWidget>& SlotWidget );
 
 	/** Removes all slots from the panel */
-	void ClearChildren();
+	SLATE_API void ClearChildren();
 
 protected:
 	// Begin SWidget overrides.
-	virtual FVector2D ComputeDesiredSize(float) const override;
+	SLATE_API virtual FVector2D ComputeDesiredSize(float) const override;
 	// End SWidget overrides.
 
-	FVector2D ComputeUniformCellSize() const;
+	SLATE_API FVector2D ComputeUniformCellSize() const;
 
 private:
 	TPanelChildren<FSlot> Children;
diff --git a/Engine/Source/Runtime/Slate/Public/Widgets/Layout/SWidgetSwitcher.h b/Engine/Source/Runtime/Slate/Public/Widgets/Layout/SWidgetSwitcher.h
--- a/Engine/Source/Runtime/Slate/Public/Widgets/Layout/SWidgetSwitcher.h
+++ b/Engine/Source/Runtime/Slate/Public/Widgets/Layout/SWidgetSwitcher.h
@@ -18,7 +18,7 @@ class FArrangedChildren;
  *
  * A widget switcher is like a tab control, but without tabs. At most one widget is visible at time.
  */
-class SLATE_API SWidgetSwitcher
+class SWidgetSwitcher
 	: public SPanel
 {
 public:
@@ -38,7 +38,7 @@ public:
 
 	SLATE_END_ARGS()
 
-	SWidgetSwitcher();
+	SLATE_API SWidgetSwitcher();
 
 public:
 
@@ -48,19 +48,19 @@ public:
 	 *
 	 * @param SlotIndex The index at which to insert the slot, or INDEX_NONE to append.
 	 */
-	FScopedWidgetSlotArguments AddSlot( int32 SlotIndex = INDEX_NONE );
+	SLATE_API FScopedWidgetSlotArguments AddSlot( int32 SlotIndex = INDEX_NONE );
 
 	/**
 	 * Constructs the widget.
 	 */
-	void Construct( const FArguments& InArgs );
+	SLATE_API void Construct( const FArguments& InArgs );
 
 	/**
 	 * Gets the active widget.
 	 *
 	 * @return Active widget.
 	 */
-	TSharedPtr<SWidget> GetActiveWidget( ) const;
+	SLATE_API TSharedPtr<SWidget> GetActiveWidget( ) const;
 
 	/**
 	 * Gets the slot index of the currently active widget.
@@ -88,7 +88,7 @@ public:
 	 * @param SlotIndex The slot index of the widget to get.
 	 * @return The widget, or nullptr if the slot does not exist.
 	 */
-	TSharedPtr<SWidget> GetWidget( int32 SlotIndex ) const;
+	SLATE_API TSharedPtr<SWidget> GetWidget( int32 SlotIndex ) const;
 
 	/**
 	 * Gets the slot index of the specified widget.
@@ -96,14 +96,14 @@ public:
 	 * @param Widget The widget to get the index for.
 	 * @return The slot index, or INDEX_NONE if the widget does not exist.
 	 */
-	int32 GetWidgetIndex( TSharedRef<SWidget> Widget ) const;
+	SLATE_API int32 GetWidgetIndex( TSharedRef<SWidget> Widget ) const;
 
 	/**
 	 * Removes a slot with the corresponding widget in it.  Returns the index where the widget was found, otherwise -1.
 	 *
 	 * @param Widget The widget to find and remove.
 	 */
-	int32 RemoveSlot( TSharedRef<SWidget> WidgetToRemove );
+	SLATE_API int32 RemoveSlot( TSharedRef<SWidget> WidgetToRemove );
 
 	/**
 	 * Sets the active widget.
@@ -120,9 +120,9 @@ public:
 	 *
 	 * @param Index The slot index.
 	 */
-	void SetActiveWidgetIndex( int32 Index );
+	SLATE_API void SetActiveWidgetIndex( int32 Index );
 
-	virtual bool ValidatePathToChild(SWidget* InChild) override;
+	SLATE_API virtual bool ValidatePathToChild(SWidget* InChild) override;
 
 public:
 
@@ -131,17 +131,17 @@ public:
 	 *
 	 * @return A new slot.
 	 */
-	static FSlot::FSlotArguments Slot();
+	static SLATE_API FSlot::FSlotArguments Slot();
 
 protected:
 
 	// SCompoundWidget interface
 
-	virtual void OnArrangeChildren( const FGeometry& AllottedGeometry, FArrangedChildren& ArrangedChildren ) const override;
-	virtual FVector2D ComputeDesiredSize( float ) const override;
-	virtual FChildren* GetChildren( ) override;
+	SLATE_API virtual void OnArrangeChildren( const FGeometry& AllottedGeometry, FArrangedChildren& ArrangedChildren ) const override;
+	SLATE_API virtual FVector2D ComputeDesiredSize( float ) const override;
+	SLATE_API virtual FChildren* GetChildren( ) override;
 	virtual bool ComputeVolatility() const override { return WidgetIndex.IsBound(); }
-	const FSlot* GetActiveSlot() const;
+	SLATE_API const FSlot* GetActiveSlot() const;
 
 	TPanelChildren<FSlot>& GetTypedChildren() { return AllChildren; }
 
