[Slate]

diff --git a/Engine/Source/Runtime/Slate/Public/Widgets/Layout/SWindowTitleBarArea.h b/Engine/Source/Runtime/Slate/Public/Widgets/Layout/SWindowTitleBarArea.h
--- a/Engine/Source/Runtime/Slate/Public/Widgets/Layout/SWindowTitleBarArea.h
+++ b/Engine/Source/Runtime/Slate/Public/Widgets/Layout/SWindowTitleBarArea.h
@@ -20,7 +20,7 @@ class SWindow;
 
 template< typename ObjectType > class TAttribute;
 
-class SLATE_API SWindowTitleBarArea
+class SWindowTitleBarArea
 	: public SPanel
 {
 public:
@@ -51,30 +51,30 @@ public:
 
 	SLATE_END_ARGS()
 
-	SWindowTitleBarArea();
+	SLATE_API SWindowTitleBarArea();
 
-	void Construct( const FArguments& InArgs );
+	SLATE_API void Construct( const FArguments& InArgs );
 
-	virtual FVector2D ComputeDesiredSize(float) const override;
-	virtual void OnArrangeChildren( const FGeometry& AllottedGeometry, FArrangedChildren& ArrangedChildren ) const override;
-	virtual FChildren* GetChildren() override;
-	virtual int32 OnPaint( const FPaintArgs& Args, const FGeometry& AllottedGeometry, const FSlateRect& MyCullingRect, FSlateWindowElementList& OutDrawElements, int32 LayerId, const FWidgetStyle& InWidgetStyle, bool bParentEnabled ) const override;
+	SLATE_API virtual FVector2D ComputeDesiredSize(float) const override;
+	SLATE_API virtual void OnArrangeChildren( const FGeometry& AllottedGeometry, FArrangedChildren& ArrangedChildren ) const override;
+	SLATE_API virtual FChildren* GetChildren() override;
+	SLATE_API virtual int32 OnPaint( const FPaintArgs& Args, const FGeometry& AllottedGeometry, const FSlateRect& MyCullingRect, FSlateWindowElementList& OutDrawElements, int32 LayerId, const FWidgetStyle& InWidgetStyle, bool bParentEnabled ) const override;
 
-	virtual FReply OnMouseButtonDoubleClick(const FGeometry& MyGeometry, const FPointerEvent& MouseEvent) override;
+	SLATE_API virtual FReply OnMouseButtonDoubleClick(const FGeometry& MyGeometry, const FPointerEvent& MouseEvent) override;
 
-	virtual EWindowZone::Type GetWindowZoneOverride() const override;
+	SLATE_API virtual EWindowZone::Type GetWindowZoneOverride() const override;
 
 	/** See the Content slot. */
-	void SetContent(const TSharedRef< SWidget >& InContent);
+	SLATE_API void SetContent(const TSharedRef< SWidget >& InContent);
 
 	/** See HAlign argument */
-	void SetHAlign(EHorizontalAlignment HAlign);
+	SLATE_API void SetHAlign(EHorizontalAlignment HAlign);
 
 	/** See VAlign argument */
-	void SetVAlign(EVerticalAlignment VAlign);
+	SLATE_API void SetVAlign(EVerticalAlignment VAlign);
 
 	/** See Padding attribute */
-	void SetPadding(TAttribute<FMargin> InPadding);
+	SLATE_API void SetPadding(TAttribute<FMargin> InPadding);
 
 	void SetGameWindow(TSharedPtr<SWindow> Window)
 	{
@@ -118,25 +118,25 @@ private:
 	TSharedPtr<SWindow> GameWindow;
 	FSimpleDelegate RequestToggleFullscreen;
 
-	static FSimpleDelegate OnCloseButtonClicked;
-	static bool bIsCloseButtonActive;
+	static SLATE_API FSimpleDelegate OnCloseButtonClicked;
+	static SLATE_API bool bIsCloseButtonActive;
 
-	static FButtonStyle MinimizeButtonStyle;
-	static FButtonStyle MaximizeButtonStyle;
-	static FButtonStyle RestoreButtonStyle;
-	static FButtonStyle CloseButtonStyle;
+	static SLATE_API FButtonStyle MinimizeButtonStyle;
+	static SLATE_API FButtonStyle MaximizeButtonStyle;
+	static SLATE_API FButtonStyle RestoreButtonStyle;
+	static SLATE_API FButtonStyle CloseButtonStyle;
 
 	bool bIsMinimizeButtonEnabled;
 	bool bIsMaximizeRestoreButtonEnabled;
 	bool bIsCloseButtonEnabled;
 
-	FReply MinimizeButton_OnClicked();
-	FReply MaximizeRestoreButton_OnClicked();
-	FReply CloseButton_OnClicked();
+	SLATE_API FReply MinimizeButton_OnClicked();
+	SLATE_API FReply MaximizeRestoreButton_OnClicked();
+	SLATE_API FReply CloseButton_OnClicked();
 
-	const FSlateBrush* GetMinimizeImage() const;
-	const FSlateBrush* GetMaximizeRestoreImage() const;
-	const FSlateBrush* GetCloseImage() const;
+	SLATE_API const FSlateBrush* GetMinimizeImage() const;
+	SLATE_API const FSlateBrush* GetMaximizeRestoreImage() const;
+	SLATE_API const FSlateBrush* GetCloseImage() const;
 
 	TSharedPtr<SButton> MinimizeButton;
 	TSharedPtr<SButton> MaximizeRestoreButton;
diff --git a/Engine/Source/Runtime/Slate/Public/Widgets/Layout/SWrapBox.h b/Engine/Source/Runtime/Slate/Public/Widgets/Layout/SWrapBox.h
--- a/Engine/Source/Runtime/Slate/Public/Widgets/Layout/SWrapBox.h
+++ b/Engine/Source/Runtime/Slate/Public/Widgets/Layout/SWrapBox.h
@@ -47,9 +47,9 @@ class FArrangedChildren;
   *		 [-3-]
   */
 
-class SLATE_API SWrapBox : public SPanel
+class SWrapBox : public SPanel
 {
-	SLATE_DECLARE_WIDGET(SWrapBox, SPanel)
+	SLATE_DECLARE_WIDGET_API(SWrapBox, SPanel, SLATE_API)
 public:
 
 	/** A slot that support alignment of content and padding */
@@ -175,7 +175,7 @@ public:
 		/** Determines if the wrap box needs to arrange the slots left-to-right or top-to-bottom.*/
 		SLATE_ARGUMENT(EOrientation, Orientation);
 	SLATE_END_ARGS()
-	PRAGMA_ENABLE_DEPRECATION_WARNINGS
+	SLATE_API PRAGMA_ENABLE_DEPRECATION_WARNINGS
 
 	SWrapBox();
 
@@ -185,49 +185,49 @@ public:
 	}
 
 	using FScopedWidgetSlotArguments = TPanelChildren<FSlot>::FScopedWidgetSlotArguments;
-	FScopedWidgetSlotArguments AddSlot();
+	SLATE_API FScopedWidgetSlotArguments AddSlot();
 
 	/** Removes a slot from this box panel which contains the specified SWidget
 	 *
 	 * @param SlotWidget The widget to match when searching through the slots
 	 * @returns The index in the children array where the slot was removed and -1 if no slot was found matching the widget
 	 */
-	int32 RemoveSlot( const TSharedRef<SWidget>& SlotWidget );
+	SLATE_API int32 RemoveSlot( const TSharedRef<SWidget>& SlotWidget );
 
-	void Construct( const FArguments& InArgs );
+	SLATE_API void Construct( const FArguments& InArgs );
 
-	virtual void Tick( const FGeometry& AllottedGeometry, const double InCurrentTime, const float InDeltaTime ) override;
+	SLATE_API virtual void Tick( const FGeometry& AllottedGeometry, const double InCurrentTime, const float InDeltaTime ) override;
 
-	virtual void OnArrangeChildren( const FGeometry& AllottedGeometry, FArrangedChildren& ArrangedChildren ) const override;
+	SLATE_API virtual void OnArrangeChildren( const FGeometry& AllottedGeometry, FArrangedChildren& ArrangedChildren ) const override;
 
-	void ClearChildren();
+	SLATE_API void ClearChildren();
 
-	virtual FVector2D ComputeDesiredSize(float) const override;
+	SLATE_API virtual FVector2D ComputeDesiredSize(float) const override;
 
-	virtual FChildren* GetChildren() override;
+	SLATE_API virtual FChildren* GetChildren() override;
 
 	/** See InnerSlotPadding Attribute */
-	void SetInnerSlotPadding(FVector2D InInnerSlotPadding);
+	SLATE_API void SetInnerSlotPadding(FVector2D InInnerSlotPadding);
 
 	/** Set the width at which the wrap panel should wrap its content. */
 	UE_DEPRECATED(4.26, "Deprecated, please use SetWrapSize() instead")
-	void SetWrapWidth(TAttribute<float> InWrapWidth);
+	SLATE_API void SetWrapWidth(TAttribute<float> InWrapWidth);
 
 	/** Set the size at which the wrap panel should wrap its content. */
-	void SetWrapSize(TAttribute<float> InWrapSize );
+	SLATE_API void SetWrapSize(TAttribute<float> InWrapSize );
 
 	/** When true, use the WrapWidth property to determine where to wrap to the next line. */
 	UE_DEPRECATED(4.26, "Deprecated, please use SetUseAllottedSize() instead")
-	void SetUseAllottedWidth(bool bInUseAllottedWidth);
+	SLATE_API void SetUseAllottedWidth(bool bInUseAllottedWidth);
 
 	/** When true, use the WrapSize property to determine where to wrap to the next line. */
-	void SetUseAllottedSize(bool bInUseAllottedSize);
+	SLATE_API void SetUseAllottedSize(bool bInUseAllottedSize);
 
 	/** Set the Orientation to determine if the wrap box needs to arrange the slots left-to-right or top-to-bottom */
-	void SetOrientation(EOrientation InOrientation);
+	SLATE_API void SetOrientation(EOrientation InOrientation);
 
 	/** How to distribute the elements among any extra space in a given row */
-	void SetHorizontalAlignment(TAttribute<EHorizontalAlignment> InHAlignment);
+	SLATE_API void SetHorizontalAlignment(TAttribute<EHorizontalAlignment> InHAlignment);
 
 private:
 	/** How wide or long, dependently of the orientation, this panel should appear to be. Any widgets past this line will be wrapped onto the next line. */
@@ -250,4 +250,4 @@ private:
 
 	class FChildArranger;
 	friend class SWrapBox::FChildArranger;
-};
\ No newline at end of file
+};
diff --git a/Engine/Source/Runtime/Slate/Public/Widgets/Notifications/GlobalNotification.h b/Engine/Source/Runtime/Slate/Public/Widgets/Notifications/GlobalNotification.h
--- a/Engine/Source/Runtime/Slate/Public/Widgets/Notifications/GlobalNotification.h
+++ b/Engine/Source/Runtime/Slate/Public/Widgets/Notifications/GlobalNotification.h
@@ -8,7 +8,7 @@
 /**
  * Class used to provide simple global notifications
  */
-class SLATE_API FGlobalNotification
+class FGlobalNotification
 {
 public:
 	FGlobalNotification(const double InEnableDelayInSeconds = 1.0)
@@ -22,7 +22,7 @@ public:
 	}
 
 	/** Called to Tick this notification and update its state */
-	void TickNotification(float DeltaTime);
+	SLATE_API void TickNotification(float DeltaTime);
 
 protected:
 	/** 
@@ -36,10 +36,10 @@ protected:
 	
 private:
 	/** Begin the notification (make it visible, and mark it as pending) */
-	TSharedPtr<SNotificationItem> BeginNotification();
+	SLATE_API TSharedPtr<SNotificationItem> BeginNotification();
 
 	/** End the notification (hide it, and mark is at complete) */
-	void EndNotification();
+	SLATE_API void EndNotification();
 
 	/** Delay (in seconds) before we should actually show the notification (default is 1) */
 	double EnableDelayInSeconds;
diff --git a/Engine/Source/Runtime/Slate/Public/Widgets/Notifications/INotificationWidget.h b/Engine/Source/Runtime/Slate/Public/Widgets/Notifications/INotificationWidget.h
--- a/Engine/Source/Runtime/Slate/Public/Widgets/Notifications/INotificationWidget.h
+++ b/Engine/Source/Runtime/Slate/Public/Widgets/Notifications/INotificationWidget.h
@@ -6,7 +6,7 @@
 #include "Widgets/SWidget.h"
 #include "Widgets/Notifications/SNotificationList.h"
 
-class SLATE_API INotificationWidget
+class INotificationWidget
 {
 public:
 	virtual void OnSetCompletionState(SNotificationItem::ECompletionState State) = 0;
diff --git a/Engine/Source/Runtime/Slate/Public/Widgets/Notifications/SErrorHint.h b/Engine/Source/Runtime/Slate/Public/Widgets/Notifications/SErrorHint.h
--- a/Engine/Source/Runtime/Slate/Public/Widgets/Notifications/SErrorHint.h
+++ b/Engine/Source/Runtime/Slate/Public/Widgets/Notifications/SErrorHint.h
@@ -11,7 +11,7 @@
 #include "Widgets/SCompoundWidget.h"
 #include "Widgets/Notifications/SErrorText.h"
 
-class SLATE_API SErrorHint
+class SErrorHint
 	: public SCompoundWidget
 	, public IErrorReportingWidget
 {
@@ -23,16 +23,16 @@ public:
 		SLATE_ARGUMENT(FText, ErrorText)
 	SLATE_END_ARGS()
 
-	void Construct(const FArguments& InArgs);
+	SLATE_API void Construct(const FArguments& InArgs);
 
 public:
 
 	// IErrorReportingWidget interface
 
-	virtual void SetError( const FText& InErrorText ) override;
-	virtual void SetError( const FString& InErrorText ) override;
-	virtual bool HasError() const override;
-	virtual TSharedRef<SWidget> AsWidget() override;
+	SLATE_API virtual void SetError( const FText& InErrorText ) override;
+	SLATE_API virtual void SetError( const FString& InErrorText ) override;
+	SLATE_API virtual bool HasError() const override;
+	SLATE_API virtual TSharedRef<SWidget> AsWidget() override;
 
 private:
 
diff --git a/Engine/Source/Runtime/Slate/Public/Widgets/Notifications/SErrorText.h b/Engine/Source/Runtime/Slate/Public/Widgets/Notifications/SErrorText.h
--- a/Engine/Source/Runtime/Slate/Public/Widgets/Notifications/SErrorText.h
+++ b/Engine/Source/Runtime/Slate/Public/Widgets/Notifications/SErrorText.h
@@ -17,7 +17,7 @@
 /**
  * Interface for error reporting widgets.
  */
-class SLATE_API IErrorReportingWidget
+class IErrorReportingWidget
 {
 public:
 
@@ -31,7 +31,7 @@ public:
 /**
  * Implements a widget that displays an error text message.
  */
-class SLATE_API SErrorText
+class SErrorText
 	: public SBorder
 	, public IErrorReportingWidget
 {
@@ -51,16 +51,16 @@ public:
 		
 	SLATE_END_ARGS()
 
-	void Construct(const FArguments& InArgs);
+	SLATE_API void Construct(const FArguments& InArgs);
 
 public:
 
 	// IErrorReportingWidget interface
 
-	virtual void SetError( const FText& InErrorText ) override;
-	virtual void SetError( const FString& InErrorText ) override;
-	virtual bool HasError() const override;
-	virtual TSharedRef<SWidget> AsWidget() override;
+	SLATE_API virtual void SetError( const FText& InErrorText ) override;
+	SLATE_API virtual void SetError( const FString& InErrorText ) override;
+	SLATE_API virtual bool HasError() const override;
+	SLATE_API virtual TSharedRef<SWidget> AsWidget() override;
 
 private:
 
diff --git a/Engine/Source/Runtime/Slate/Public/Widgets/Notifications/SNotificationList.h b/Engine/Source/Runtime/Slate/Public/Widgets/Notifications/SNotificationList.h
--- a/Engine/Source/Runtime/Slate/Public/Widgets/Notifications/SNotificationList.h
+++ b/Engine/Source/Runtime/Slate/Public/Widgets/Notifications/SNotificationList.h
@@ -20,7 +20,7 @@ class SWindow;
  * Interface class for an item in the event message list.
  * Real implementation is found in SEventMessageItemImpl
  */
-class SLATE_API SNotificationItem
+class SNotificationItem
 	: public SCompoundWidget
 {
 public:
@@ -224,7 +224,7 @@ struct FNotificationInfo
 /**
  * A list of non-intrusive messages about the status of currently active work.
  */
-class SLATE_API SNotificationList
+class SNotificationList
 	: public SCompoundWidget
 {
 	friend class SNotificationExtendable;
@@ -240,14 +240,14 @@ public:
 	 *
 	 * @param InArgs    Declaration from which to construct the widget.
 	 */
-	void Construct( const FArguments& InArgs );
+	SLATE_API void Construct( const FArguments& InArgs );
 
 	/**
 	 * Adds a floating notification.
 	 *
 	 * @param Info 		Contains various settings used to initialize the notification.
 	 */
-	virtual TSharedRef<SNotificationItem> AddNotification(const FNotificationInfo& Info);
+	SLATE_API virtual TSharedRef<SNotificationItem> AddNotification(const FNotificationInfo& Info);
 
 protected:
 
@@ -256,7 +256,7 @@ protected:
 	 *
 	 * @param NotificationItem The item which finished fading out.
 	 */
-	virtual void NotificationItemFadedOut (const TSharedRef<SNotificationItem>& NotificationItem);
+	SLATE_API virtual void NotificationItemFadedOut (const TSharedRef<SNotificationItem>& NotificationItem);
 
 protected:
 
diff --git a/Engine/Source/Runtime/Slate/Public/Widgets/Notifications/SPopUpErrorText.h b/Engine/Source/Runtime/Slate/Public/Widgets/Notifications/SPopUpErrorText.h
--- a/Engine/Source/Runtime/Slate/Public/Widgets/Notifications/SPopUpErrorText.h
+++ b/Engine/Source/Runtime/Slate/Public/Widgets/Notifications/SPopUpErrorText.h
@@ -9,7 +9,7 @@
 #include "Widgets/Notifications/SErrorText.h"
 #include "Widgets/Input/SComboButton.h"
 
-class SLATE_API SPopupErrorText : public SComboButton, public IErrorReportingWidget
+class SPopupErrorText : public SComboButton, public IErrorReportingWidget
 {
 public:
 	SLATE_BEGIN_ARGS(SPopupErrorText)
@@ -21,16 +21,16 @@ public:
  		SLATE_ATTRIBUTE( FSlateFontInfo, Font )
 	SLATE_END_ARGS()
 
-	virtual void Construct( const FArguments& InArgs );
+	SLATE_API virtual void Construct( const FArguments& InArgs );
 
 	// IErrorReportingWidget interface
 
-	virtual void SetError( const FText& InErrorText ) override;
-	virtual void SetError( const FString& InErrorText ) override;
+	SLATE_API virtual void SetError( const FText& InErrorText ) override;
+	SLATE_API virtual void SetError( const FString& InErrorText ) override;
 
-	virtual bool HasError() const override;
+	SLATE_API virtual bool HasError() const override;
 
-	virtual TSharedRef<SWidget> AsWidget() override;
+	SLATE_API virtual TSharedRef<SWidget> AsWidget() override;
 
 	// IErrorReportingWidget interface
 
diff --git a/Engine/Source/Runtime/Slate/Public/Widgets/Notifications/SProgressBar.h b/Engine/Source/Runtime/Slate/Public/Widgets/Notifications/SProgressBar.h
--- a/Engine/Source/Runtime/Slate/Public/Widgets/Notifications/SProgressBar.h
+++ b/Engine/Source/Runtime/Slate/Public/Widgets/Notifications/SProgressBar.h
@@ -58,7 +58,7 @@ namespace EProgressBarFillStyle
 }
 
 /** A progress bar widget.*/
-class SLATE_API SProgressBar : public SLeafWidget
+class SProgressBar : public SLeafWidget
 {
 
 public:
@@ -114,38 +114,38 @@ public:
 	 * 
 	 * @param InArgs   A declaration from which to construct the widget
 	 */
-	void Construct(const FArguments& InArgs);
+	SLATE_API void Construct(const FArguments& InArgs);
 
-	virtual int32 OnPaint( const FPaintArgs& Args, const FGeometry& AllottedGeometry, const FSlateRect& MyCullingRect, FSlateWindowElementList& OutDrawElements, int32 LayerId, const FWidgetStyle& InWidgetStyle, bool bParentEnabled ) const override;
-	virtual FVector2D ComputeDesiredSize(float) const override;
-	virtual bool ComputeVolatility() const override;
+	SLATE_API virtual int32 OnPaint( const FPaintArgs& Args, const FGeometry& AllottedGeometry, const FSlateRect& MyCullingRect, FSlateWindowElementList& OutDrawElements, int32 LayerId, const FWidgetStyle& InWidgetStyle, bool bParentEnabled ) const override;
+	SLATE_API virtual FVector2D ComputeDesiredSize(float) const override;
+	SLATE_API virtual bool ComputeVolatility() const override;
 
 	/** See attribute Percent */
-	void SetPercent(TAttribute< TOptional<float> > InPercent);
+	SLATE_API void SetPercent(TAttribute< TOptional<float> > InPercent);
 	
 	/** See attribute Style */
-	void SetStyle(const FProgressBarStyle* InStyle);
+	SLATE_API void SetStyle(const FProgressBarStyle* InStyle);
 	
 	/** See attribute BarFillType */
-	void SetBarFillType(EProgressBarFillType::Type InBarFillType);
+	SLATE_API void SetBarFillType(EProgressBarFillType::Type InBarFillType);
 
 	/** See attribute BarFillStyle */
-	void SetBarFillStyle(EProgressBarFillStyle::Type InBarFillStyle);
+	SLATE_API void SetBarFillStyle(EProgressBarFillStyle::Type InBarFillStyle);
 	
 	/** See attribute SetFillColorAndOpacity */
-	void SetFillColorAndOpacity(TAttribute< FSlateColor > InFillColorAndOpacity);
+	SLATE_API void SetFillColorAndOpacity(TAttribute< FSlateColor > InFillColorAndOpacity);
 	
 	/** See attribute BorderPadding */
-	void SetBorderPadding(TAttribute< FVector2D > InBorderPadding);
+	SLATE_API void SetBorderPadding(TAttribute< FVector2D > InBorderPadding);
 	
 	/** See attribute BackgroundImage */
-	void SetBackgroundImage(const FSlateBrush* InBackgroundImage);
+	SLATE_API void SetBackgroundImage(const FSlateBrush* InBackgroundImage);
 	
 	/** See attribute FillImage */
-	void SetFillImage(const FSlateBrush* InFillImage);
+	SLATE_API void SetFillImage(const FSlateBrush* InFillImage);
 	
 	/** See attribute MarqueeImage */
-	void SetMarqueeImage(const FSlateBrush* InMarqueeImage);
+	SLATE_API void SetMarqueeImage(const FSlateBrush* InMarqueeImage);
 	
 private:
 
diff --git a/Engine/Source/Runtime/Slate/Public/Widgets/SCanvas.h b/Engine/Source/Runtime/Slate/Public/Widgets/SCanvas.h
--- a/Engine/Source/Runtime/Slate/Public/Widgets/SCanvas.h
+++ b/Engine/Source/Runtime/Slate/Public/Widgets/SCanvas.h
@@ -18,9 +18,9 @@ class FSlateWindowElementList;
 /**
  * Canvas is a layout widget that allows you to arbitrary position and size child widgets in a relative coordinate space
  */
-class SLATE_API SCanvas : public SPanel
+class SCanvas : public SPanel
 {
-	SLATE_DECLARE_WIDGET(SCanvas, SPanel)
+	SLATE_DECLARE_WIDGET_API(SCanvas, SPanel, SLATE_API)
 
 public:
 
@@ -50,7 +50,7 @@ public:
 	 * 
 	 *  Note: FILL is NOT supported.
 	 */
-	class SLATE_API FSlot : public TWidgetSlotWithAttributeSupport<FSlot>, public TAlignmentWidgetSlotMixin<FSlot>
+	class FSlot : public TWidgetSlotWithAttributeSupport<FSlot>, public TAlignmentWidgetSlotMixin<FSlot>
 	{
 	public:
 		SLATE_SLOT_BEGIN_ARGS_OneMixin(FSlot, TSlotBase<FSlot>, TAlignmentWidgetSlotMixin<FSlot>)
@@ -93,8 +93,8 @@ public:
 			, Size(*this, FVector2D(1.0f, 1.0f))
 		{ }
 
-		void Construct(const FChildren& SlotOwner, FSlotArguments&& InArg);
-		static void RegisterAttributes(FSlateWidgetSlotAttributeInitializer& AttributeInitializer);
+		SLATE_API void Construct(const FChildren& SlotOwner, FSlotArguments&& InArg);
+		static SLATE_API void RegisterAttributes(FSlateWidgetSlotAttributeInitializer& AttributeInitializer);
 
 	private:
 		/** Position */
@@ -113,16 +113,16 @@ public:
 
 	SLATE_END_ARGS()
 
-	SCanvas();
+	SLATE_API SCanvas();
 
 	/**
 	 * Construct this widget
 	 *
 	 * @param	InArgs	The declaration data for this widget
 	 */
-	void Construct( const FArguments& InArgs );
+	SLATE_API void Construct( const FArguments& InArgs );
 
-	static FSlot::FSlotArguments Slot();
+	static SLATE_API FSlot::FSlotArguments Slot();
 
 	using FScopedWidgetSlotArguments = TPanelChildren<FSlot>::FScopedWidgetSlotArguments;
 	/**
@@ -130,30 +130,30 @@ public:
 	 *
 	 * @return The added slot.
 	 */
-	FScopedWidgetSlotArguments AddSlot();
+	SLATE_API FScopedWidgetSlotArguments AddSlot();
 
 	/**
 	 * Removes a particular content slot.
 	 *
 	 * @param SlotWidget The widget in the slot to remove.
 	 */
-	int32 RemoveSlot( const TSharedRef<SWidget>& SlotWidget );
+	SLATE_API int32 RemoveSlot( const TSharedRef<SWidget>& SlotWidget );
 
 	/**
 	 * Removes all slots from the panel.
 	 */
-	void ClearChildren();
+	SLATE_API void ClearChildren();
 
 public:
 
 	//~ SWidget overrides
-	virtual void OnArrangeChildren( const FGeometry& AllottedGeometry, FArrangedChildren& ArrangedChildren ) const override;
-	virtual int32 OnPaint( const FPaintArgs& Args, const FGeometry& AllottedGeometry, const FSlateRect& MyCullingRect, FSlateWindowElementList& OutDrawElements, int32 LayerId, const FWidgetStyle& InWidgetStyle, bool bParentEnabled ) const override;
-	virtual FChildren* GetChildren() override;
+	SLATE_API virtual void OnArrangeChildren( const FGeometry& AllottedGeometry, FArrangedChildren& ArrangedChildren ) const override;
+	SLATE_API virtual int32 OnPaint( const FPaintArgs& Args, const FGeometry& AllottedGeometry, const FSlateRect& MyCullingRect, FSlateWindowElementList& OutDrawElements, int32 LayerId, const FWidgetStyle& InWidgetStyle, bool bParentEnabled ) const override;
+	SLATE_API virtual FChildren* GetChildren() override;
 
 protected:
 	//~ Begin SWidget overrides.
-	virtual FVector2D ComputeDesiredSize(float) const override;
+	SLATE_API virtual FVector2D ComputeDesiredSize(float) const override;
 	//~ End SWidget overrides.
 
 protected:
diff --git a/Engine/Source/Runtime/Slate/Public/Widgets/SInvalidationPanel.h b/Engine/Source/Runtime/Slate/Public/Widgets/SInvalidationPanel.h
--- a/Engine/Source/Runtime/Slate/Public/Widgets/SInvalidationPanel.h
+++ b/Engine/Source/Runtime/Slate/Public/Widgets/SInvalidationPanel.h
@@ -18,7 +18,7 @@ class FSlateRenderDataHandle;
 class FSlateWindowElementList;
 class SWindow;
 
-class SLATE_API SInvalidationPanel : public SCompoundWidget, public FSlateInvalidationRoot
+class SInvalidationPanel : public SCompoundWidget, public FSlateInvalidationRoot
 {
 public:
 	SLATE_BEGIN_ARGS( SInvalidationPanel )
@@ -31,35 +31,35 @@ public:
 #endif
 	SLATE_END_ARGS()
 
-	SInvalidationPanel();
-	~SInvalidationPanel();
+	SLATE_API SInvalidationPanel();
+	SLATE_API ~SInvalidationPanel();
 
 #if WITH_SLATE_DEBUGGING
-	static bool AreInvalidationPanelsEnabled();
-	static void EnableInvalidationPanels(bool bEnable);
+	static SLATE_API bool AreInvalidationPanelsEnabled();
+	static SLATE_API void EnableInvalidationPanels(bool bEnable);
 #endif
-	void Construct( const FArguments& InArgs );
+	SLATE_API void Construct( const FArguments& InArgs );
 
-	bool GetCanCache() const;
+	SLATE_API bool GetCanCache() const;
 
-	void SetCanCache(bool InCanCache);
+	SLATE_API void SetCanCache(bool InCanCache);
 
 	//~ SWidget overrides
-	virtual int32 OnPaint( const FPaintArgs& Args, const FGeometry& AllottedGeometry, const FSlateRect& MyCullingRect, FSlateWindowElementList& OutDrawElements, int32 LayerId, const FWidgetStyle& InWidgetStyle, bool bParentEnabled ) const override;
-	virtual FChildren* GetChildren() override;
+	SLATE_API virtual int32 OnPaint( const FPaintArgs& Args, const FGeometry& AllottedGeometry, const FSlateRect& MyCullingRect, FSlateWindowElementList& OutDrawElements, int32 LayerId, const FWidgetStyle& InWidgetStyle, bool bParentEnabled ) const override;
+	SLATE_API virtual FChildren* GetChildren() override;
 #if WITH_SLATE_DEBUGGING
-	virtual FChildren* Debug_GetChildrenForReflector() override;
+	SLATE_API virtual FChildren* Debug_GetChildrenForReflector() override;
 #endif
 	//~ End SWidget
 
-	void SetContent(const TSharedRef< SWidget >& InContent);
+	SLATE_API void SetContent(const TSharedRef< SWidget >& InContent);
 
 protected:
-	virtual bool CustomPrepass(float LayoutScaleMultiplier) override;
-	virtual bool Advanced_IsInvalidationRoot() const override;
-	virtual const FSlateInvalidationRoot* Advanced_AsInvalidationRoot() const override;
-	virtual TSharedRef<SWidget> GetRootWidget() override;
-	virtual int32 PaintSlowPath(const FSlateInvalidationContext& Context) override;
+	SLATE_API virtual bool CustomPrepass(float LayoutScaleMultiplier) override;
+	SLATE_API virtual bool Advanced_IsInvalidationRoot() const override;
+	SLATE_API virtual const FSlateInvalidationRoot* Advanced_AsInvalidationRoot() const override;
+	SLATE_API virtual TSharedRef<SWidget> GetRootWidget() override;
+	SLATE_API virtual int32 PaintSlowPath(const FSlateInvalidationContext& Context) override;
 
 private:
 	void OnGlobalInvalidationToggled(bool bGlobalInvalidationEnabled);
diff --git a/Engine/Source/Runtime/Slate/Public/Widgets/SToolTip.h b/Engine/Source/Runtime/Slate/Public/Widgets/SToolTip.h
--- a/Engine/Source/Runtime/Slate/Public/Widgets/SToolTip.h
+++ b/Engine/Source/Runtime/Slate/Public/Widgets/SToolTip.h
@@ -16,7 +16,7 @@
 /**
  * Slate tool tip widget
  */
-class SLATE_API SToolTip
+class SToolTip
 	: public SCompoundWidget
 	, public IToolTip
 {
@@ -61,7 +61,7 @@ public:
 	 *
 	 * @param	InArgs	The declaration data for this widget
 	 */
-	void Construct( const FArguments& InArgs );
+	SLATE_API void Construct( const FArguments& InArgs );
 
 public:
 
@@ -77,14 +77,14 @@ public:
 		return ToolTipContent.ToSharedRef();
 	}
 
-	virtual void SetContentWidget(const TSharedRef<SWidget>& InContentWidget) override;
-	virtual void ResetContentWidget () override;
+	SLATE_API virtual void SetContentWidget(const TSharedRef<SWidget>& InContentWidget) override;
+	SLATE_API virtual void ResetContentWidget () override;
 
-	virtual bool IsEmpty( ) const override;
-	virtual bool IsInteractive( ) const override;
+	SLATE_API virtual bool IsEmpty( ) const override;
+	SLATE_API virtual bool IsInteractive( ) const override;
 	virtual void OnOpening() override { }
 	virtual void OnClosed() override { }
-	virtual void OnSetInteractiveWindowLocation(FVector2D& InOutDesiredLocation) const override;
+	SLATE_API virtual void OnSetInteractiveWindowLocation(FVector2D& InOutDesiredLocation) const override;
 
 	virtual const FText& GetTextTooltip() const
 	{
@@ -93,7 +93,7 @@ public:
 
 public:
 
-	static float GetToolTipWrapWidth();
+	static SLATE_API float GetToolTipWrapWidth();
 
 private:
 
diff --git a/Engine/Source/Runtime/Slate/Public/Widgets/SViewport.h b/Engine/Source/Runtime/Slate/Public/Widgets/SViewport.h
--- a/Engine/Source/Runtime/Slate/Public/Widgets/SViewport.h
+++ b/Engine/Source/Runtime/Slate/Public/Widgets/SViewport.h
@@ -19,7 +19,7 @@ class FPaintArgs;
 class FSlateWindowElementList;
 class ICustomHitTestPath;
 
-class SLATE_API SViewport
+class SViewport
 	: public SCompoundWidget
 {
 public:
@@ -83,14 +83,14 @@ public:
 	SLATE_END_ARGS()
 
 	/** Default constructor. */
-	SViewport();
+	SLATE_API SViewport();
 
 	/**
 	 * Construct the widget.
 	 *
 	 * @param InArgs  Declaration from which to construct the widget.
 	 */
-	void Construct(const FArguments& InArgs);
+	SLATE_API void Construct(const FArguments& InArgs);
 
 public:
 
@@ -112,7 +112,7 @@ public:
 	 *
 	 * @param InViewportInterface The interface to use
 	 */
-	void SetViewportInterface( TSharedRef<ISlateViewport> InViewportInterface );
+	SLATE_API void SetViewportInterface( TSharedRef<ISlateViewport> InViewportInterface );
 
 	/**
 	 * Sets the interface to be used by this viewport for rendering and I/O
@@ -129,11 +129,11 @@ public:
 	 *
 	 * @param InContent	The new content (can be null)
 	 */
-	void SetContent( TSharedPtr<SWidget> InContent );
+	SLATE_API void SetContent( TSharedPtr<SWidget> InContent );
 
-	void SetCustomHitTestPath( TSharedPtr<ICustomHitTestPath> CustomHitTestPath );
+	SLATE_API void SetCustomHitTestPath( TSharedPtr<ICustomHitTestPath> CustomHitTestPath );
 
-	TSharedPtr<ICustomHitTestPath> GetCustomHitTestPath();
+	SLATE_API TSharedPtr<ICustomHitTestPath> GetCustomHitTestPath();
 
 	const TSharedPtr<SWidget> GetContent() const { return ChildSlot.GetWidget(); }
 
@@ -142,17 +142,17 @@ public:
 	 *
 	 * @param InWindowBeingClosed	The window that is about to be closed
 	 */
-	void OnWindowClosed( const TSharedRef<SWindow>& InWindowBeingClosed );
+	SLATE_API void OnWindowClosed( const TSharedRef<SWindow>& InWindowBeingClosed );
 
 	/**
 	 * A delegate called when the viewports top level window is activated
 	 */
-	FReply OnViewportActivated(const FWindowActivateEvent& InActivateEvent);
+	SLATE_API FReply OnViewportActivated(const FWindowActivateEvent& InActivateEvent);
 
 	/**
 	 * A delegate called when the viewports top level window is deactivated
 	 */
-	void OnViewportDeactivated(const FWindowActivateEvent& InActivateEvent);
+	SLATE_API void OnViewportDeactivated(const FWindowActivateEvent& InActivateEvent);
 
 	/** @return Whether or not this viewport renders directly to the backbuffer */
 	bool ShouldRenderDirectly() const { return bRenderDirectlyToWindow; }
@@ -165,7 +165,7 @@ public:
 	 * 
 	 * @param	bInRenderDirectlyToWindow	Whether we should be able to render to the back buffer
 	 */
-	void SetRenderDirectlyToWindow( const bool bInRenderDirectlyToWindow );
+	SLATE_API void SetRenderDirectlyToWindow( const bool bInRenderDirectlyToWindow );
 
 	/**
 	 * If true, the viewport's texture alpha is ignored when performing blending.  In this case only the viewport tint opacity is used
@@ -173,7 +173,7 @@ public:
 	 * 
 	 * @param bIgnoreTextureAlpha If texture alpha should be ignored when blending.
 	 */
-	void SetIgnoreTextureAlpha(const bool bInIgnoreTextureAlpha);
+	SLATE_API void SetIgnoreTextureAlpha(const bool bInIgnoreTextureAlpha);
 
 	/** @return Whether or not to ignore texture alpha when blending */
 	bool GetIgnoreTextureAlpha(void) const
@@ -196,42 +196,42 @@ public:
 	 * While active, a persistent Active Timer is registered and a Slate tick/paint pass is guaranteed every frame.
 	 * @param bActive Whether to set the viewport as active
 	 */
-	void SetActive(bool bActive);
+	SLATE_API void SetActive(bool bActive);
 
 public:
 
 	// SWidget interface
 
-	virtual int32 OnPaint( const FPaintArgs& Args, const FGeometry& AllottedGeometry, const FSlateRect& MyCullingRect, FSlateWindowElementList& OutDrawElements, int32 LayerId, const FWidgetStyle& InWidgetStyle, bool bParentEnabled ) const override;
-	virtual void Tick( const FGeometry& AllottedGeometry, const double InCurrentTime, const float InDeltaTime ) override;
-	virtual FReply OnTouchStarted( const FGeometry& MyGeometry, const FPointerEvent& InTouchEvent ) override;
-	virtual FReply OnTouchMoved( const FGeometry& MyGeometry, const FPointerEvent& InTouchEvent ) override;
-	virtual FReply OnTouchEnded( const FGeometry& MyGeometry, const FPointerEvent& InTouchEvent ) override;
-	virtual FReply OnTouchForceChanged(const FGeometry& MyGeometry, const FPointerEvent& TouchEvent) override;
-	virtual FReply OnTouchFirstMove(const FGeometry& MyGeometry, const FPointerEvent& TouchEvent) override;
-	virtual FReply OnTouchGesture( const FGeometry& MyGeometry, const FPointerEvent& GestureEvent ) override;
-	virtual FCursorReply OnCursorQuery( const FGeometry& MyGeometry, const FPointerEvent& CursorEvent ) const override;
-	virtual TOptional<TSharedRef<SWidget>> OnMapCursor(const FCursorReply& CursorReply) const override;
-	virtual FReply OnMouseButtonDown( const FGeometry& MyGeometry, const FPointerEvent& MouseEvent ) override;
-	virtual FReply OnMouseButtonUp( const FGeometry& MyGeometry, const FPointerEvent& MouseEvent ) override;
-	virtual void OnMouseEnter( const FGeometry& MyGeometry, const FPointerEvent& MouseEvent ) override;
-	virtual void OnMouseLeave( const FPointerEvent& MouseEvent ) override;
-	virtual FReply OnMouseMove( const FGeometry& MyGeometry, const FPointerEvent& MouseEvent ) override;
-	virtual FReply OnMouseWheel( const FGeometry& MyGeometry, const FPointerEvent& MouseEvent ) override;
-	virtual FReply OnMouseButtonDoubleClick( const FGeometry& MyGeometry, const FPointerEvent& MouseEvent ) override;
-	virtual FReply OnKeyDown( const FGeometry& MyGeometry, const FKeyEvent& KeyEvent ) override;
-	virtual FReply OnKeyUp( const FGeometry& MyGeometry, const FKeyEvent& KeyEvent ) override;
-	virtual FReply OnAnalogValueChanged( const FGeometry& MyGeometry, const FAnalogInputEvent& InAnalogInputEvent ) override;
-	virtual FReply OnKeyChar( const FGeometry& MyGeometry, const FCharacterEvent& CharacterEvent ) override;
-	virtual FReply OnFocusReceived( const FGeometry& MyGeometry, const FFocusEvent& InFocusEvent ) override;
-	virtual void OnFocusLost( const FFocusEvent& InFocusEvent ) override;
-	virtual FReply OnMotionDetected( const FGeometry& MyGeometry, const FMotionEvent& InMotionEvent ) override;
-	virtual TOptional<bool> OnQueryShowFocus( const EFocusCause InFocusCause ) const override;
-	virtual FPopupMethodReply OnQueryPopupMethod() const override;
-	virtual void OnFinishedPointerInput() override;
-	virtual void OnArrangeChildren( const FGeometry& AllottedGeometry, FArrangedChildren& ArrangedChildren ) const override;
-	virtual TOptional<FVirtualPointerPosition> TranslateMouseCoordinateForCustomHitTestChild(const SWidget& ChildWidget, const FGeometry& MyGeometry, const FVector2D ScreenSpaceMouseCoordinate, const FVector2D LastScreenSpaceMouseCoordinate) const override;
-	virtual FNavigationReply OnNavigation( const FGeometry& MyGeometry, const FNavigationEvent& InNavigationEvent ) override;
+	SLATE_API virtual int32 OnPaint( const FPaintArgs& Args, const FGeometry& AllottedGeometry, const FSlateRect& MyCullingRect, FSlateWindowElementList& OutDrawElements, int32 LayerId, const FWidgetStyle& InWidgetStyle, bool bParentEnabled ) const override;
+	SLATE_API virtual void Tick( const FGeometry& AllottedGeometry, const double InCurrentTime, const float InDeltaTime ) override;
+	SLATE_API virtual FReply OnTouchStarted( const FGeometry& MyGeometry, const FPointerEvent& InTouchEvent ) override;
+	SLATE_API virtual FReply OnTouchMoved( const FGeometry& MyGeometry, const FPointerEvent& InTouchEvent ) override;
+	SLATE_API virtual FReply OnTouchEnded( const FGeometry& MyGeometry, const FPointerEvent& InTouchEvent ) override;
+	SLATE_API virtual FReply OnTouchForceChanged(const FGeometry& MyGeometry, const FPointerEvent& TouchEvent) override;
+	SLATE_API virtual FReply OnTouchFirstMove(const FGeometry& MyGeometry, const FPointerEvent& TouchEvent) override;
+	SLATE_API virtual FReply OnTouchGesture( const FGeometry& MyGeometry, const FPointerEvent& GestureEvent ) override;
+	SLATE_API virtual FCursorReply OnCursorQuery( const FGeometry& MyGeometry, const FPointerEvent& CursorEvent ) const override;
+	SLATE_API virtual TOptional<TSharedRef<SWidget>> OnMapCursor(const FCursorReply& CursorReply) const override;
+	SLATE_API virtual FReply OnMouseButtonDown( const FGeometry& MyGeometry, const FPointerEvent& MouseEvent ) override;
+	SLATE_API virtual FReply OnMouseButtonUp( const FGeometry& MyGeometry, const FPointerEvent& MouseEvent ) override;
+	SLATE_API virtual void OnMouseEnter( const FGeometry& MyGeometry, const FPointerEvent& MouseEvent ) override;
+	SLATE_API virtual void OnMouseLeave( const FPointerEvent& MouseEvent ) override;
+	SLATE_API virtual FReply OnMouseMove( const FGeometry& MyGeometry, const FPointerEvent& MouseEvent ) override;
+	SLATE_API virtual FReply OnMouseWheel( const FGeometry& MyGeometry, const FPointerEvent& MouseEvent ) override;
+	SLATE_API virtual FReply OnMouseButtonDoubleClick( const FGeometry& MyGeometry, const FPointerEvent& MouseEvent ) override;
+	SLATE_API virtual FReply OnKeyDown( const FGeometry& MyGeometry, const FKeyEvent& KeyEvent ) override;
+	SLATE_API virtual FReply OnKeyUp( const FGeometry& MyGeometry, const FKeyEvent& KeyEvent ) override;
+	SLATE_API virtual FReply OnAnalogValueChanged( const FGeometry& MyGeometry, const FAnalogInputEvent& InAnalogInputEvent ) override;
+	SLATE_API virtual FReply OnKeyChar( const FGeometry& MyGeometry, const FCharacterEvent& CharacterEvent ) override;
+	SLATE_API virtual FReply OnFocusReceived( const FGeometry& MyGeometry, const FFocusEvent& InFocusEvent ) override;
+	SLATE_API virtual void OnFocusLost( const FFocusEvent& InFocusEvent ) override;
+	SLATE_API virtual FReply OnMotionDetected( const FGeometry& MyGeometry, const FMotionEvent& InMotionEvent ) override;
+	SLATE_API virtual TOptional<bool> OnQueryShowFocus( const EFocusCause InFocusCause ) const override;
+	SLATE_API virtual FPopupMethodReply OnQueryPopupMethod() const override;
+	SLATE_API virtual void OnFinishedPointerInput() override;
+	SLATE_API virtual void OnArrangeChildren( const FGeometry& AllottedGeometry, FArrangedChildren& ArrangedChildren ) const override;
+	SLATE_API virtual TOptional<FVirtualPointerPosition> TranslateMouseCoordinateForCustomHitTestChild(const SWidget& ChildWidget, const FGeometry& MyGeometry, const FVector2D ScreenSpaceMouseCoordinate, const FVector2D LastScreenSpaceMouseCoordinate) const override;
+	SLATE_API virtual FNavigationReply OnNavigation( const FGeometry& MyGeometry, const FNavigationEvent& InNavigationEvent ) override;
 
 private:
 	// Viewports shouldn't show focus
@@ -242,7 +242,7 @@ private:
 
 protected:
 	/** Empty active timer meant to ensure a tick/paint pass while this viewport is active */
-	EActiveTimerReturnType EnsureTick(double InCurrentTime, float InDeltaTime);
+	SLATE_API EActiveTimerReturnType EnsureTick(double InCurrentTime, float InDeltaTime);
 
 	/** Interface to the rendering and I/O implementation of the viewport. */
 	TWeakPtr<ISlateViewport> ViewportInterface;
diff --git a/Engine/Source/Runtime/Slate/Public/Widgets/SVirtualWindow.h b/Engine/Source/Runtime/Slate/Public/Widgets/SVirtualWindow.h
--- a/Engine/Source/Runtime/Slate/Public/Widgets/SVirtualWindow.h
+++ b/Engine/Source/Runtime/Slate/Public/Widgets/SVirtualWindow.h
@@ -18,7 +18,7 @@ class ISlate3DRenderer;
 class STooltipPresenter;
 class UTextureRenderTarget2D;
 
-class SLATE_API SVirtualWindow : public SWindow
+class SVirtualWindow : public SWindow
 {
 	SLATE_BEGIN_ARGS(SVirtualWindow)
 		: _Size(FVector2D(100, 100))
@@ -29,7 +29,7 @@ class SLATE_API SVirtualWindow : public SWindow
 		SLATE_END_ARGS()
 
 public:
-	void Construct(const FArguments& InArgs);
+	SLATE_API void Construct(const FArguments& InArgs);
 
 	/**
 	* We allow users to control the resolve policy for deferred content.  When a virtual window is used in a retainer, we don't want it trying to resolve
@@ -37,17 +37,17 @@ public:
 	* a case like the WidgetComponent, we always want to resolve the deferred content because there won't be another opportunity since the 3D widget has
 	* no owner window that can draw it.
 	*/
-	void SetShouldResolveDeferred(bool bResolve);
+	SLATE_API void SetShouldResolveDeferred(bool bResolve);
 
 	/**  */
-	void SetIsFocusable(bool bFocusable);
-
-	virtual FPopupMethodReply OnQueryPopupMethod() const override;
-	virtual bool OnVisualizeTooltip(const TSharedPtr<SWidget>& TooltipContent) override;
-	virtual int32 OnPaint(const FPaintArgs& Args, const FGeometry& AllottedGeometry, const FSlateRect& MyCullingRect, FSlateWindowElementList& OutDrawElements, int32 LayerId, const FWidgetStyle& InWidgetStyle, bool bParentEnabled) const override;
-	virtual void OnArrangeChildren(const FGeometry& AllottedGeometry, FArrangedChildren& ArrangedChildren) const override;
-	virtual bool SupportsKeyboardFocus() const override;
-	virtual FVector2D ComputeDesiredSize(float LayoutScaleMultiplier) const override;
+	SLATE_API void SetIsFocusable(bool bFocusable);
+
+	SLATE_API virtual FPopupMethodReply OnQueryPopupMethod() const override;
+	SLATE_API virtual bool OnVisualizeTooltip(const TSharedPtr<SWidget>& TooltipContent) override;
+	SLATE_API virtual int32 OnPaint(const FPaintArgs& Args, const FGeometry& AllottedGeometry, const FSlateRect& MyCullingRect, FSlateWindowElementList& OutDrawElements, int32 LayerId, const FWidgetStyle& InWidgetStyle, bool bParentEnabled) const override;
+	SLATE_API virtual void OnArrangeChildren(const FGeometry& AllottedGeometry, FArrangedChildren& ArrangedChildren) const override;
+	SLATE_API virtual bool SupportsKeyboardFocus() const override;
+	SLATE_API virtual FVector2D ComputeDesiredSize(float LayoutScaleMultiplier) const override;
 
 private:
 	TSharedPtr<class STooltipPresenter> TooltipPresenter;
diff --git a/Engine/Source/Runtime/Slate/Public/Widgets/SWeakWidget.h b/Engine/Source/Runtime/Slate/Public/Widgets/SWeakWidget.h
--- a/Engine/Source/Runtime/Slate/Public/Widgets/SWeakWidget.h
+++ b/Engine/Source/Runtime/Slate/Public/Widgets/SWeakWidget.h
@@ -20,7 +20,7 @@ class FSlateWindowElementList;
  *      by a floating window. That window cannot own the tooltip
  *      and must therefore use an SWeakWidget.
  */
-class SLATE_API SWeakWidget : public SWidget
+class SWeakWidget : public SWidget
 {
 public:
 
@@ -33,26 +33,26 @@ public:
 
 public:
 
-	SWeakWidget();
+	SLATE_API SWeakWidget();
 
-	void Construct(const FArguments& InArgs);
+	SLATE_API void Construct(const FArguments& InArgs);
 
-	void SetContent(const TSharedRef<SWidget>& InWidget);
+	SLATE_API void SetContent(const TSharedRef<SWidget>& InWidget);
 
-	bool ChildWidgetIsValid() const;
+	SLATE_API bool ChildWidgetIsValid() const;
 
-	TWeakPtr<SWidget> GetChildWidget() const;
+	SLATE_API TWeakPtr<SWidget> GetChildWidget() const;
 
 public:
 
 	// SWidget interface
-	virtual void OnArrangeChildren( const FGeometry& AllottedGeometry, FArrangedChildren& ArrangedChildren ) const override;
-	virtual FChildren* GetChildren() override;
-	virtual int32 OnPaint( const FPaintArgs& Args, const FGeometry& AllottedGeometry, const FSlateRect& MyCullingRect, FSlateWindowElementList& OutDrawElements, int32 LayerId, const FWidgetStyle& InWidgetStyle, bool bParentEnabled ) const override;
+	SLATE_API virtual void OnArrangeChildren( const FGeometry& AllottedGeometry, FArrangedChildren& ArrangedChildren ) const override;
+	SLATE_API virtual FChildren* GetChildren() override;
+	SLATE_API virtual int32 OnPaint( const FPaintArgs& Args, const FGeometry& AllottedGeometry, const FSlateRect& MyCullingRect, FSlateWindowElementList& OutDrawElements, int32 LayerId, const FWidgetStyle& InWidgetStyle, bool bParentEnabled ) const override;
 
 protected:
 	// Begin SWidget overrides.
-	virtual FVector2D ComputeDesiredSize(float) const override;
+	SLATE_API virtual FVector2D ComputeDesiredSize(float) const override;
 	// End SWidget overrides.
 
 private:
diff --git a/Engine/Source/Runtime/Slate/Public/Widgets/Text/SInlineEditableTextBlock.h b/Engine/Source/Runtime/Slate/Public/Widgets/Text/SInlineEditableTextBlock.h
--- a/Engine/Source/Runtime/Slate/Public/Widgets/Text/SInlineEditableTextBlock.h
+++ b/Engine/Source/Runtime/Slate/Public/Widgets/Text/SInlineEditableTextBlock.h
@@ -28,7 +28,7 @@ DECLARE_DELEGATE_OneParam(FOnBeginTextEdit, const FText&)
 /**
  * Slate's InlineEditableTextBlock's are double selectable to go from a STextBlock to become SEditableTextBox.
  */
-class SLATE_API SInlineEditableTextBlock: public SCompoundWidget
+class SInlineEditableTextBlock: public SCompoundWidget
 {
 	SLATE_BEGIN_ARGS( SInlineEditableTextBlock )
 		: _Text()
@@ -109,7 +109,7 @@ class SLATE_API SInlineEditableTextBlock: public SCompoundWidget
 		SLATE_ARGUMENT(TOptional<ETextOverflowPolicy>, OverflowPolicy)
 	SLATE_END_ARGS()
 
-	~SInlineEditableTextBlock();
+	SLATE_API ~SInlineEditableTextBlock();
 
 	/**
 	 * Construct this widget.  Called by the SNew() Slate macro.
@@ -117,54 +117,54 @@ class SLATE_API SInlineEditableTextBlock: public SCompoundWidget
 	 * @param	InArgs				Declaration used by the SNew() macro to construct this widget
 	 * @param	InViewModel			The UI logic not specific to slate
 	 */
-	void Construct( const FArguments& InArgs );
+	SLATE_API void Construct( const FArguments& InArgs );
 
 	/**
 	 * See SWidget::SupportsKeyboardFocus().
 	 *
 	 * @return  This widget can gain focus if IsSelected is not bound.
 	 */
-	virtual bool SupportsKeyboardFocus() const override;
+	SLATE_API virtual bool SupportsKeyboardFocus() const override;
 
-	virtual FReply OnMouseButtonDown( const FGeometry& MyGeometry, const FPointerEvent& MouseEvent ) override;
+	SLATE_API virtual FReply OnMouseButtonDown( const FGeometry& MyGeometry, const FPointerEvent& MouseEvent ) override;
 
-	virtual FReply OnDragOver( const FGeometry& MyGeometry, const FDragDropEvent& DragDropEvent ) override;
+	SLATE_API virtual FReply OnDragOver( const FGeometry& MyGeometry, const FDragDropEvent& DragDropEvent ) override;
 
-	virtual FReply OnMouseButtonDoubleClick( const FGeometry& InMyGeometry, const FPointerEvent& InMouseEvent ) override;
+	SLATE_API virtual FReply OnMouseButtonDoubleClick( const FGeometry& InMyGeometry, const FPointerEvent& InMouseEvent ) override;
 
-	virtual FReply OnKeyDown( const FGeometry& MyGeometry, const FKeyEvent& InKeyEvent ) override;
+	SLATE_API virtual FReply OnKeyDown( const FGeometry& MyGeometry, const FKeyEvent& InKeyEvent ) override;
 
 	//virtual void Tick( const FGeometry& AllottedGeometry, const double InCurrentTime, const float InDeltaTime ) override;
 	 
 	/** Switches the widget to editing mode */
-	void EnterEditingMode();
+	SLATE_API void EnterEditingMode();
 
 	/** Switches the widget to label mode. */
-	void ExitEditingMode();
+	SLATE_API void ExitEditingMode();
 
 	/** Checks if the widget is in edit mode */
-	bool IsInEditMode() const;
+	SLATE_API bool IsInEditMode() const;
 
-	void SetReadOnly(bool bInIsReadOnly);
+	SLATE_API void SetReadOnly(bool bInIsReadOnly);
 
-	void SetText( const TAttribute< FText >& InText );
-	void SetText( const FString& InText );
+	SLATE_API void SetText( const TAttribute< FText >& InText );
+	SLATE_API void SetText( const FString& InText );
 
 	/** Sets the wrap text at attribute.  See WrapTextAt attribute */
-	void SetWrapTextAt(const TAttribute<float>& InWrapTextAt);
+	SLATE_API void SetWrapTextAt(const TAttribute<float>& InWrapTextAt);
 
 	/** Sets the overflow policy for this text block */
-	void SetOverflowPolicy(TOptional<ETextOverflowPolicy> InOverflowPolicy);
+	SLATE_API void SetOverflowPolicy(TOptional<ETextOverflowPolicy> InOverflowPolicy);
 
 protected:
 	/** Callback for the text box's OnTextChanged event */
-	void OnTextChanged(const FText& InText);
+	SLATE_API void OnTextChanged(const FText& InText);
 
 	/** Callback when the text box is committed, switches back to label mode. */ 
-	void OnTextBoxCommitted(const FText& InText, ETextCommit::Type InCommitType);
+	SLATE_API void OnTextBoxCommitted(const FText& InText, ETextCommit::Type InCommitType);
 
 	/** Cancels the edit mode and switches back to label mode */
-	void CancelEditMode();
+	SLATE_API void CancelEditMode();
 
 protected:
 	/** The widget used when in label mode */
diff --git a/Engine/Source/Runtime/Slate/Public/Widgets/Text/SMultiLineEditableText.h b/Engine/Source/Runtime/Slate/Public/Widgets/Text/SMultiLineEditableText.h
--- a/Engine/Source/Runtime/Slate/Public/Widgets/Text/SMultiLineEditableText.h
+++ b/Engine/Source/Runtime/Slate/Public/Widgets/Text/SMultiLineEditableText.h
@@ -38,7 +38,7 @@ class ITextLayoutMarshaller;
 class FSlateEditableTextLayout;
 
 /** An editable text widget that supports multiple lines and soft word-wrapping. */
-class SLATE_API SMultiLineEditableText : public SWidget, public ISlateEditableTextWidget
+class SMultiLineEditableText : public SWidget, public ISlateEditableTextWidget
 {
 public:
 
@@ -229,35 +229,35 @@ public:
 
 	SLATE_END_ARGS()
 
-	SMultiLineEditableText();
-	~SMultiLineEditableText();
+	SLATE_API SMultiLineEditableText();
+	SLATE_API ~SMultiLineEditableText();
 
-	void Construct( const FArguments& InArgs );
+	SLATE_API void Construct( const FArguments& InArgs );
 
 	/**
 	 * Sets the text for this text block
 	 */
-	void SetText(const TAttribute< FText >& InText);
+	SLATE_API void SetText(const TAttribute< FText >& InText);
 
 	/**
 	 * Returns the text string
 	 *
 	 * @return  Text string
 	 */
-	FText GetText() const;
+	SLATE_API FText GetText() const;
 
 	/**
 	 * Returns the plain text string without richtext formatting
 	 * @return  Text string
 	 */
-	FText GetPlainText() const;
+	SLATE_API FText GetPlainText() const;
 
 	/**
 	 * Fill OutTextLine with the text line where the current cursor location is at
 	 *
 	 * @param OutTextLine   FString of the line
 	 */
-	void GetCurrentTextLine(FString& OutTextLine) const;
+	SLATE_API void GetCurrentTextLine(FString& OutTextLine) const;
 
 	/**
 	 * Fill OutTextLine with the text line at the specified index
@@ -265,178 +265,178 @@ public:
 	 * @param InLineIndex   Index of the line
 	 * @param OutTextLine   FString of the line
 	 */
-	void GetTextLine(const int32 InLineIndex, FString& OutTextLine) const;
+	SLATE_API void GetTextLine(const int32 InLineIndex, FString& OutTextLine) const;
 	
 	/**
 	 * Sets the text that appears when there is no text in the text box
 	 */
-	void SetHintText(const TAttribute< FText >& InHintText);
+	SLATE_API void SetHintText(const TAttribute< FText >& InHintText);
 
 	/** Get the text that appears when there is no text in the text box */
-	FText GetHintText() const;
+	SLATE_API FText GetHintText() const;
 
 	/** Set the text that is currently being searched for (if any) */
-	void SetSearchText(const TAttribute<FText>& InSearchText);
+	SLATE_API void SetSearchText(const TAttribute<FText>& InSearchText);
 
 	/** Get the text that is currently being searched for (if any) */
-	FText GetSearchText() const;
+	SLATE_API FText GetSearchText() const;
 
 	/** Get the index of the search result (0 if none) */
-	int32 GetSearchResultIndex() const;
+	SLATE_API int32 GetSearchResultIndex() const;
 
 	/** Get the total number of search results (0 if none) */
-	int32 GetNumSearchResults() const;
+	SLATE_API int32 GetNumSearchResults() const;
 
 	/** See attribute TextStyle */
-	void SetTextStyle(const FTextBlockStyle* InTextStyle);
+	SLATE_API void SetTextStyle(const FTextBlockStyle* InTextStyle);
 
 	/** See attribute Font */
-	void SetFont(const TAttribute< FSlateFontInfo >& InNewFont);
-	FSlateFontInfo GetFont() const;
+	SLATE_API void SetFont(const TAttribute< FSlateFontInfo >& InNewFont);
+	SLATE_API FSlateFontInfo GetFont() const;
 
 	/** See TextShapingMethod attribute */
-	void SetTextShapingMethod(const TOptional<ETextShapingMethod>& InTextShapingMethod);
+	SLATE_API void SetTextShapingMethod(const TOptional<ETextShapingMethod>& InTextShapingMethod);
 
 	/** See TextFlowDirection attribute */
-	void SetTextFlowDirection(const TOptional<ETextFlowDirection>& InTextFlowDirection);
+	SLATE_API void SetTextFlowDirection(const TOptional<ETextFlowDirection>& InTextFlowDirection);
 
 	/** See WrapTextAt attribute */
-	void SetWrapTextAt(const TAttribute<float>& InWrapTextAt);
+	SLATE_API void SetWrapTextAt(const TAttribute<float>& InWrapTextAt);
 
 	/** See AutoWrapText attribute */
-	void SetAutoWrapText(const TAttribute<bool>& InAutoWrapText);
+	SLATE_API void SetAutoWrapText(const TAttribute<bool>& InAutoWrapText);
 
 	/** Set WrappingPolicy attribute */
-	void SetWrappingPolicy(const TAttribute<ETextWrappingPolicy>& InWrappingPolicy);
+	SLATE_API void SetWrappingPolicy(const TAttribute<ETextWrappingPolicy>& InWrappingPolicy);
 
 	/** See LineHeightPercentage attribute */
-	void SetLineHeightPercentage(const TAttribute<float>& InLineHeightPercentage);
+	SLATE_API void SetLineHeightPercentage(const TAttribute<float>& InLineHeightPercentage);
 
 	/** See Margin attribute */
-	void SetMargin(const TAttribute<FMargin>& InMargin);
+	SLATE_API void SetMargin(const TAttribute<FMargin>& InMargin);
 
 	/** See Justification attribute */
-	void SetJustification(const TAttribute<ETextJustify::Type>& InJustification);
+	SLATE_API void SetJustification(const TAttribute<ETextJustify::Type>& InJustification);
 
 	/** Sets the overflow policy for this text block */
-	void SetOverflowPolicy(TOptional<ETextOverflowPolicy> InOverflowPolicy);
+	SLATE_API void SetOverflowPolicy(TOptional<ETextOverflowPolicy> InOverflowPolicy);
 
 	/** See the AllowContextMenu attribute */
-	void SetAllowContextMenu(const TAttribute< bool >& InAllowContextMenu);
+	SLATE_API void SetAllowContextMenu(const TAttribute< bool >& InAllowContextMenu);
 
 	/** Set the VirtualKeyboardDismissAction attribute */
-	void SetVirtualKeyboardDismissAction(TAttribute< EVirtualKeyboardDismissAction > InVirtualKeyboardDismissAction);
+	SLATE_API void SetVirtualKeyboardDismissAction(TAttribute< EVirtualKeyboardDismissAction > InVirtualKeyboardDismissAction);
 
 	/**
 	 * Sets whether to select word on the mouse double click
 	 *
 	 * @param  InSelectWordOnMouseDoubleClick		Select word on the mouse double click
 	 */
-	void SetSelectWordOnMouseDoubleClick(const TAttribute<bool>& InSelectWordOnMouseDoubleClick);
+	SLATE_API void SetSelectWordOnMouseDoubleClick(const TAttribute<bool>& InSelectWordOnMouseDoubleClick);
 	
 	/** Sets the ReadOnly attribute */
-	void SetIsReadOnly(const TAttribute< bool >& InIsReadOnly);
+	SLATE_API void SetIsReadOnly(const TAttribute< bool >& InIsReadOnly);
 
 	/** Get the number of Text Lines */
-	int32 GetTextLineCount();
+	SLATE_API int32 GetTextLineCount();
 
 	/**
 	 * Sets whether to select all text when the user clicks to give focus on the widget
 	 *
 	 * @param  InSelectAllTextWhenFocused	Select all text when the user clicks?
 	 */
-	void SetSelectAllTextWhenFocused(const TAttribute<bool>& InSelectAllTextWhenFocused);
+	SLATE_API void SetSelectAllTextWhenFocused(const TAttribute<bool>& InSelectAllTextWhenFocused);
 
 	/**
 	 * Sets whether to clear text selection when focus is lost
 	 *
 	 * @param  InClearTextSelectionOnFocusLoss	Clear text selection when focus is lost?
 	 */
-	void SetClearTextSelectionOnFocusLoss(const TAttribute<bool>& InClearTextSelectionOnFocusLoss);
+	SLATE_API void SetClearTextSelectionOnFocusLoss(const TAttribute<bool>& InClearTextSelectionOnFocusLoss);
 
 	/**
 	 * Sets whether to allow the user to back out of changes when they press the escape key
 	 *
 	 * @param  InRevertTextOnEscape			Allow the user to back out of changes?
 	 */
-	void SetRevertTextOnEscape(const TAttribute<bool>& InRevertTextOnEscape);
+	SLATE_API void SetRevertTextOnEscape(const TAttribute<bool>& InRevertTextOnEscape);
 
 	/**
 	 * Sets whether to clear keyboard focus when pressing enter to commit changes
 	 *
 	 * @param  InClearKeyboardFocusOnCommit		Clear keyboard focus when pressing enter?
 	 */
-	void SetClearKeyboardFocusOnCommit(const TAttribute<bool>& InClearKeyboardFocusOnCommit);
+	SLATE_API void SetClearKeyboardFocusOnCommit(const TAttribute<bool>& InClearKeyboardFocusOnCommit);
 
 	/** Query to see if any text is selected within the document */
-	bool AnyTextSelected() const;
+	SLATE_API bool AnyTextSelected() const;
 
 	/** Select all the text in the document */
-	void SelectAllText();
+	SLATE_API void SelectAllText();
 
 	/** Select a block of text */
-	void SelectText(const FTextLocation& InSelectionStart, const FTextLocation& InCursorLocation);
+	SLATE_API void SelectText(const FTextLocation& InSelectionStart, const FTextLocation& InCursorLocation);
 	
 	/** Clear the active text selection */
-	void ClearSelection();
+	SLATE_API void ClearSelection();
 
 	/** Get the currently selected text */
-	FText GetSelectedText() const;
+	SLATE_API FText GetSelectedText() const;
 
 	/** Get the current selection */
-	FTextSelection GetSelection() const;
+	SLATE_API FTextSelection GetSelection() const;
 
 	/** Delete any currently selected text */
-	void DeleteSelectedText();
+	SLATE_API void DeleteSelectedText();
 
 	/** Insert the given text at the current cursor position, correctly taking into account new line characters */
-	void InsertTextAtCursor(const FText& InText);
-	void InsertTextAtCursor(const FString& InString);
+	SLATE_API void InsertTextAtCursor(const FText& InText);
+	SLATE_API void InsertTextAtCursor(const FString& InString);
 
 	/** Insert the given run at the current cursor position */
-	void InsertRunAtCursor(TSharedRef<IRun> InRun);
+	SLATE_API void InsertRunAtCursor(TSharedRef<IRun> InRun);
 
 	/** Move the cursor to the given location in the document (will also scroll to this point) */
-	void GoTo(const FTextLocation& NewLocation);
+	SLATE_API void GoTo(const FTextLocation& NewLocation);
 
 	/** Move the cursor specified location */
-	void GoTo(const ETextLocation NewLocation);
+	SLATE_API void GoTo(const ETextLocation NewLocation);
 
 	/** Scroll to the given location in the document (without moving the cursor) */
-	void ScrollTo(const FTextLocation& NewLocation);
+	SLATE_API void ScrollTo(const FTextLocation& NewLocation);
 
 	/** Scroll to the given location in the document (without moving the cursor) */
-	void ScrollTo(const ETextLocation NewLocation);
+	SLATE_API void ScrollTo(const ETextLocation NewLocation);
 
 	/** Apply the given style to the currently selected text (or insert a new run at the current cursor position if no text is selected) */
-	void ApplyToSelection(const FRunInfo& InRunInfo, const FTextBlockStyle& InStyle);
+	SLATE_API void ApplyToSelection(const FRunInfo& InRunInfo, const FTextBlockStyle& InStyle);
 
 	/** Begin a new text search (this is called automatically when the bound search text changes) */
-	void BeginSearch(const FText& InSearchText, const ESearchCase::Type InSearchCase = ESearchCase::IgnoreCase, const bool InReverse = false);
+	SLATE_API void BeginSearch(const FText& InSearchText, const ESearchCase::Type InSearchCase = ESearchCase::IgnoreCase, const bool InReverse = false);
 
 	/** Advance the current search to the next match (does nothing if not currently searching) */
-	void AdvanceSearch(const bool InReverse = false);
+	SLATE_API void AdvanceSearch(const bool InReverse = false);
 
 	/** Get the run currently under the cursor, or null if there is no run currently under the cursor */
-	TSharedPtr<const IRun> GetRunUnderCursor() const;
+	SLATE_API TSharedPtr<const IRun> GetRunUnderCursor() const;
 
 	/** Get the runs currently that are current selected, some of which may be only partially selected */
-	TArray<TSharedRef<const IRun>> GetSelectedRuns() const;
+	SLATE_API TArray<TSharedRef<const IRun>> GetSelectedRuns() const;
 
 	/** Get the interaction position of the cursor (where to insert, delete, etc, text from/to) */
-	FTextLocation GetCursorLocation() const;
+	SLATE_API FTextLocation GetCursorLocation() const;
 
 	/** Get the character at Location */
-	TCHAR GetCharacterAt(const FTextLocation& Location) const;
+	SLATE_API TCHAR GetCharacterAt(const FTextLocation& Location) const;
 
 	/** Get the horizontal scroll bar widget */
-	TSharedPtr<const SScrollBar> GetHScrollBar() const;
+	SLATE_API TSharedPtr<const SScrollBar> GetHScrollBar() const;
 
 	/** Get the vertical scroll bar widget */
-	TSharedPtr<const SScrollBar> GetVScrollBar() const;
+	SLATE_API TSharedPtr<const SScrollBar> GetVScrollBar() const;
 
 	/** Refresh this editable text immediately, rather than wait for the usual caching mechanisms to take affect on the text Tick */
-	void Refresh();
+	SLATE_API void Refresh();
 
 	/**
 	 * Sets the OnKeyCharHandler to provide first chance handling of the OnKeyChar event
@@ -461,78 +461,78 @@ public:
 	/**
 	 *	Force a single scroll operation. 
 	 */
-	void ForceScroll(int32 UserIndex, float ScrollAxisMagnitude);
+	SLATE_API void ForceScroll(int32 UserIndex, float ScrollAxisMagnitude);
 
 protected:
 	//~ Begin SWidget Interface
-	virtual void Tick( const FGeometry& AllottedGeometry, const double InCurrentTime, const float InDeltaTime ) override;
-	virtual int32 OnPaint( const FPaintArgs& Args, const FGeometry& AllottedGeometry, const FSlateRect& MyCullingRect, FSlateWindowElementList& OutDrawElements, int32 LayerId, const FWidgetStyle& InWidgetStyle, bool bParentEnabled ) const override;
-	virtual void CacheDesiredSize(float LayoutScaleMultiplier) override;
-	virtual FVector2D ComputeDesiredSize(float LayoutScaleMultiplier) const override;
-	virtual FChildren* GetChildren() override;
-	virtual void OnArrangeChildren( const FGeometry& AllottedGeometry, FArrangedChildren& ArrangedChildren ) const override;
-	virtual bool SupportsKeyboardFocus() const override;
-	virtual FReply OnKeyChar( const FGeometry& MyGeometry,const FCharacterEvent& InCharacterEvent ) override;
-	virtual FReply OnKeyDown( const FGeometry& MyGeometry, const FKeyEvent& InKeyEvent ) override;
-	virtual FReply OnKeyUp( const FGeometry& MyGeometry, const FKeyEvent& InKeyEvent ) override;
-	virtual FReply OnFocusReceived( const FGeometry& MyGeometry, const FFocusEvent& InFocusEvent ) override;
-	virtual void OnFocusLost( const FFocusEvent& InFocusEvent ) override;
-	virtual FReply OnMouseButtonDown( const FGeometry& MyGeometry, const FPointerEvent& MouseEvent ) override;
-	virtual FReply OnMouseButtonUp( const FGeometry& MyGeometry, const FPointerEvent& MouseEvent ) override;
-	virtual FReply OnMouseMove( const FGeometry& MyGeometry, const FPointerEvent& MouseEvent ) override;
-	virtual FReply OnMouseWheel( const FGeometry& MyGeometry, const FPointerEvent& MouseEvent ) override;
-	virtual FReply OnMouseButtonDoubleClick(const FGeometry& MyGeometry, const FPointerEvent& MouseEvent) override;
-	virtual FCursorReply OnCursorQuery( const FGeometry& MyGeometry, const FPointerEvent& CursorEvent ) const override;
-	virtual bool IsInteractable() const override;
-	virtual bool ComputeVolatility() const override;
+	SLATE_API virtual void Tick( const FGeometry& AllottedGeometry, const double InCurrentTime, const float InDeltaTime ) override;
+	SLATE_API virtual int32 OnPaint( const FPaintArgs& Args, const FGeometry& AllottedGeometry, const FSlateRect& MyCullingRect, FSlateWindowElementList& OutDrawElements, int32 LayerId, const FWidgetStyle& InWidgetStyle, bool bParentEnabled ) const override;
+	SLATE_API virtual void CacheDesiredSize(float LayoutScaleMultiplier) override;
+	SLATE_API virtual FVector2D ComputeDesiredSize(float LayoutScaleMultiplier) const override;
+	SLATE_API virtual FChildren* GetChildren() override;
+	SLATE_API virtual void OnArrangeChildren( const FGeometry& AllottedGeometry, FArrangedChildren& ArrangedChildren ) const override;
+	SLATE_API virtual bool SupportsKeyboardFocus() const override;
+	SLATE_API virtual FReply OnKeyChar( const FGeometry& MyGeometry,const FCharacterEvent& InCharacterEvent ) override;
+	SLATE_API virtual FReply OnKeyDown( const FGeometry& MyGeometry, const FKeyEvent& InKeyEvent ) override;
+	SLATE_API virtual FReply OnKeyUp( const FGeometry& MyGeometry, const FKeyEvent& InKeyEvent ) override;
+	SLATE_API virtual FReply OnFocusReceived( const FGeometry& MyGeometry, const FFocusEvent& InFocusEvent ) override;
+	SLATE_API virtual void OnFocusLost( const FFocusEvent& InFocusEvent ) override;
+	SLATE_API virtual FReply OnMouseButtonDown( const FGeometry& MyGeometry, const FPointerEvent& MouseEvent ) override;
+	SLATE_API virtual FReply OnMouseButtonUp( const FGeometry& MyGeometry, const FPointerEvent& MouseEvent ) override;
+	SLATE_API virtual FReply OnMouseMove( const FGeometry& MyGeometry, const FPointerEvent& MouseEvent ) override;
+	SLATE_API virtual FReply OnMouseWheel( const FGeometry& MyGeometry, const FPointerEvent& MouseEvent ) override;
+	SLATE_API virtual FReply OnMouseButtonDoubleClick(const FGeometry& MyGeometry, const FPointerEvent& MouseEvent) override;
+	SLATE_API virtual FCursorReply OnCursorQuery( const FGeometry& MyGeometry, const FPointerEvent& CursorEvent ) const override;
+	SLATE_API virtual bool IsInteractable() const override;
+	SLATE_API virtual bool ComputeVolatility() const override;
 	//~ End SWidget Interface
 
 protected:
-	void OnHScrollBarMoved(const float InScrollOffsetFraction);
-	void OnVScrollBarMoved(const float InScrollOffsetFraction);
+	SLATE_API void OnHScrollBarMoved(const float InScrollOffsetFraction);
+	SLATE_API void OnVScrollBarMoved(const float InScrollOffsetFraction);
 
 	/** Return whether a RMB+Drag scroll operation is taking place */
-	bool IsRightClickScrolling() const;
+	SLATE_API bool IsRightClickScrolling() const;
 
 public:
 	//~ Begin ISlateEditableTextWidget Interface
-	virtual bool IsTextReadOnly() const override;
-	virtual bool IsTextPassword() const override;
-	virtual bool IsMultiLineTextEdit() const override;
+	SLATE_API virtual bool IsTextReadOnly() const override;
+	SLATE_API virtual bool IsTextPassword() const override;
+	SLATE_API virtual bool IsMultiLineTextEdit() const override;
 	//~ End ISlateEditableTextWidget Interface
 
 protected:
 	//~ Begin ISlateEditableTextWidget Interface
-	virtual bool ShouldJumpCursorToEndWhenFocused() const override;
-	virtual bool ShouldSelectAllTextWhenFocused() const override;
-	virtual bool ShouldClearTextSelectionOnFocusLoss() const override;
-	virtual bool ShouldRevertTextOnEscape() const override;
-	virtual bool ShouldClearKeyboardFocusOnCommit() const override;
-	virtual bool ShouldSelectAllTextOnCommit() const override;
-	virtual bool ShouldSelectWordOnMouseDoubleClick() const override;
-	virtual bool CanInsertCarriageReturn() const override;
-	virtual bool CanTypeCharacter(const TCHAR InChar) const override;
-	virtual void EnsureActiveTick() override;
-	virtual EKeyboardType GetVirtualKeyboardType() const override;
-	virtual FVirtualKeyboardOptions GetVirtualKeyboardOptions() const override;
-	virtual EVirtualKeyboardTrigger GetVirtualKeyboardTrigger() const override;
-	virtual EVirtualKeyboardDismissAction GetVirtualKeyboardDismissAction() const override;
-	virtual TSharedRef<SWidget> GetSlateWidget() override;
-	virtual TSharedPtr<SWidget> GetSlateWidgetPtr() override;
-	virtual TSharedPtr<SWidget> BuildContextMenuContent() const override;
-	virtual void OnTextChanged(const FText& InText) override;
-	virtual void OnTextCommitted(const FText& InText, const ETextCommit::Type InTextAction) override;
-	virtual void OnCursorMoved(const FTextLocation& InLocation) override;
-	virtual float UpdateAndClampHorizontalScrollBar(const float InViewOffset, const float InViewFraction, const EVisibility InVisiblityOverride) override;
-	virtual float UpdateAndClampVerticalScrollBar(const float InViewOffset, const float InViewFraction, const EVisibility InVisiblityOverride) override;
+	SLATE_API virtual bool ShouldJumpCursorToEndWhenFocused() const override;
+	SLATE_API virtual bool ShouldSelectAllTextWhenFocused() const override;
+	SLATE_API virtual bool ShouldClearTextSelectionOnFocusLoss() const override;
+	SLATE_API virtual bool ShouldRevertTextOnEscape() const override;
+	SLATE_API virtual bool ShouldClearKeyboardFocusOnCommit() const override;
+	SLATE_API virtual bool ShouldSelectAllTextOnCommit() const override;
+	SLATE_API virtual bool ShouldSelectWordOnMouseDoubleClick() const override;
+	SLATE_API virtual bool CanInsertCarriageReturn() const override;
+	SLATE_API virtual bool CanTypeCharacter(const TCHAR InChar) const override;
+	SLATE_API virtual void EnsureActiveTick() override;
+	SLATE_API virtual EKeyboardType GetVirtualKeyboardType() const override;
+	SLATE_API virtual FVirtualKeyboardOptions GetVirtualKeyboardOptions() const override;
+	SLATE_API virtual EVirtualKeyboardTrigger GetVirtualKeyboardTrigger() const override;
+	SLATE_API virtual EVirtualKeyboardDismissAction GetVirtualKeyboardDismissAction() const override;
+	SLATE_API virtual TSharedRef<SWidget> GetSlateWidget() override;
+	SLATE_API virtual TSharedPtr<SWidget> GetSlateWidgetPtr() override;
+	SLATE_API virtual TSharedPtr<SWidget> BuildContextMenuContent() const override;
+	SLATE_API virtual void OnTextChanged(const FText& InText) override;
+	SLATE_API virtual void OnTextCommitted(const FText& InText, const ETextCommit::Type InTextAction) override;
+	SLATE_API virtual void OnCursorMoved(const FTextLocation& InLocation) override;
+	SLATE_API virtual float UpdateAndClampHorizontalScrollBar(const float InViewOffset, const float InViewFraction, const EVisibility InVisiblityOverride) override;
+	SLATE_API virtual float UpdateAndClampVerticalScrollBar(const float InViewOffset, const float InViewFraction, const EVisibility InVisiblityOverride) override;
 	//~ End ISlateEditableTextWidget Interface
 
 protected:
 	/** Called to begin an undoable editable text transaction, marked as protected for use with FScopedEditableTextTransaction only */
-	void BeginEditTransaction();
+	SLATE_API void BeginEditTransaction();
 
 	/** Called to end an undoable editable text transaction, marked as protected for use with FScopedEditableTextTransaction only */
-	void EndEditTransaction();
+	SLATE_API void EndEditTransaction();
 
 protected:
 	/** The text layout that deals with the editable text */
diff --git a/Engine/Source/Runtime/Slate/Public/Widgets/Text/SRichTextBlock.h b/Engine/Source/Runtime/Slate/Public/Widgets/Text/SRichTextBlock.h
--- a/Engine/Source/Runtime/Slate/Public/Widgets/Text/SRichTextBlock.h
+++ b/Engine/Source/Runtime/Slate/Public/Widgets/Text/SRichTextBlock.h
@@ -30,7 +30,7 @@ enum class ETextShapingMethod : uint8;
  * A rich static text widget. 
  * Through the use of markup and text decorators, text with different styles, embedded image and widgets can be achieved.
  */
-class SLATE_API SRichTextBlock : public SWidget
+class SRichTextBlock : public SWidget
 {
 public:
 
@@ -158,16 +158,16 @@ public:
 	}
 
 	/** Constructor */
-	SRichTextBlock();
+	SLATE_API SRichTextBlock();
 
 	/** Destructor */
-	~SRichTextBlock();
+	SLATE_API ~SRichTextBlock();
 
 	//~ Begin SWidget Interface
-	void Construct( const FArguments& InArgs );
-	virtual int32 OnPaint( const FPaintArgs& Args, const FGeometry& AllottedGeometry, const FSlateRect& MyCullingRect, FSlateWindowElementList& OutDrawElements, int32 LayerId, const FWidgetStyle& InWidgetStyle, bool bParentEnabled ) const override;
-	virtual FChildren* GetChildren() override;
-	virtual void OnArrangeChildren( const FGeometry& AllottedGeometry, FArrangedChildren& ArrangedChildren ) const override;
+	SLATE_API void Construct( const FArguments& InArgs );
+	SLATE_API virtual int32 OnPaint( const FPaintArgs& Args, const FGeometry& AllottedGeometry, const FSlateRect& MyCullingRect, FSlateWindowElementList& OutDrawElements, int32 LayerId, const FWidgetStyle& InWidgetStyle, bool bParentEnabled ) const override;
+	SLATE_API virtual FChildren* GetChildren() override;
+	SLATE_API virtual void OnArrangeChildren( const FGeometry& AllottedGeometry, FArrangedChildren& ArrangedChildren ) const override;
 	//~ End SWidget Interface
 
 	/**
@@ -181,65 +181,65 @@ public:
 	/**
 	 * Sets the text for this text block
 	 */
-	void SetText( const TAttribute<FText>& InTextAttr );
+	SLATE_API void SetText( const TAttribute<FText>& InTextAttr );
 
 	/** See HighlightText attribute */
-	void SetHighlightText( const TAttribute<FText>& InHighlightText );
+	SLATE_API void SetHighlightText( const TAttribute<FText>& InHighlightText );
 
 	/** See TextShapingMethod attribute */
-	void SetTextShapingMethod(const TOptional<ETextShapingMethod>& InTextShapingMethod);
+	SLATE_API void SetTextShapingMethod(const TOptional<ETextShapingMethod>& InTextShapingMethod);
 
 	/** See TextFlowDirection attribute */
-	void SetTextFlowDirection(const TOptional<ETextFlowDirection>& InTextFlowDirection);
+	SLATE_API void SetTextFlowDirection(const TOptional<ETextFlowDirection>& InTextFlowDirection);
 
 	/** See WrapTextAt attribute */
-	void SetWrapTextAt(const TAttribute<float>& InWrapTextAt);
+	SLATE_API void SetWrapTextAt(const TAttribute<float>& InWrapTextAt);
 
 	/** See AutoWrapText attribute */
-	void SetAutoWrapText(const TAttribute<bool>& InAutoWrapText);
+	SLATE_API void SetAutoWrapText(const TAttribute<bool>& InAutoWrapText);
 
 	/** Set WrappingPolicy attribute */
-	void SetWrappingPolicy(const TAttribute<ETextWrappingPolicy>& InWrappingPolicy);
+	SLATE_API void SetWrappingPolicy(const TAttribute<ETextWrappingPolicy>& InWrappingPolicy);
 
 	/** Set TransformPolicy attribute */
-	void SetTransformPolicy(const TAttribute<ETextTransformPolicy>& InTransformPolicy);
+	SLATE_API void SetTransformPolicy(const TAttribute<ETextTransformPolicy>& InTransformPolicy);
 
 	/** See LineHeightPercentage attribute */
-	void SetLineHeightPercentage(const TAttribute<float>& InLineHeightPercentage);
+	SLATE_API void SetLineHeightPercentage(const TAttribute<float>& InLineHeightPercentage);
 
 	/** See Margin attribute */
-	void SetMargin(const TAttribute<FMargin>& InMargin);
+	SLATE_API void SetMargin(const TAttribute<FMargin>& InMargin);
 
 	/** See Justification attribute */
-	void SetJustification(const TAttribute<ETextJustify::Type>& InJustification);
+	SLATE_API void SetJustification(const TAttribute<ETextJustify::Type>& InJustification);
 
 	/** See TextStyle argument */
-	void SetTextStyle(const FTextBlockStyle& InTextStyle);
+	SLATE_API void SetTextStyle(const FTextBlockStyle& InTextStyle);
 
 	/** See MinDesiredWidth attribute */
-	void SetMinDesiredWidth(const TAttribute<float>& InMinDesiredWidth);
+	SLATE_API void SetMinDesiredWidth(const TAttribute<float>& InMinDesiredWidth);
 
 	/**  */
-	void SetDecoratorStyleSet(const ISlateStyle* NewDecoratorStyleSet);
+	SLATE_API void SetDecoratorStyleSet(const ISlateStyle* NewDecoratorStyleSet);
 
 	/** Replaces the decorators for this text block */
-	void SetDecorators(TArrayView<TSharedRef<ITextDecorator>> InDecorators);
+	SLATE_API void SetDecorators(TArrayView<TSharedRef<ITextDecorator>> InDecorators);
 
 	/** Sets the overflow policy for this text block */
-	void SetOverflowPolicy(TOptional<ETextOverflowPolicy> InOverflowPolicy);
+	SLATE_API void SetOverflowPolicy(TOptional<ETextOverflowPolicy> InOverflowPolicy);
 
 	/**
 	 * Causes the text to reflow it's layout
 	 */
-	void Refresh();
+	SLATE_API void Refresh();
 
 	/** set the scale value at the TextLayout*/
-	void SetTextBlockScale(const float NewTextBlockScale);
+	SLATE_API void SetTextBlockScale(const float NewTextBlockScale);
 
 protected:
 	//~ SWidget interface
-	virtual FVector2D ComputeDesiredSize(float) const override;
-	virtual bool ComputeVolatility() const override;
+	SLATE_API virtual FVector2D ComputeDesiredSize(float) const override;
+	SLATE_API virtual bool ComputeVolatility() const override;
 	//~ End of SWidget interface
 
 private:
diff --git a/Engine/Source/Runtime/Slate/Public/Widgets/Text/STextBlock.h b/Engine/Source/Runtime/Slate/Public/Widgets/Text/STextBlock.h
--- a/Engine/Source/Runtime/Slate/Public/Widgets/Text/STextBlock.h
+++ b/Engine/Source/Runtime/Slate/Public/Widgets/Text/STextBlock.h
@@ -36,9 +36,9 @@ namespace ETextRole
 /**
  * A simple static text widget
  */
-class SLATE_API STextBlock : public SLeafWidget
+class STextBlock : public SLeafWidget
 {
-	SLATE_DECLARE_WIDGET(STextBlock, SLeafWidget)
+	SLATE_DECLARE_WIDGET_API(STextBlock, SLeafWidget, SLATE_API)
 
 public:
 
@@ -151,17 +151,17 @@ public:
 	SLATE_END_ARGS()
 
 	/** Constructor */
-	STextBlock();
+	SLATE_API STextBlock();
 
 	/** Destructor */
-	~STextBlock();
+	SLATE_API ~STextBlock();
 
 	/**
 	 * Construct this widget
 	 *
 	 * @param	InArgs	The declaration data for this widget
 	 */
-	void Construct( const FArguments& InArgs );
+	SLATE_API void Construct( const FArguments& InArgs );
 
 	/**
 	 * Gets the text assigned to this text block
@@ -180,118 +180,118 @@ public:
 	
 public:
 	/** Sets the text for this text block */
-	void SetText(TAttribute<FText> InText);
+	SLATE_API void SetText(TAttribute<FText> InText);
 
 	/** Sets the highlight text for this text block */
-	void SetHighlightText(TAttribute<FText> InText);
+	SLATE_API void SetHighlightText(TAttribute<FText> InText);
 
 	/** Sets the font used to draw the text	*/
-	void SetFont(TAttribute<FSlateFontInfo> InFont);
+	SLATE_API void SetFont(TAttribute<FSlateFontInfo> InFont);
 
 	/** Sets the brush used to strike through the text */
-	void SetStrikeBrush(TAttribute<const FSlateBrush*> InStrikeBrush);
+	SLATE_API void SetStrikeBrush(TAttribute<const FSlateBrush*> InStrikeBrush);
 
 	/** See ColorAndOpacity attribute */
-	void SetColorAndOpacity(TAttribute<FSlateColor> InColorAndOpacity);
+	SLATE_API void SetColorAndOpacity(TAttribute<FSlateColor> InColorAndOpacity);
 
 	/** See TextStyle argument */
-	void SetTextStyle(const FTextBlockStyle* InTextStyle);
+	SLATE_API void SetTextStyle(const FTextBlockStyle* InTextStyle);
 
 	/** See TextShapingMethod attribute */
-	void SetTextShapingMethod(const TOptional<ETextShapingMethod>& InTextShapingMethod);
+	SLATE_API void SetTextShapingMethod(const TOptional<ETextShapingMethod>& InTextShapingMethod);
 
 	/** See TextFlowDirection attribute */
-	void SetTextFlowDirection(const TOptional<ETextFlowDirection>& InTextFlowDirection);
+	SLATE_API void SetTextFlowDirection(const TOptional<ETextFlowDirection>& InTextFlowDirection);
 
 	/** See WrapTextAt attribute */
-	void SetWrapTextAt(TAttribute<float> InWrapTextAt);
+	SLATE_API void SetWrapTextAt(TAttribute<float> InWrapTextAt);
 
 	/** See AutoWrapText attribute */
-	void SetAutoWrapText(TAttribute<bool> InAutoWrapText);
+	SLATE_API void SetAutoWrapText(TAttribute<bool> InAutoWrapText);
 
 	/** Set WrappingPolicy attribute */
-	void SetWrappingPolicy(TAttribute<ETextWrappingPolicy> InWrappingPolicy);
+	SLATE_API void SetWrappingPolicy(TAttribute<ETextWrappingPolicy> InWrappingPolicy);
 
 	/** Set TransformPolicy attribute */
-	void SetTransformPolicy(TAttribute<ETextTransformPolicy> InTransformPolicy);
+	SLATE_API void SetTransformPolicy(TAttribute<ETextTransformPolicy> InTransformPolicy);
 
 	/** Get TransformPolicy attribute */
 	UE_DEPRECATED(5.0, "GetTransformPolicy is not accessible anymore since it's attribute value may not have been updated yet.")
 	ETextTransformPolicy GetTransformPolicy() const { return GetTransformPolicyImpl(); }
 
 	/** Sets the overflow policy for this text block */
-	void SetOverflowPolicy(TOptional<ETextOverflowPolicy> InOverflowPolicy);
+	SLATE_API void SetOverflowPolicy(TOptional<ETextOverflowPolicy> InOverflowPolicy);
 
 	/** See ShadowOffset attribute */
-	void SetShadowOffset(TAttribute<FVector2D> InShadowOffset);
+	SLATE_API void SetShadowOffset(TAttribute<FVector2D> InShadowOffset);
 
 	/** See ShadowColorAndOpacity attribute */
-	void SetShadowColorAndOpacity(TAttribute<FLinearColor> InShadowColorAndOpacity);
+	SLATE_API void SetShadowColorAndOpacity(TAttribute<FLinearColor> InShadowColorAndOpacity);
 
 	/** See HighlightColor attribute */
-	void SetHighlightColor(TAttribute<FLinearColor> InHighlightColor);
+	SLATE_API void SetHighlightColor(TAttribute<FLinearColor> InHighlightColor);
 	
 	/** See HighlightShape attribute */
-	void SetHighlightShape(TAttribute<const FSlateBrush*> InHighlightShape);
+	SLATE_API void SetHighlightShape(TAttribute<const FSlateBrush*> InHighlightShape);
 
 	/** See MinDesiredWidth attribute */
-	void SetMinDesiredWidth(TAttribute<float> InMinDesiredWidth);
+	SLATE_API void SetMinDesiredWidth(TAttribute<float> InMinDesiredWidth);
 
 	/** See LineHeightPercentage attribute */
-	void SetLineHeightPercentage(TAttribute<float> InLineHeightPercentage);
+	SLATE_API void SetLineHeightPercentage(TAttribute<float> InLineHeightPercentage);
 
 	/** See Margin attribute */
-	void SetMargin(TAttribute<FMargin> InMargin);
+	SLATE_API void SetMargin(TAttribute<FMargin> InMargin);
 
 	/** See Justification attribute */
-	void SetJustification(TAttribute<ETextJustify::Type> InJustification);
+	SLATE_API void SetJustification(TAttribute<ETextJustify::Type> InJustification);
 
 	// SWidget interface
-	virtual int32 OnPaint( const FPaintArgs& Args, const FGeometry& AllottedGeometry, const FSlateRect& MyCullingRect, FSlateWindowElementList& OutDrawElements, int32 LayerId, const FWidgetStyle& InWidgetStyle, bool bParentEnabled ) const override;
-	virtual FVector2D ComputeDesiredSize(float) const override;
+	SLATE_API virtual int32 OnPaint( const FPaintArgs& Args, const FGeometry& AllottedGeometry, const FSlateRect& MyCullingRect, FSlateWindowElementList& OutDrawElements, int32 LayerId, const FWidgetStyle& InWidgetStyle, bool bParentEnabled ) const override;
+	SLATE_API virtual FVector2D ComputeDesiredSize(float) const override;
 #if WITH_ACCESSIBILITY
-	virtual TSharedRef<FSlateAccessibleWidget> CreateAccessibleWidget() override;
-	virtual TOptional<FText> GetDefaultAccessibleText(EAccessibleType AccessibleType = EAccessibleType::Main) const override;
+	SLATE_API virtual TSharedRef<FSlateAccessibleWidget> CreateAccessibleWidget() override;
+	SLATE_API virtual TOptional<FText> GetDefaultAccessibleText(EAccessibleType AccessibleType = EAccessibleType::Main) const override;
 #endif
 	// End of SWidget interface
 
 public:
 	/** Gets the current foreground color */
-	FSlateColor GetColorAndOpacity() const;
+	SLATE_API FSlateColor GetColorAndOpacity() const;
 
 	/** Gets the current font */
-	FSlateFontInfo GetFont() const;
+	SLATE_API FSlateFontInfo GetFont() const;
 
 	/** Gets the current strike brush */
-	const FSlateBrush* GetStrikeBrush() const;
+	SLATE_API const FSlateBrush* GetStrikeBrush() const;
 
 	/** Get TransformPolicy attribute */
-	ETextTransformPolicy GetTransformPolicyImpl() const;
+	SLATE_API ETextTransformPolicy GetTransformPolicyImpl() const;
 	
 	/** Gets the current shadow offset */
-	FVector2f GetShadowOffset() const;
+	SLATE_API FVector2f GetShadowOffset() const;
 
 	/** Gets the current shadow color and opacity */
-	FLinearColor GetShadowColorAndOpacity() const;
+	SLATE_API FLinearColor GetShadowColorAndOpacity() const;
 
 	/** Gets the current highlight color */
-	FSlateColor GetHighlightColor() const;
+	SLATE_API FSlateColor GetHighlightColor() const;
 
 	/** Gets the current highlight shape */
-	const FSlateBrush* GetHighlightShape() const;
+	SLATE_API const FSlateBrush* GetHighlightShape() const;
 
 	/** Gets the current margin */
-	FMargin GetMargin() const;
+	SLATE_API FMargin GetMargin() const;
 
 	/** Gets the current min desired width */
-	float GetMinDesiredWidth() const;
+	SLATE_API float GetMinDesiredWidth() const;
 
 private:
 	/** Get the computed text style to use with the text marshaller */
-	FTextBlockStyle GetComputedTextStyle() const;
+	SLATE_API FTextBlockStyle GetComputedTextStyle() const;
 
 	/** Call to invalidate this text block */
-	void InvalidateText(EInvalidateWidgetReason InvalidateReason);
+	SLATE_API void InvalidateText(EInvalidateWidgetReason InvalidateReason);
 
 private:
 	/** The text displayed in this text block */
diff --git a/Engine/Source/Runtime/Slate/Public/Widgets/Text/STextScroller.h b/Engine/Source/Runtime/Slate/Public/Widgets/Text/STextScroller.h
--- a/Engine/Source/Runtime/Slate/Public/Widgets/Text/STextScroller.h
+++ b/Engine/Source/Runtime/Slate/Public/Widgets/Text/STextScroller.h
@@ -26,7 +26,7 @@ struct FTextScrollerOptions
 /**
  * Utility to wrap a single-line text widget (STextBlock or SRichTextBlock) and provide support for auto-scrolling the text if it's longer than the available space.
  */
-class SLATE_API STextScroller : public SCompoundWidget
+class STextScroller : public SCompoundWidget
 {
 public:
 	SLATE_BEGIN_ARGS(STextScroller)
@@ -40,19 +40,19 @@ public:
 	SLATE_END_ARGS()
 
 public:
-	void Construct(const FArguments& InArgs);
+	SLATE_API void Construct(const FArguments& InArgs);
 
-	virtual void OnArrangeChildren(const FGeometry& AllottedGeometry, FArrangedChildren& ArrangedChildren) const override;
-	virtual int32 OnPaint(const FPaintArgs& Args, const FGeometry& AllottedGeometry, const FSlateRect& MyCullingRect, FSlateWindowElementList& OutDrawElements, int32 LayerId, const FWidgetStyle& InWidgetStyle, bool bParentEnabled) const override;
+	SLATE_API virtual void OnArrangeChildren(const FGeometry& AllottedGeometry, FArrangedChildren& ArrangedChildren) const override;
+	SLATE_API virtual int32 OnPaint(const FPaintArgs& Args, const FGeometry& AllottedGeometry, const FSlateRect& MyCullingRect, FSlateWindowElementList& OutDrawElements, int32 LayerId, const FWidgetStyle& InWidgetStyle, bool bParentEnabled) const override;
 
-	void ResetScrollState();
-	void StartScrolling();
-	void SuspendScrolling();
+	SLATE_API void ResetScrollState();
+	SLATE_API void StartScrolling();
+	SLATE_API void SuspendScrolling();
 
 	bool IsScrollingEnabled() const { return ActiveState != EActiveState::Suspend; }
 
 private:
-	EActiveTimerReturnType OnScrollTextTick(double CurrentTime, float DeltaTime);
+	SLATE_API EActiveTimerReturnType OnScrollTextTick(double CurrentTime, float DeltaTime);
 
 private:
 	enum class EActiveState : uint8
