[Slate]

diff --git a/Engine/Source/Runtime/Slate/Public/Widgets/Text/SlateEditableTextLayout.h b/Engine/Source/Runtime/Slate/Public/Widgets/Text/SlateEditableTextLayout.h
--- a/Engine/Source/Runtime/Slate/Public/Widgets/Text/SlateEditableTextLayout.h
+++ b/Engine/Source/Runtime/Slate/Public/Widgets/Text/SlateEditableTextLayout.h
@@ -31,41 +31,41 @@ class SWindow;
 enum class ETextShapingMethod : uint8;
 
 /** Class to handle the cached layout of SEditableText/SMultiLineEditableText by proxying around a FTextLayout */
-class SLATE_API FSlateEditableTextLayout
+class FSlateEditableTextLayout
 {
 public:
-	FSlateEditableTextLayout(ISlateEditableTextWidget& InOwnerWidget, const TAttribute<FText>& InInitialText, FTextBlockStyle InTextStyle, const TOptional<ETextShapingMethod> InTextShapingMethod, const TOptional<ETextFlowDirection> InTextFlowDirection, const FCreateSlateTextLayout& InCreateSlateTextLayout, TSharedRef<ITextLayoutMarshaller> InTextMarshaller, TSharedRef<ITextLayoutMarshaller> InHintTextMarshaller);
-	~FSlateEditableTextLayout();
+	SLATE_API FSlateEditableTextLayout(ISlateEditableTextWidget& InOwnerWidget, const TAttribute<FText>& InInitialText, FTextBlockStyle InTextStyle, const TOptional<ETextShapingMethod> InTextShapingMethod, const TOptional<ETextFlowDirection> InTextFlowDirection, const FCreateSlateTextLayout& InCreateSlateTextLayout, TSharedRef<ITextLayoutMarshaller> InTextMarshaller, TSharedRef<ITextLayoutMarshaller> InHintTextMarshaller);
+	SLATE_API ~FSlateEditableTextLayout();
 
-	void SetText(const TAttribute<FText>& InText);
-	FText GetText() const;
+	SLATE_API void SetText(const TAttribute<FText>& InText);
+	SLATE_API FText GetText() const;
 
-	void SetHintText(const TAttribute<FText>& InHintText);
-	FText GetHintText() const;
+	SLATE_API void SetHintText(const TAttribute<FText>& InHintText);
+	SLATE_API FText GetHintText() const;
 
-	void SetSearchText(const TAttribute<FText>& InSearchText);
-	FText GetSearchText() const;
+	SLATE_API void SetSearchText(const TAttribute<FText>& InSearchText);
+	SLATE_API FText GetSearchText() const;
 	
 	/** Get the index of the search result (0 if none) */
-	int32 GetSearchResultIndex() const;
+	SLATE_API int32 GetSearchResultIndex() const;
 	
 	/** Get the total number of search results (0 if none) */
-	int32 GetNumSearchResults() const;
+	SLATE_API int32 GetNumSearchResults() const;
 
-	void SetTextStyle(const FTextBlockStyle& InTextStyle);
-	const FTextBlockStyle& GetTextStyle() const;
+	SLATE_API void SetTextStyle(const FTextBlockStyle& InTextStyle);
+	SLATE_API const FTextBlockStyle& GetTextStyle() const;
 
 	/** Get the number of Text Lines */
-	int32 GetTextLineCount();
+	SLATE_API int32 GetTextLineCount();
 
 	/** Set the brush to use when drawing the cursor */
-	void SetCursorBrush(const TAttribute<const FSlateBrush*>& InCursorBrush);
+	SLATE_API void SetCursorBrush(const TAttribute<const FSlateBrush*>& InCursorBrush);
 
 	/** Set the brush to use when drawing the composition highlight */
-	void SetCompositionBrush(const TAttribute<const FSlateBrush*>& InCompositionBrush);
+	SLATE_API void SetCompositionBrush(const TAttribute<const FSlateBrush*>& InCompositionBrush);
 
 	/** Get the plain text string without rich-text formatting */
-	FText GetPlainText() const;
+	SLATE_API FText GetPlainText() const;
 
 	/**
 	 * Sets the current editable text for this text block
@@ -76,229 +76,229 @@ public:
 	 *
 	 * @return true if the text was updated, false if the text wasn't update (because it was already up-to-date)
 	 */
-	bool SetEditableText(const FText& TextToSet, const bool bForce = false);
+	SLATE_API bool SetEditableText(const FText& TextToSet, const bool bForce = false);
 
 	/**
 	 * Gets the current editable text for this text block
 	 * Note: We don't store text in this form (it's stored as lines in the text layout) so every call to this function has to reconstruct it
 	 */
-	FText GetEditableText() const;
+	SLATE_API FText GetEditableText() const;
 
 	/** Get the currently selected text */
-	FText GetSelectedText() const;
+	SLATE_API FText GetSelectedText() const;
 
 	/** Get the current selection*/
-	FTextSelection GetSelection() const;
+	SLATE_API FTextSelection GetSelection() const;
 	
 	/** Set the text shaping method that the internal text layout should use */
-	void SetTextShapingMethod(const TOptional<ETextShapingMethod>& InTextShapingMethod);
+	SLATE_API void SetTextShapingMethod(const TOptional<ETextShapingMethod>& InTextShapingMethod);
 
 	/** Set the text flow direction that the internal text layout should use */
-	void SetTextFlowDirection(const TOptional<ETextFlowDirection>& InTextFlowDirection);
+	SLATE_API void SetTextFlowDirection(const TOptional<ETextFlowDirection>& InTextFlowDirection);
 
 	/** Set the wrapping to use for this document */
-	void SetTextWrapping(const TAttribute<float>& InWrapTextAt, const TAttribute<bool>& InAutoWrapText, const TAttribute<ETextWrappingPolicy>& InWrappingPolicy);
+	SLATE_API void SetTextWrapping(const TAttribute<float>& InWrapTextAt, const TAttribute<bool>& InAutoWrapText, const TAttribute<ETextWrappingPolicy>& InWrappingPolicy);
 
 	/** Set whether text wraps onto a new line when it's length exceeds this width; if this value is zero or negative, no wrapping occurs */
-	void SetWrapTextAt(const TAttribute<float>& InWrapTextAt);
+	SLATE_API void SetWrapTextAt(const TAttribute<float>& InWrapTextAt);
 
 	/** Set whether to wrap text automatically based on the widget's computed horizontal space */
-	void SetAutoWrapText(const TAttribute<bool>& InAutoWrapText);
+	SLATE_API void SetAutoWrapText(const TAttribute<bool>& InAutoWrapText);
 
 	/** Set the wrapping policy to use */
-	void SetWrappingPolicy(const TAttribute<ETextWrappingPolicy>& InWrappingPolicy);
+	SLATE_API void SetWrappingPolicy(const TAttribute<ETextWrappingPolicy>& InWrappingPolicy);
 
 	/** Set the amount of blank space left around the edges of text area */
-	void SetMargin(const TAttribute<FMargin>& InMargin);
+	SLATE_API void SetMargin(const TAttribute<FMargin>& InMargin);
 
 	/** Set how the text should be aligned with the margin */
-	void SetJustification(const TAttribute<ETextJustify::Type>& InJustification);
+	SLATE_API void SetJustification(const TAttribute<ETextJustify::Type>& InJustification);
 
 	/** Set the amount to scale each lines height by */
-	void SetLineHeightPercentage(const TAttribute<float>& InLineHeightPercentage);
+	SLATE_API void SetLineHeightPercentage(const TAttribute<float>& InLineHeightPercentage);
 
 	/** Set the text overflow policy that should be used to determine what happens to clipped text */
-	void SetOverflowPolicy(TOptional<ETextOverflowPolicy> InOverflowPolicy);
+	SLATE_API void SetOverflowPolicy(TOptional<ETextOverflowPolicy> InOverflowPolicy);
 
 	/** Set the information used to help identify who owns this text layout in the case of an error */
-	void SetDebugSourceInfo(const TAttribute<FString>& InDebugSourceInfo);
+	SLATE_API void SetDebugSourceInfo(const TAttribute<FString>& InDebugSourceInfo);
 
 	/** Get the virtual keyboard handler for this text layout */
-	TSharedRef<IVirtualKeyboardEntry> GetVirtualKeyboardEntry() const;
+	SLATE_API TSharedRef<IVirtualKeyboardEntry> GetVirtualKeyboardEntry() const;
 
 	/** Get the IME context for this text layout */
-	TSharedRef<ITextInputMethodContext> GetTextInputMethodContext() const;
+	SLATE_API TSharedRef<ITextInputMethodContext> GetTextInputMethodContext() const;
 
 	/** Register and activate the IME context for this text layout */
-	void EnableTextInputMethodContext();
+	SLATE_API void EnableTextInputMethodContext();
 
 	/** Refresh this editable text immediately, rather than wait for the usual caching mechanisms to take affect on the text Tick */
-	bool Refresh();
+	SLATE_API bool Refresh();
 
 	/** Force the text layout to be updated from the marshaller */
-	void ForceRefreshTextLayout(const FText& CurrentText);
+	SLATE_API void ForceRefreshTextLayout(const FText& CurrentText);
 
 	/** Begin a new text search (this is called automatically when BoundSearchText changes) */
-	void BeginSearch(const FText& InSearchText, const ESearchCase::Type InSearchCase = ESearchCase::IgnoreCase, const bool InReverse = false);
+	SLATE_API void BeginSearch(const FText& InSearchText, const ESearchCase::Type InSearchCase = ESearchCase::IgnoreCase, const bool InReverse = false);
 
 	/** Advance the current search to the next match (does nothing if not currently searching) */
-	void AdvanceSearch(const bool InReverse = false);
+	SLATE_API void AdvanceSearch(const bool InReverse = false);
 
 	/** Update the horizontal scroll amount from the given fraction */
-	UE::Slate::FDeprecateVector2DResult SetHorizontalScrollFraction(const float InScrollOffsetFraction);
+	SLATE_API UE::Slate::FDeprecateVector2DResult SetHorizontalScrollFraction(const float InScrollOffsetFraction);
 
 	/** Update the vertical scroll amount from the given fraction */
-	UE::Slate::FDeprecateVector2DResult SetVerticalScrollFraction(const float InScrollOffsetFraction);
+	SLATE_API UE::Slate::FDeprecateVector2DResult SetVerticalScrollFraction(const float InScrollOffsetFraction);
 
 	/** Set the absolute scroll offset value */
-	UE::Slate::FDeprecateVector2DResult SetScrollOffset(const UE::Slate::FDeprecateVector2DParameter& InScrollOffset, const FGeometry& InGeometry);
+	SLATE_API UE::Slate::FDeprecateVector2DResult SetScrollOffset(const UE::Slate::FDeprecateVector2DParameter& InScrollOffset, const FGeometry& InGeometry);
 
 	/** Get the absolute scroll offset value */
-	UE::Slate::FDeprecateVector2DResult GetScrollOffset() const;
+	SLATE_API UE::Slate::FDeprecateVector2DResult GetScrollOffset() const;
 
 	/** Returns the computed wrap location for this layout */
-	float GetComputedWrappingWidth() const;
+	SLATE_API float GetComputedWrappingWidth() const;
 
 	/** Returns whether or not we are auto wrapping text */
-	bool GetAutoWrapText() const;
+	SLATE_API bool GetAutoWrapText() const;
 
 	/** Called when our parent widget receives focus */
-	bool HandleFocusReceived(const FFocusEvent& InFocusEvent);
+	SLATE_API bool HandleFocusReceived(const FFocusEvent& InFocusEvent);
 
 	/** Called when our parent widget loses focus */
-	bool HandleFocusLost(const FFocusEvent& InFocusEvent);
+	SLATE_API bool HandleFocusLost(const FFocusEvent& InFocusEvent);
 
 	/** Called to handle an OnKeyChar event from our parent widget */
-	FReply HandleKeyChar(const FCharacterEvent& InCharacterEvent);
+	SLATE_API FReply HandleKeyChar(const FCharacterEvent& InCharacterEvent);
 
 	/** Called to handle an OnKeyDown event from our parent widget */
-	FReply HandleKeyDown(const FKeyEvent& InKeyEvent);
+	SLATE_API FReply HandleKeyDown(const FKeyEvent& InKeyEvent);
 
 	/** Called to handle an OnKeyUp event from our parent widget */
-	FReply HandleKeyUp(const FKeyEvent& InKeyEvent);
+	SLATE_API FReply HandleKeyUp(const FKeyEvent& InKeyEvent);
 
 	/** Called to handle an OnMouseButtonDown event from our parent widget */
-	FReply HandleMouseButtonDown(const FGeometry& MyGeometry, const FPointerEvent& InMouseEvent);
+	SLATE_API FReply HandleMouseButtonDown(const FGeometry& MyGeometry, const FPointerEvent& InMouseEvent);
 
 	/** Called to handle an OnMouseButtonUp event from our parent widget */
-	FReply HandleMouseButtonUp(const FGeometry& MyGeometry, const FPointerEvent& InMouseEvent);
+	SLATE_API FReply HandleMouseButtonUp(const FGeometry& MyGeometry, const FPointerEvent& InMouseEvent);
 
 	/** Called to handle an OnMouseMove event from our parent widget */
-	FReply HandleMouseMove(const FGeometry& InMyGeometry, const FPointerEvent& InMouseEvent);
+	SLATE_API FReply HandleMouseMove(const FGeometry& InMyGeometry, const FPointerEvent& InMouseEvent);
 
 	/** Called to handle an OnMouseButtonDoubleClick event from our parent widget */
-	FReply HandleMouseButtonDoubleClick(const FGeometry& InMyGeometry, const FPointerEvent& InMouseEvent);
+	SLATE_API FReply HandleMouseButtonDoubleClick(const FGeometry& InMyGeometry, const FPointerEvent& InMouseEvent);
 
 	/** Called to handle an escape action acting on the current selection */
-	bool HandleEscape();
+	SLATE_API bool HandleEscape();
 
 	/** Called to handle a backspace action acting on the current selection or at the cursor position */
-	bool HandleBackspace();
+	SLATE_API bool HandleBackspace();
 
 	/** Called to handle a delete action acting on the current selection or at the cursor position */
-	bool HandleDelete();
+	SLATE_API bool HandleDelete();
 
 	/** Called to handle a typing a character on the current selection or at the cursor position */
-	bool HandleTypeChar(const TCHAR InChar);
+	SLATE_API bool HandleTypeChar(const TCHAR InChar);
 
 	/** Called to handle a carriage return action acting on the current selection or at the cursor position */
-	bool HandleCarriageReturn(bool isRepeat);
+	SLATE_API bool HandleCarriageReturn(bool isRepeat);
 
 	/** Are we able to delete the currently selected text? */
-	bool CanExecuteDelete() const;
+	SLATE_API bool CanExecuteDelete() const;
 
 	/** Delete any currently selected text */
-	void DeleteSelectedText();
+	SLATE_API void DeleteSelectedText();
 
 	/** Query to see if any text is selected within the document */
-	bool AnyTextSelected() const;
+	SLATE_API bool AnyTextSelected() const;
 
 	/** Query to see if the text under the given position is currently selected */
-	bool IsTextSelectedAt(const FGeometry& MyGeometry, const UE::Slate::FDeprecateVector2DParameter& ScreenSpacePosition) const;
+	SLATE_API bool IsTextSelectedAt(const FGeometry& MyGeometry, const UE::Slate::FDeprecateVector2DParameter& ScreenSpacePosition) const;
 
 	/** Query to see if the text under the given position is currently selected (the position is local to the text layout space) */
-	bool IsTextSelectedAt(const UE::Slate::FDeprecateVector2DParameter& InLocalPosition) const;
+	SLATE_API bool IsTextSelectedAt(const UE::Slate::FDeprecateVector2DParameter& InLocalPosition) const;
 
 	/** Are we able to execute the "Select All" command? */
-	bool CanExecuteSelectAll() const;
+	SLATE_API bool CanExecuteSelectAll() const;
 
 	/** Select all the text in the document */
-	void SelectAllText();
+	SLATE_API void SelectAllText();
 
 	/** Select the word under the given position */
-	void SelectWordAt(const FGeometry& MyGeometry, const UE::Slate::FDeprecateVector2DParameter& ScreenSpacePosition);
+	SLATE_API void SelectWordAt(const FGeometry& MyGeometry, const UE::Slate::FDeprecateVector2DParameter& ScreenSpacePosition);
 
 	/** Select the word under the given position (the position is local to the text layout space) */
-	void SelectWordAt(const UE::Slate::FDeprecateVector2DParameter& InLocalPosition);
+	SLATE_API void SelectWordAt(const UE::Slate::FDeprecateVector2DParameter& InLocalPosition);
 
 	/** Select a block of text */
-	void SelectText(const FTextLocation& InSelectionStart, const FTextLocation& InCursorLocation);
+	SLATE_API void SelectText(const FTextLocation& InSelectionStart, const FTextLocation& InCursorLocation);
 
 	/** Clear the active text selection */
-	void ClearSelection();
+	SLATE_API void ClearSelection();
 
 	/** Are we able to cut the currently selected text? */
-	bool CanExecuteCut() const;
+	SLATE_API bool CanExecuteCut() const;
 
 	/** Cut the currently selected text and place it on the clipboard */
-	void CutSelectedTextToClipboard();
+	SLATE_API void CutSelectedTextToClipboard();
 
 	/** Are we able to copy the currently selected text? */
-	bool CanExecuteCopy() const;
+	SLATE_API bool CanExecuteCopy() const;
 
 	/** Copy the currently selected text and place it on the clipboard */
-	void CopySelectedTextToClipboard();
+	SLATE_API void CopySelectedTextToClipboard();
 
 	/** Are we able to paste the text from the clipboard into this document? */
-	bool CanExecutePaste() const;
+	SLATE_API bool CanExecutePaste() const;
 
 	/** Paste the text from the clipboard into this document */
-	void PasteTextFromClipboard();
+	SLATE_API void PasteTextFromClipboard();
 
 	/** Insert the given text at the current cursor position, correctly taking into account new line characters */
-	void InsertTextAtCursor(const FString& InString);
+	SLATE_API void InsertTextAtCursor(const FString& InString);
 
 	/** Insert the given run at the current cursor position */
-	void InsertRunAtCursor(TSharedRef<IRun> InRun);
+	SLATE_API void InsertRunAtCursor(TSharedRef<IRun> InRun);
 
 	/** Move the cursor in the document using the specified move method */
-	bool MoveCursor(const FMoveCursor& InArgs);
+	SLATE_API bool MoveCursor(const FMoveCursor& InArgs);
 
 	/** Move the cursor to the given location in the document (will also scroll to this point) */
-	void GoTo(const FTextLocation& NewLocation);
+	SLATE_API void GoTo(const FTextLocation& NewLocation);
 
 	/** Move the cursor specified location */
-	void GoTo(const ETextLocation NewLocation);
+	SLATE_API void GoTo(const ETextLocation NewLocation);
 
 	/** Jump the cursor to the given location in the document */
-	void JumpTo(ETextLocation JumpLocation, ECursorAction Action);
+	SLATE_API void JumpTo(ETextLocation JumpLocation, ECursorAction Action);
 
 	/** Scroll to the given location in the document (without moving the cursor) */
-	void ScrollTo(const FTextLocation& NewLocation);
+	SLATE_API void ScrollTo(const FTextLocation& NewLocation);
 
 	/** Scroll to the given location in the document (without moving the cursor) */
-	void ScrollTo(const ETextLocation NewLocation);
+	SLATE_API void ScrollTo(const ETextLocation NewLocation);
 
 	/** Update the active cursor highlight based on the state of the text layout */
-	void UpdateCursorHighlight();
+	SLATE_API void UpdateCursorHighlight();
 
 	/** Remove any active cursor highlights */
-	void RemoveCursorHighlight();
+	SLATE_API void RemoveCursorHighlight();
 
 	/** Update the preferred offset of the cursor, based on the current state of the text layout */
-	void UpdatePreferredCursorScreenOffsetInLine();
+	SLATE_API void UpdatePreferredCursorScreenOffsetInLine();
 
 	/** Apply the given style to the currently selected text (or insert a new run at the current cursor position if no text is selected) */
-	void ApplyToSelection(const FRunInfo& InRunInfo, const FTextBlockStyle& InStyle);
+	SLATE_API void ApplyToSelection(const FRunInfo& InRunInfo, const FTextBlockStyle& InStyle);
 
 	/** Get the run currently under the cursor, or null if there is no run currently under the cursor */
-	TSharedPtr<const IRun> GetRunUnderCursor() const;
+	SLATE_API TSharedPtr<const IRun> GetRunUnderCursor() const;
 
 	/** Get the runs currently that are current selected, some of which may be only partially selected */
-	TArray<TSharedRef<const IRun>> GetSelectedRuns() const;
+	SLATE_API TArray<TSharedRef<const IRun>> GetSelectedRuns() const;
 
 	/** Get the interaction position of the cursor (where to insert, delete, etc, text from/to) */
-	FTextLocation GetCursorLocation() const;
+	SLATE_API FTextLocation GetCursorLocation() const;
 
 	/**
 	 * Given a location and a Direction to offset, return a new location.
@@ -306,7 +306,7 @@ public:
 	 * @param Location    Cursor location from which to offset
 	 * @param Direction   Positive means right, negative means left.
 	 */
-	FTextLocation TranslatedLocation(const FTextLocation& CurrentLocation, int8 Direction) const;
+	SLATE_API FTextLocation TranslatedLocation(const FTextLocation& CurrentLocation, int8 Direction) const;
 
 	/**
 	 * Given a location and a Direction to offset, return a new location.
@@ -317,92 +317,92 @@ public:
 	 * @param OutCursorPosition     Fill with the updated cursor position.
 	 * @param OutCursorAlignment    Optionally fill with a new cursor alignment (will be auto-calculated if not set).
 	 */
-	void TranslateLocationVertical(const FTextLocation& Location, int32 NumLinesToMove, float GeometryScale, FTextLocation& OutCursorPosition, TOptional<SlateEditableTextTypes::ECursorAlignment>& OutCursorAlignment) const;
+	SLATE_API void TranslateLocationVertical(const FTextLocation& Location, int32 NumLinesToMove, float GeometryScale, FTextLocation& OutCursorPosition, TOptional<SlateEditableTextTypes::ECursorAlignment>& OutCursorAlignment) const;
 
 	/** Find the closest word boundary */
-	FTextLocation ScanForWordBoundary(const FTextLocation& Location, int8 Direction) const;
+	SLATE_API FTextLocation ScanForWordBoundary(const FTextLocation& Location, int8 Direction) const;
 
 	/** Get the character at Location */
-	TCHAR GetCharacterAt(const FTextLocation& Location) const;
+	SLATE_API TCHAR GetCharacterAt(const FTextLocation& Location) const;
 
 	/** Are we at the beginning of all the text. */
-	bool IsAtBeginningOfDocument(const FTextLocation& Location) const;
+	SLATE_API bool IsAtBeginningOfDocument(const FTextLocation& Location) const;
 
 	/** Are we at the end of all the text. */
-	bool IsAtEndOfDocument(const FTextLocation& Location) const;
+	SLATE_API bool IsAtEndOfDocument(const FTextLocation& Location) const;
 
 	/** Is this location the beginning of a line */
-	bool IsAtBeginningOfLine(const FTextLocation& Location) const;
+	SLATE_API bool IsAtBeginningOfLine(const FTextLocation& Location) const;
 
 	/** Is this location the end of a line. */
-	bool IsAtEndOfLine(const FTextLocation& Location) const;
+	SLATE_API bool IsAtEndOfLine(const FTextLocation& Location) const;
 
 	/** Are we currently at the beginning of a word */
-	bool IsAtWordStart(const FTextLocation& Location) const;
+	SLATE_API bool IsAtWordStart(const FTextLocation& Location) const;
 
 	/** Restores the text to the original state */
-	void RestoreOriginalText();
+	SLATE_API void RestoreOriginalText();
 
 	/** Returns whether the current text varies from the original */
-	bool HasTextChangedFromOriginal() const;
+	SLATE_API bool HasTextChangedFromOriginal() const;
 
 	/** Called to begin an undoable editable text transaction */
-	void BeginEditTransation();
+	SLATE_API void BeginEditTransation();
 
 	/** Called to end an undoable editable text transaction */
-	void EndEditTransaction();
+	SLATE_API void EndEditTransaction();
 
 	/** Push the given undo state onto the undo stack */
-	void PushUndoState(const SlateEditableTextTypes::FUndoState& InUndoState);
+	SLATE_API void PushUndoState(const SlateEditableTextTypes::FUndoState& InUndoState);
 
 	/** Clear the current undo stack */
-	void ClearUndoStates();
+	SLATE_API void ClearUndoStates();
 
 	/** Create an undo state that reflects the current state of the document */
-	void MakeUndoState(SlateEditableTextTypes::FUndoState& OutUndoState);
+	SLATE_API void MakeUndoState(SlateEditableTextTypes::FUndoState& OutUndoState);
 
 	/** Are we currently able to execute an undo action? */
-	bool CanExecuteUndo() const;
+	SLATE_API bool CanExecuteUndo() const;
 	
 	/** Execute an undo action */
-	void Undo();
+	SLATE_API void Undo();
 
 	/** Are we currently able to execute a redo action? */
-	bool CanExecuteRedo() const;
+	SLATE_API bool CanExecuteRedo() const;
 
 	/** Execute a redo action */
-	void Redo();
+	SLATE_API void Redo();
 
-	void SaveText(const FText& TextToSave);
+	SLATE_API void SaveText(const FText& TextToSave);
 
-	void LoadText();
+	SLATE_API void LoadText();
 
-	bool ComputeVolatility() const;
+	SLATE_API bool ComputeVolatility() const;
 
-	void Tick(const FGeometry& AllottedGeometry, const double InCurrentTime, const float InDeltaTime);
+	SLATE_API void Tick(const FGeometry& AllottedGeometry, const double InCurrentTime, const float InDeltaTime);
 
-	int32 OnPaint(const FPaintArgs& Args, const FGeometry& AllottedGeometry, const FSlateRect& MyCullingRect, FSlateWindowElementList& OutDrawElements, int32 LayerId, const FWidgetStyle& InWidgetStyle, bool bParentEnabled);
+	SLATE_API int32 OnPaint(const FPaintArgs& Args, const FGeometry& AllottedGeometry, const FSlateRect& MyCullingRect, FSlateWindowElementList& OutDrawElements, int32 LayerId, const FWidgetStyle& InWidgetStyle, bool bParentEnabled);
 
-	void CacheDesiredSize(float LayoutScaleMultiplier);
+	SLATE_API void CacheDesiredSize(float LayoutScaleMultiplier);
 
-	FVector2D ComputeDesiredSize(float LayoutScaleMultiplier) const;
+	SLATE_API FVector2D ComputeDesiredSize(float LayoutScaleMultiplier) const;
 
-	FChildren* GetChildren();
+	SLATE_API FChildren* GetChildren();
 
-	void OnArrangeChildren(const FGeometry& AllottedGeometry, FArrangedChildren& ArrangedChildren) const;
+	SLATE_API void OnArrangeChildren(const FGeometry& AllottedGeometry, FArrangedChildren& ArrangedChildren) const;
 
-	UE::Slate::FDeprecateVector2DResult GetSize() const;
+	SLATE_API UE::Slate::FDeprecateVector2DResult GetSize() const;
 
-	TSharedRef<SWidget> BuildDefaultContextMenu(const TSharedPtr<FExtender>& InMenuExtender) const;
+	SLATE_API TSharedRef<SWidget> BuildDefaultContextMenu(const TSharedPtr<FExtender>& InMenuExtender) const;
 
-	bool HasActiveContextMenu() const;
+	SLATE_API bool HasActiveContextMenu() const;
 
 	/**
 	 * Fill OutTextLine with the text line where the current cursor location is at
 	 *
 	 * @param OutTextLine   FString of the line
 	 */
-	void GetCurrentTextLine(FString& OutTextLine) const;
+	SLATE_API void GetCurrentTextLine(FString& OutTextLine) const;
 
 	/**
 	 * Fill OutTextLine with the text line at the specified index
@@ -410,23 +410,23 @@ public:
 	 * @param InLineIndex   Index of the line
 	 * @param OutTextLine   FString of the line
 	 */
-	void GetTextLine(const int32 InLineIndex, FString& OutTextLine) const;
+	SLATE_API void GetTextLine(const int32 InLineIndex, FString& OutTextLine) const;
 
 private:
 	/** Insert the given text at the current cursor position, correctly taking into account new line characters */
-	void InsertTextAtCursorImpl(const FString& InString);
+	SLATE_API void InsertTextAtCursorImpl(const FString& InString);
 
 	/** Insert a new-line at the current cursor position */
-	void InsertNewLineAtCursorImpl();
+	SLATE_API void InsertNewLineAtCursorImpl();
 
 	/** Implementation of Refresh that actually updates the layout. Optionally takes text to set, or will use the current editable text if none if provided */
-	bool RefreshImpl(const FText* InTextToSet, const bool bForce = false);
+	SLATE_API bool RefreshImpl(const FText* InTextToSet, const bool bForce = false);
 
 	/** Create a text or password run using the given text and style */
-	TSharedRef<IRun> CreateTextOrPasswordRun(const FRunInfo& InRunInfo, const TSharedRef<const FString>& InText, const FTextBlockStyle& InStyle);
+	SLATE_API TSharedRef<IRun> CreateTextOrPasswordRun(const FRunInfo& InRunInfo, const TSharedRef<const FString>& InText, const FTextBlockStyle& InStyle);
 
 	/** Called when the active context menu is closed */
-	void OnContextMenuClosed(TSharedRef<IMenu> Menu);
+	SLATE_API void OnContextMenuClosed(TSharedRef<IMenu> Menu);
 
 private:
 	/** Virtual keyboard handler for an editable text layout */
diff --git a/Engine/Source/Runtime/Slate/Public/Widgets/Text/SlateEditableTextTypes.h b/Engine/Source/Runtime/Slate/Public/Widgets/Text/SlateEditableTextTypes.h
--- a/Engine/Source/Runtime/Slate/Public/Widgets/Text/SlateEditableTextTypes.h
+++ b/Engine/Source/Runtime/Slate/Public/Widgets/Text/SlateEditableTextTypes.h
@@ -28,7 +28,7 @@ enum class ECursorAlignment : uint8
 };
 
 /** Store the information about the current cursor position */
-class SLATE_API FCursorInfo
+class FCursorInfo
 {
 public:
 	FCursorInfo()
@@ -70,19 +70,19 @@ public:
 	}
 
 	/** Update the last cursor interaction time to now */
-	void UpdateLastCursorInteractionTime();
+	SLATE_API void UpdateLastCursorInteractionTime();
 
 	/** Set the position of the cursor, and then work out the correct alignment based on the current text layout */
-	void SetCursorLocationAndCalculateAlignment(const FTextLayout& InTextLayout, const FTextLocation& InCursorPosition);
+	SLATE_API void SetCursorLocationAndCalculateAlignment(const FTextLayout& InTextLayout, const FTextLocation& InCursorPosition);
 
 	/** Set the literal position and alignment of the cursor */
-	void SetCursorLocationAndAlignment(const FTextLayout& InTextLayout, const FTextLocation& InCursorPosition, const ECursorAlignment InCursorAlignment);
+	SLATE_API void SetCursorLocationAndAlignment(const FTextLayout& InTextLayout, const FTextLocation& InCursorPosition, const ECursorAlignment InCursorAlignment);
 
 	/** Create an undo for this cursor data */
-	FCursorInfo CreateUndo() const;
+	SLATE_API FCursorInfo CreateUndo() const;
 
 	/** Restore this cursor data from an undo */
-	void RestoreFromUndo(const FCursorInfo& UndoData);
+	SLATE_API void RestoreFromUndo(const FCursorInfo& UndoData);
 
 private:
 	/** Current cursor position; there is always one. */
@@ -134,17 +134,17 @@ struct FScrollInfo
 };
 
 /** Run highlighter used to draw the cursor */
-class SLATE_API FCursorLineHighlighter : public ISlateLineHighlighter
+class FCursorLineHighlighter : public ISlateLineHighlighter
 {
 public:
-	static TSharedRef< FCursorLineHighlighter > Create(const FCursorInfo* InCursorInfo);
+	static SLATE_API TSharedRef< FCursorLineHighlighter > Create(const FCursorInfo* InCursorInfo);
 
-	void SetCursorBrush(const TAttribute<const FSlateBrush*>& InCursorBrush);
+	SLATE_API void SetCursorBrush(const TAttribute<const FSlateBrush*>& InCursorBrush);
 
-	virtual int32 OnPaint(const FPaintArgs& Args, const FTextLayout::FLineView& Line, const float OffsetX, const float Width, const FTextBlockStyle& DefaultStyle, const FGeometry& AllottedGeometry, const FSlateRect& MyCullingRect, FSlateWindowElementList& OutDrawElements, int32 LayerId, const FWidgetStyle& InWidgetStyle, bool bParentEnabled) const override;
+	SLATE_API virtual int32 OnPaint(const FPaintArgs& Args, const FTextLayout::FLineView& Line, const float OffsetX, const float Width, const FTextBlockStyle& DefaultStyle, const FGeometry& AllottedGeometry, const FSlateRect& MyCullingRect, FSlateWindowElementList& OutDrawElements, int32 LayerId, const FWidgetStyle& InWidgetStyle, bool bParentEnabled) const override;
 
 protected:
-	FCursorLineHighlighter(const FCursorInfo* InCursorInfo);
+	SLATE_API FCursorLineHighlighter(const FCursorInfo* InCursorInfo);
 
 	/** Cursor data that this highlighter is tracking */
 	const FCursorInfo* CursorInfo;
@@ -154,29 +154,29 @@ protected:
 };
 
 /** Run highlighter used to draw the composition range */
-class SLATE_API FTextCompositionHighlighter : public ISlateLineHighlighter
+class FTextCompositionHighlighter : public ISlateLineHighlighter
 {
 public:
-	static TSharedRef< FTextCompositionHighlighter > Create();
+	static SLATE_API TSharedRef< FTextCompositionHighlighter > Create();
 
-	void SetCompositionBrush(const TAttribute<const FSlateBrush*>& InCompositionBrush);
+	SLATE_API void SetCompositionBrush(const TAttribute<const FSlateBrush*>& InCompositionBrush);
 
-	virtual int32 OnPaint(const FPaintArgs& Args, const FTextLayout::FLineView& Line, const float OffsetX, const float Width, const FTextBlockStyle& DefaultStyle, const FGeometry& AllottedGeometry, const FSlateRect& MyCullingRect, FSlateWindowElementList& OutDrawElements, int32 LayerId, const FWidgetStyle& InWidgetStyle, bool bParentEnabled) const override;
+	SLATE_API virtual int32 OnPaint(const FPaintArgs& Args, const FTextLayout::FLineView& Line, const float OffsetX, const float Width, const FTextBlockStyle& DefaultStyle, const FGeometry& AllottedGeometry, const FSlateRect& MyCullingRect, FSlateWindowElementList& OutDrawElements, int32 LayerId, const FWidgetStyle& InWidgetStyle, bool bParentEnabled) const override;
 
 protected:
-	FTextCompositionHighlighter();
+	SLATE_API FTextCompositionHighlighter();
 
 	/** Brush used to draw the composition highlight */
 	TAttribute<const FSlateBrush*> CompositionBrush;
 };
 
 /** Run highlighter used to draw selection ranges */
-class SLATE_API FTextSelectionHighlighter : public ISlateLineHighlighter
+class FTextSelectionHighlighter : public ISlateLineHighlighter
 {
 public:
-	static TSharedRef< FTextSelectionHighlighter > Create();
+	static SLATE_API TSharedRef< FTextSelectionHighlighter > Create();
 
-	virtual int32 OnPaint(const FPaintArgs& Args, const FTextLayout::FLineView& Line, const float OffsetX, const float Width, const FTextBlockStyle& DefaultStyle, const FGeometry& AllottedGeometry, const FSlateRect& MyCullingRect, FSlateWindowElementList& OutDrawElements, int32 LayerId, const FWidgetStyle& InWidgetStyle, bool bParentEnabled) const override;
+	SLATE_API virtual int32 OnPaint(const FPaintArgs& Args, const FTextLayout::FLineView& Line, const float OffsetX, const float Width, const FTextBlockStyle& DefaultStyle, const FGeometry& AllottedGeometry, const FSlateRect& MyCullingRect, FSlateWindowElementList& OutDrawElements, int32 LayerId, const FWidgetStyle& InWidgetStyle, bool bParentEnabled) const override;
 
 	void SetHasKeyboardFocus(const bool bInHasKeyboardFocus)
 	{
@@ -184,19 +184,19 @@ public:
 	}
 
 protected:
-	FTextSelectionHighlighter();
+	SLATE_API FTextSelectionHighlighter();
 
 	/** true if the parent widget has keyboard focus, false otherwise */
 	bool bHasKeyboardFocus;
 };
 
 /** Run highlighter used to draw search ranges */
-class SLATE_API FTextSearchHighlighter : public ISlateLineHighlighter
+class FTextSearchHighlighter : public ISlateLineHighlighter
 {
 public:
-	static TSharedRef< FTextSearchHighlighter > Create();
+	static SLATE_API TSharedRef< FTextSearchHighlighter > Create();
 
-	virtual int32 OnPaint(const FPaintArgs& Args, const FTextLayout::FLineView& Line, const float OffsetX, const float Width, const FTextBlockStyle& DefaultStyle, const FGeometry& AllottedGeometry, const FSlateRect& MyCullingRect, FSlateWindowElementList& OutDrawElements, int32 LayerId, const FWidgetStyle& InWidgetStyle, bool bParentEnabled) const override;
+	SLATE_API virtual int32 OnPaint(const FPaintArgs& Args, const FTextLayout::FLineView& Line, const float OffsetX, const float Width, const FTextBlockStyle& DefaultStyle, const FGeometry& AllottedGeometry, const FSlateRect& MyCullingRect, FSlateWindowElementList& OutDrawElements, int32 LayerId, const FWidgetStyle& InWidgetStyle, bool bParentEnabled) const override;
 
 	void SetHasKeyboardFocus(const bool bInHasKeyboardFocus)
 	{
@@ -204,7 +204,7 @@ public:
 	}
 
 protected:
-	FTextSearchHighlighter();
+	SLATE_API FTextSearchHighlighter();
 
 	/** true if the parent widget has keyboard focus, false otherwise */
 	bool bHasKeyboardFocus;
diff --git a/Engine/Source/Runtime/Slate/Public/Widgets/Text/SlateTextBlockLayout.h b/Engine/Source/Runtime/Slate/Public/Widgets/Text/SlateTextBlockLayout.h
--- a/Engine/Source/Runtime/Slate/Public/Widgets/Text/SlateTextBlockLayout.h
+++ b/Engine/Source/Runtime/Slate/Public/Widgets/Text/SlateTextBlockLayout.h
@@ -17,7 +17,7 @@ class ISlateRunRenderer;
 enum class ETextShapingMethod : uint8;
 
 /** Class to handle the cached layout of STextBlock/SRichTextBlock by proxying around a FTextLayout */
-class SLATE_API FSlateTextBlockLayout
+class FSlateTextBlockLayout
 {
 public:
 	struct FWidgetDesiredSizeArgs
@@ -97,20 +97,20 @@ public:
 	/**
 	 * Constructor
 	 */
-	FSlateTextBlockLayout(SWidget* InOwner, FTextBlockStyle InDefaultTextStyle, const TOptional<ETextShapingMethod> InTextShapingMethod, const TOptional<ETextFlowDirection> InTextFlowDirection, const FCreateSlateTextLayout& InCreateSlateTextLayout, TSharedRef<ITextLayoutMarshaller> InMarshaller, TSharedPtr<IBreakIterator> InLineBreakPolicy);
+	SLATE_API FSlateTextBlockLayout(SWidget* InOwner, FTextBlockStyle InDefaultTextStyle, const TOptional<ETextShapingMethod> InTextShapingMethod, const TOptional<ETextFlowDirection> InTextFlowDirection, const FCreateSlateTextLayout& InCreateSlateTextLayout, TSharedRef<ITextLayoutMarshaller> InMarshaller, TSharedPtr<IBreakIterator> InLineBreakPolicy);
 
 	/**
 	 * Get the computed desired size for this layout, updating the internal cache as required
 	 */
-	FVector2D ComputeDesiredSize(const FWidgetDesiredSizeArgs& InWidgetArgs, const float InScale, const FTextBlockStyle& InTextStyle);
+	SLATE_API FVector2D ComputeDesiredSize(const FWidgetDesiredSizeArgs& InWidgetArgs, const float InScale, const FTextBlockStyle& InTextStyle);
 
 	PRAGMA_DISABLE_DEPRECATION_WARNINGS
 	/**
 	 * Get the computed desired size for this layout, updating the internal cache as required
 	 */
 	UE_DEPRECATED(5.0, "FWidgetArgs is deprecated. Upgrade to FWidgetDesiredSizeArgs instead.")
-	FVector2D ComputeDesiredSize(const FWidgetArgs& InWidgetArgs, const float InScale, const FTextBlockStyle& InTextStyle);
-	PRAGMA_ENABLE_DEPRECATION_WARNINGS
+	SLATE_API FVector2D ComputeDesiredSize(const FWidgetArgs& InWidgetArgs, const float InScale, const FTextBlockStyle& InTextStyle);
+	SLATE_API PRAGMA_ENABLE_DEPRECATION_WARNINGS
 
 	/**
 	 * Gets the last computed desired size.
@@ -121,61 +121,61 @@ public:
 	/**
 	 * Get the TextLayout scale.
 	 */
-	float GetLayoutScale() const;
+	SLATE_API float GetLayoutScale() const;
 
 
 	/**
 	 * Paint this layout, updating the internal cache as required
 	 */
-	int32 OnPaint(const FPaintArgs& InPaintArgs, const FGeometry& InAllottedGeometry, const FSlateRect& InClippingRect, FSlateWindowElementList& OutDrawElements, int32 LayerId, const FWidgetStyle& InWidgetStyle, bool bParentEnabled);
+	SLATE_API int32 OnPaint(const FPaintArgs& InPaintArgs, const FGeometry& InAllottedGeometry, const FSlateRect& InClippingRect, FSlateWindowElementList& OutDrawElements, int32 LayerId, const FWidgetStyle& InWidgetStyle, bool bParentEnabled);
 
 	/**
 	 * Force dirty the layout due to an external change that can't be picked up automatically by this cache
 	 */
-	void DirtyLayout();
+	SLATE_API void DirtyLayout();
 
 	/**
 	 * Force dirty the content due to an external change that can't be picked up automatically by this cache.
 	 * Also force dirties the layout. Will cause the contained text to be re-parsed by the marshaller next frame.
 	 */
-	void DirtyContent();
+	SLATE_API void DirtyContent();
 
 	/**
 	 * Override the text style used and immediately update the text layout (if required).
 	 * This can be used to override the text style after calling ComputeDesiredSize (eg, if you can only compute your text style in OnPaint)
 	 * Please note that changing the size or font used by the text may causing clipping issues until the next call to ComputeDesiredSize
 	 */
-	void OverrideTextStyle(const FTextBlockStyle& InTextStyle);
+	SLATE_API void OverrideTextStyle(const FTextBlockStyle& InTextStyle);
 
 	/**
 	 * Set the text shaping method that the internal text layout should use
 	 */
-	void SetTextShapingMethod(const TOptional<ETextShapingMethod>& InTextShapingMethod);
+	SLATE_API void SetTextShapingMethod(const TOptional<ETextShapingMethod>& InTextShapingMethod);
 
 	/**
 	 * Set the text flow direction that the internal text layout should use
 	 */
-	void SetTextFlowDirection(const TOptional<ETextFlowDirection>& InTextFlowDirection);
+	SLATE_API void SetTextFlowDirection(const TOptional<ETextFlowDirection>& InTextFlowDirection);
 
 	/**
 	 * Set the text overflow policy that the internal text layout should use
 	 */
-	void SetTextOverflowPolicy(const TOptional<ETextOverflowPolicy> InTextOverflowPolicy);
+	SLATE_API void SetTextOverflowPolicy(const TOptional<ETextOverflowPolicy> InTextOverflowPolicy);
 
 	/**
 	 * Set the information used to help identify who owns this text layout in the case of an error
 	 */
-	void SetDebugSourceInfo(const TAttribute<FString>& InDebugSourceInfo);
+	SLATE_API void SetDebugSourceInfo(const TAttribute<FString>& InDebugSourceInfo);
 
 	/**
 	 * Get the child widgets of this layout
 	 */
-	FChildren* GetChildren();
+	SLATE_API FChildren* GetChildren();
 
 	/**
 	 * Arrange any child widgets in this layout
 	 */
-	void ArrangeChildren(const FGeometry& AllottedGeometry, FArrangedChildren& ArrangedChildren) const;
+	SLATE_API void ArrangeChildren(const FGeometry& AllottedGeometry, FArrangedChildren& ArrangedChildren) const;
 
 private:
 	/** Updates the text layout to contain the given text */
diff --git a/Engine/Source/Runtime/Slate/Public/Widgets/Views/ITableRow.h b/Engine/Source/Runtime/Slate/Public/Widgets/Views/ITableRow.h
--- a/Engine/Source/Runtime/Slate/Public/Widgets/Views/ITableRow.h
+++ b/Engine/Source/Runtime/Slate/Public/Widgets/Views/ITableRow.h
@@ -11,7 +11,7 @@
 /**
  * Interface for table views to talk to their rows.
  */
-class SLATE_API ITableRow
+class ITableRow
 {
 public:
 
@@ -71,4 +71,4 @@ public:
 protected:
 	/** Called to query the selection mode for the row */
 	virtual ESelectionMode::Type GetSelectionMode() const = 0;
-};
\ No newline at end of file
+};
diff --git a/Engine/Source/Runtime/Slate/Public/Widgets/Views/SExpanderArrow.h b/Engine/Source/Runtime/Slate/Public/Widgets/Views/SExpanderArrow.h
--- a/Engine/Source/Runtime/Slate/Public/Widgets/Views/SExpanderArrow.h
+++ b/Engine/Source/Runtime/Slate/Public/Widgets/Views/SExpanderArrow.h
@@ -18,9 +18,9 @@ class SButton;
  * Expander arrow and indentation component that can be placed in a TableRow
  * of a TreeView. Intended for use by TMultiColumnRow in TreeViews.
  */
-class SLATE_API SExpanderArrow : public SCompoundWidget
+class SExpanderArrow : public SCompoundWidget
 {
-	SLATE_DECLARE_WIDGET(SExpanderArrow, SCompoundWidget)
+	SLATE_DECLARE_WIDGET_API(SExpanderArrow, SCompoundWidget, SLATE_API)
 public:
 
 	SLATE_BEGIN_ARGS( SExpanderArrow )
@@ -38,24 +38,24 @@ public:
 		SLATE_ATTRIBUTE(bool, ShouldDrawWires)
 	SLATE_END_ARGS()
 
-	SExpanderArrow();
-	void Construct( const FArguments& InArgs, const TSharedPtr<class ITableRow>& TableRow );
+	SLATE_API SExpanderArrow();
+	SLATE_API void Construct( const FArguments& InArgs, const TSharedPtr<class ITableRow>& TableRow );
 
 protected:
 
-	virtual int32 OnPaint(const FPaintArgs& Args, const FGeometry& AllottedGeometry, const FSlateRect& MyCullingRect, FSlateWindowElementList& OutDrawElements, int32 LayerId, const FWidgetStyle& InWidgetStyle, bool bParentEnabled) const override;
+	SLATE_API virtual int32 OnPaint(const FPaintArgs& Args, const FGeometry& AllottedGeometry, const FSlateRect& MyCullingRect, FSlateWindowElementList& OutDrawElements, int32 LayerId, const FWidgetStyle& InWidgetStyle, bool bParentEnabled) const override;
 
 	/** Invoked when the expanded button is clicked (toggle item expansion) */
-	FReply OnArrowClicked();
+	SLATE_API FReply OnArrowClicked();
 
 	/** @return Visible when has children; invisible otherwise */
-	EVisibility GetExpanderVisibility() const;
+	SLATE_API EVisibility GetExpanderVisibility() const;
 
 	/** @return the margin corresponding to how far this item is indented */
-	FMargin GetExpanderPadding() const;
+	SLATE_API FMargin GetExpanderPadding() const;
 
 	/** @return the name of an image that should be shown as the expander arrow */
-	const FSlateBrush* GetExpanderImage() const;
+	SLATE_API const FSlateBrush* GetExpanderImage() const;
 
 	TWeakPtr<class ITableRow> OwnerRowPtr;
 
diff --git a/Engine/Source/Runtime/Slate/Public/Widgets/Views/SHeaderRow.h b/Engine/Source/Runtime/Slate/Public/Widgets/Views/SHeaderRow.h
--- a/Engine/Source/Runtime/Slate/Public/Widgets/Views/SHeaderRow.h
+++ b/Engine/Source/Runtime/Slate/Public/Widgets/Views/SHeaderRow.h
@@ -93,7 +93,7 @@ DECLARE_DELEGATE_RetVal_TwoParams(FVector2D, FOnGetMaxRowSizeForColumn, const FN
 /**
  * The header that appears above lists and trees when they are showing multiple columns.
  */
-class SLATE_API SHeaderRow : public SBorder
+class SHeaderRow : public SBorder
 {
 public:
 	/** Describes a single column header */
@@ -309,70 +309,70 @@ public:
 		SLATE_EVENT( FSimpleDelegate, OnHiddenColumnsListChanged )
 	SLATE_END_ARGS()
 
-	void Construct( const FArguments& InArgs );
+	SLATE_API void Construct( const FArguments& InArgs );
 
 	/** Restore the columns to their original width */
-	void ResetColumnWidths();
+	SLATE_API void ResetColumnWidths();
 
 	/** @return the Columns driven by the column headers */
-	const TIndirectArray<FColumn>& GetColumns() const;
+	SLATE_API const TIndirectArray<FColumn>& GetColumns() const;
 
 	/** Adds a column to the header */
-	void AddColumn( const FColumn::FArguments& NewColumnArgs );
-	void AddColumn( FColumn& NewColumn );
+	SLATE_API void AddColumn( const FColumn::FArguments& NewColumnArgs );
+	SLATE_API void AddColumn( FColumn& NewColumn );
 
 	/** Inserts a column at the specified index in the header */
-	void InsertColumn( const FColumn::FArguments& NewColumnArgs, int32 InsertIdx );
-	void InsertColumn( FColumn& NewColumn, int32 InsertIdx );
+	SLATE_API void InsertColumn( const FColumn::FArguments& NewColumnArgs, int32 InsertIdx );
+	SLATE_API void InsertColumn( FColumn& NewColumn, int32 InsertIdx );
 
 	/** Removes a column from the header */
-	void RemoveColumn( const FName& InColumnId );
+	SLATE_API void RemoveColumn( const FName& InColumnId );
 
 	/** Force refreshing of the column widgets*/
-	void RefreshColumns();
+	SLATE_API void RefreshColumns();
 
 	/** Removes all columns from the header */
-	void ClearColumns();
+	SLATE_API void ClearColumns();
 
-	void SetAssociatedVerticalScrollBar( const TSharedRef< SScrollBar >& ScrollBar, const float ScrollBarSize );
+	SLATE_API void SetAssociatedVerticalScrollBar( const TSharedRef< SScrollBar >& ScrollBar, const float ScrollBarSize );
 
 	/** Sets the column, with the specified name, to the desired width */
-	void SetColumnWidth( const FName& InColumnId, float InWidth );
+	SLATE_API void SetColumnWidth( const FName& InColumnId, float InWidth );
 
 	/** Will return the size for this row at the specified slot index */
-	FVector2D GetRowSizeForSlotIndex(int32 SlotIndex) const;
+	SLATE_API FVector2D GetRowSizeForSlotIndex(int32 SlotIndex) const;
 
 	/** Simple function to set the delegate to fetch the max row size for column id */
 	void SetOnGetMaxRowSizeForColumn(const FOnGetMaxRowSizeForColumn& Delegate) { OnGetMaxRowSizeForColumn = Delegate; }
 
 	/** @return The columns id of the header widgets that were not generated */
-	TArray<FName> GetHiddenColumnIds() const;
+	SLATE_API TArray<FName> GetHiddenColumnIds() const;
 
 	/** @return is the header widget should be generated */
-	bool ShouldGeneratedColumn( const FName& InColumnId ) const;
+	SLATE_API bool ShouldGeneratedColumn( const FName& InColumnId ) const;
 
 	/** @return is the header widget with the column id is generated */
-	bool IsColumnGenerated( const FName& InColumnId ) const;
+	SLATE_API bool IsColumnGenerated( const FName& InColumnId ) const;
 
 	/** @return true if the header widget with the column id is visible */
-	bool IsColumnVisible( const FName& InColumnId ) const;
+	SLATE_API bool IsColumnVisible( const FName& InColumnId ) const;
 
 	//~ Begin SWidget interface
-	virtual FReply OnMouseButtonUp( const FGeometry& MyGeometry, const FPointerEvent& MouseEvent ) override;
+	SLATE_API virtual FReply OnMouseButtonUp( const FGeometry& MyGeometry, const FPointerEvent& MouseEvent ) override;
 	//~ End SWidget interface
 
 	/** Show or Hide a generated column */
-	void SetShowGeneratedColumn(const FName& InColumnIdm, bool InShow = true);
+	SLATE_API void SetShowGeneratedColumn(const FName& InColumnIdm, bool InShow = true);
 
 
 private:
 
 	/** Regenerates all widgets in the header */
-	void RegenerateWidgets();
+	SLATE_API void RegenerateWidgets();
 
-	void OnGenerateSelectColumnsSubMenu(FMenuBuilder& InSubMenuBuilder);
-	void ToggleGeneratedColumn(FName ColumnId);
-	ECheckBoxState GetGeneratedColumnCheckedState(FName ColumnId) const;
+	SLATE_API void OnGenerateSelectColumnsSubMenu(FMenuBuilder& InSubMenuBuilder);
+	SLATE_API void ToggleGeneratedColumn(FName ColumnId);
+	SLATE_API ECheckBoxState GetGeneratedColumnCheckedState(FName ColumnId) const;
 
 	/** Information about the various columns */
 	TIndirectArray<FColumn> Columns;
diff --git a/Engine/Source/Runtime/Slate/Public/Widgets/Views/STableViewBase.h b/Engine/Source/Runtime/Slate/Public/Widgets/Views/STableViewBase.h
--- a/Engine/Source/Runtime/Slate/Public/Widgets/Views/STableViewBase.h
+++ b/Engine/Source/Runtime/Slate/Public/Widgets/Views/STableViewBase.h
@@ -58,11 +58,11 @@ DECLARE_DELEGATE_OneParam(
 	double );	/** Scroll offset from the beginning of the list in items */
 
 /** Abstracts away the need to distinguish between X or Y when calculating table layout elements */
-struct SLATE_API FTableViewDimensions
+struct FTableViewDimensions
 {
-	FTableViewDimensions(EOrientation InOrientation);
-	FTableViewDimensions(EOrientation InOrientation, float X, float Y);
-	FTableViewDimensions(EOrientation InOrientation, const UE::Slate::FDeprecateVector2DParameter& Size);
+	SLATE_API FTableViewDimensions(EOrientation InOrientation);
+	SLATE_API FTableViewDimensions(EOrientation InOrientation, float X, float Y);
+	SLATE_API FTableViewDimensions(EOrientation InOrientation, const UE::Slate::FDeprecateVector2DParameter& Size);
 
 	UE::Slate::FDeprecateVector2DResult ToVector2D() const
 	{
@@ -86,140 +86,140 @@ struct SLATE_API FTableViewDimensions
 /**
  * Contains ListView functionality that does not depend on the type of data being observed by the ListView.
  */
-class SLATE_API STableViewBase
+class STableViewBase
 	: public SCompoundWidget
 	, public IScrollableWidget
 {
 public:
 
 	/** Create the child widgets that comprise the list */
-	void ConstructChildren( const TAttribute<float>& InItemWidth, const TAttribute<float>& InItemHeight, const TAttribute<EListItemAlignment>& InItemAlignment, const TSharedPtr<SHeaderRow>& InHeaderRow, const TSharedPtr<SScrollBar>& InScrollBar, EOrientation InScrollOrientation, const FOnTableViewScrolled& InOnTableViewScrolled, const FScrollBarStyle* InScrollBarStyle = nullptr, const bool bInPreventThrottling = false );
+	SLATE_API void ConstructChildren( const TAttribute<float>& InItemWidth, const TAttribute<float>& InItemHeight, const TAttribute<EListItemAlignment>& InItemAlignment, const TSharedPtr<SHeaderRow>& InHeaderRow, const TSharedPtr<SScrollBar>& InScrollBar, EOrientation InScrollOrientation, const FOnTableViewScrolled& InOnTableViewScrolled, const FScrollBarStyle* InScrollBarStyle = nullptr, const bool bInPreventThrottling = false );
 
 	/** Sets the item height */
-	void SetItemHeight(TAttribute<float> Height);
+	SLATE_API void SetItemHeight(TAttribute<float> Height);
 
 	/** Sets the item width */
-	void SetItemWidth(TAttribute<float> Width);
+	SLATE_API void SetItemWidth(TAttribute<float> Width);
 
 	/**
 	 * Invoked by the scrollbar when the user scrolls.
 	 *
 	 * @param InScrollOffsetFraction  The location to which the user scrolled as a fraction (between 0 and 1) of total height of the content.
 	 */
-	void ScrollBar_OnUserScrolled( float InScrollOffsetFraction );
+	SLATE_API void ScrollBar_OnUserScrolled( float InScrollOffsetFraction );
 
 	/** @return The number of Widgets we currently have generated. */
-	int32 GetNumGeneratedChildren() const;
+	SLATE_API int32 GetNumGeneratedChildren() const;
 
-	TSharedPtr<SHeaderRow> GetHeaderRow() const;
+	SLATE_API TSharedPtr<SHeaderRow> GetHeaderRow() const;
 
 	/** @return Returns true if the user is currently interactively scrolling the view by holding
 		        the right mouse button and dragging. */
-	bool IsRightClickScrolling() const;
+	SLATE_API bool IsRightClickScrolling() const;
 
 	/** @return Returns true if the user is currently interactively scrolling the view by holding
 		        either mouse button and dragging. */
-	bool IsUserScrolling() const;
+	SLATE_API bool IsUserScrolling() const;
 
 	/**
 	 * Mark the list as dirty so that it will refresh its widgets on next tick.
 	 * Note that refreshing will only generate/release widgets as needed from any deltas in the list items source.
 	 */
-	virtual void RequestListRefresh();
+	SLATE_API virtual void RequestListRefresh();
 
 	/** Completely wipe existing widgets and fully regenerate them on next tick. */
 	virtual void RebuildList() = 0;
 
 	/** Return true if there is currently a refresh pending, false otherwise */
-	bool IsPendingRefresh() const;
+	SLATE_API bool IsPendingRefresh() const;
 
 	/** Is this list backing a tree or just a standalone list */
 	const ETableViewMode::Type TableViewMode;
 
 	/** Scrolls the view to the top */
-	void ScrollToTop();
+	SLATE_API void ScrollToTop();
 
 	/** Scrolls the view to the bottom */
-	void ScrollToBottom();
+	SLATE_API void ScrollToBottom();
 
 	/** Returns whether the attached scrollbar is scrolling */
-	bool IsScrolling() const;
+	SLATE_API bool IsScrolling() const;
 
 	/** Gets the scroll offset of this view (in items) */
-	float GetScrollOffset() const;
+	SLATE_API float GetScrollOffset() const;
 
 	/** Set the scroll offset of this view (in items) */
-	void SetScrollOffset( const float InScrollOffset );
+	SLATE_API void SetScrollOffset( const float InScrollOffset );
 
 	/** Reset the inertial scroll velocity accumulated in the InertialScrollManager */
-	void EndInertialScrolling();
+	SLATE_API void EndInertialScrolling();
 
 	/** Add the scroll offset of this view (in items) */
-	void AddScrollOffset(const float InScrollOffsetDelta, bool RefreshList = false);
+	SLATE_API void AddScrollOffset(const float InScrollOffsetDelta, bool RefreshList = false);
 
-	EVisibility GetScrollbarVisibility() const;
+	SLATE_API EVisibility GetScrollbarVisibility() const;
 
-	void SetScrollbarVisibility(const EVisibility InVisibility);
+	SLATE_API void SetScrollbarVisibility(const EVisibility InVisibility);
 
 	/** Returns true if scrolling is possible; false if the view is big enough to fit all the content. */
-	bool IsScrollbarNeeded() const;
+	SLATE_API bool IsScrollbarNeeded() const;
 
 	/** Sets the fixed offset in items to always apply to the top/left (depending on orientation) of the list. */
-	void SetFixedLineScrollOffset(TOptional<double> InFixedLineScrollOffset);
+	SLATE_API void SetFixedLineScrollOffset(TOptional<double> InFixedLineScrollOffset);
 
 	/** Sets whether the list should lerp between scroll offsets or jump instantly between them. */
-	void SetIsScrollAnimationEnabled(bool bInEnableScrollAnimation);
+	SLATE_API void SetIsScrollAnimationEnabled(bool bInEnableScrollAnimation);
 
 	/** Sets whether to permit overscroll on this list view */
-	void SetAllowOverscroll(EAllowOverscroll InAllowOverscroll);
+	SLATE_API void SetAllowOverscroll(EAllowOverscroll InAllowOverscroll);
 
 	/** Enables/disables being able to scroll with the right mouse button. */
-	void SetIsRightClickScrollingEnabled(const bool bInEnableRightClickScrolling);
+	SLATE_API void SetIsRightClickScrollingEnabled(const bool bInEnableRightClickScrolling);
 
 	/** Sets the multiplier applied when wheel scrolling. Higher numbers will cover more distance per click of the wheel. */
-	void SetWheelScrollMultiplier(float NewWheelScrollMultiplier);
+	SLATE_API void SetWheelScrollMultiplier(float NewWheelScrollMultiplier);
 
 	/** Sets the Background Brush */
-	void SetBackgroundBrush(const TAttribute<const FSlateBrush*>& InBackgroundBrush);
+	SLATE_API void SetBackgroundBrush(const TAttribute<const FSlateBrush*>& InBackgroundBrush);
 
 public:
 
 	// SWidget interface
 
-	virtual void OnFocusLost( const FFocusEvent& InFocusEvent ) override;
-	virtual void OnMouseCaptureLost(const FCaptureLostEvent& CaptureLostEvent) override;
-	virtual bool SupportsKeyboardFocus() const override;
-	virtual void Tick( const FGeometry& AllottedGeometry, const double InCurrentTime, const float InDeltaTime ) override;
-	virtual FReply OnPreviewMouseButtonDown( const FGeometry& MyGeometry, const FPointerEvent& MouseEvent ) override;
-	virtual FReply OnMouseButtonDown( const FGeometry& MyGeometry, const FPointerEvent& MouseEvent ) override;
-	virtual FReply OnMouseButtonDoubleClick( const FGeometry& InMyGeometry, const FPointerEvent& InMouseEvent ) override;
-	virtual FReply OnMouseButtonUp( const FGeometry& MyGeometry, const FPointerEvent& MouseEvent ) override;
-	virtual void OnMouseEnter( const FGeometry& MyGeometry, const FPointerEvent& MouseEvent ) override;
-	virtual void OnMouseLeave( const FPointerEvent& MouseEvent ) override;
-	virtual FReply OnMouseMove( const FGeometry& MyGeometry, const FPointerEvent& MouseEvent ) override;
-	virtual FReply OnMouseWheel( const FGeometry& MyGeometry, const FPointerEvent& MouseEvent ) override;
-	virtual FReply OnKeyDown( const FGeometry& MyGeometry, const FKeyEvent& InKeyEvent ) override;
-	virtual FCursorReply OnCursorQuery( const FGeometry& MyGeometry, const FPointerEvent& CursorEvent ) const override;
-	virtual int32 OnPaint(const FPaintArgs& Args, const FGeometry& AllottedGeometry, const FSlateRect& MyCullingRect, FSlateWindowElementList& OutDrawElements, int32 LayerId, const FWidgetStyle& InWidgetStyle, bool bParentEnabled) const override;
-	virtual bool ComputeVolatility() const override;
-	virtual FReply OnTouchStarted( const FGeometry& MyGeometry, const FPointerEvent& InTouchEvent ) override;
-	virtual FReply OnTouchMoved( const FGeometry& MyGeometry, const FPointerEvent& InTouchEvent ) override;
-	virtual FReply OnTouchEnded( const FGeometry& MyGeometry, const FPointerEvent& InTouchEvent ) override;
+	SLATE_API virtual void OnFocusLost( const FFocusEvent& InFocusEvent ) override;
+	SLATE_API virtual void OnMouseCaptureLost(const FCaptureLostEvent& CaptureLostEvent) override;
+	SLATE_API virtual bool SupportsKeyboardFocus() const override;
+	SLATE_API virtual void Tick( const FGeometry& AllottedGeometry, const double InCurrentTime, const float InDeltaTime ) override;
+	SLATE_API virtual FReply OnPreviewMouseButtonDown( const FGeometry& MyGeometry, const FPointerEvent& MouseEvent ) override;
+	SLATE_API virtual FReply OnMouseButtonDown( const FGeometry& MyGeometry, const FPointerEvent& MouseEvent ) override;
+	SLATE_API virtual FReply OnMouseButtonDoubleClick( const FGeometry& InMyGeometry, const FPointerEvent& InMouseEvent ) override;
+	SLATE_API virtual FReply OnMouseButtonUp( const FGeometry& MyGeometry, const FPointerEvent& MouseEvent ) override;
+	SLATE_API virtual void OnMouseEnter( const FGeometry& MyGeometry, const FPointerEvent& MouseEvent ) override;
+	SLATE_API virtual void OnMouseLeave( const FPointerEvent& MouseEvent ) override;
+	SLATE_API virtual FReply OnMouseMove( const FGeometry& MyGeometry, const FPointerEvent& MouseEvent ) override;
+	SLATE_API virtual FReply OnMouseWheel( const FGeometry& MyGeometry, const FPointerEvent& MouseEvent ) override;
+	SLATE_API virtual FReply OnKeyDown( const FGeometry& MyGeometry, const FKeyEvent& InKeyEvent ) override;
+	SLATE_API virtual FCursorReply OnCursorQuery( const FGeometry& MyGeometry, const FPointerEvent& CursorEvent ) const override;
+	SLATE_API virtual int32 OnPaint(const FPaintArgs& Args, const FGeometry& AllottedGeometry, const FSlateRect& MyCullingRect, FSlateWindowElementList& OutDrawElements, int32 LayerId, const FWidgetStyle& InWidgetStyle, bool bParentEnabled) const override;
+	SLATE_API virtual bool ComputeVolatility() const override;
+	SLATE_API virtual FReply OnTouchStarted( const FGeometry& MyGeometry, const FPointerEvent& InTouchEvent ) override;
+	SLATE_API virtual FReply OnTouchMoved( const FGeometry& MyGeometry, const FPointerEvent& InTouchEvent ) override;
+	SLATE_API virtual FReply OnTouchEnded( const FGeometry& MyGeometry, const FPointerEvent& InTouchEvent ) override;
 
 public:
 
 	// IScrollableWidget interface
 
-	virtual FVector2D GetScrollDistance() override;
-	virtual FVector2D GetScrollDistanceRemaining() override;
-	virtual TSharedRef<class SWidget> GetScrollWidget() override;
+	SLATE_API virtual FVector2D GetScrollDistance() override;
+	SLATE_API virtual FVector2D GetScrollDistanceRemaining() override;
+	SLATE_API virtual TSharedRef<class SWidget> GetScrollWidget() override;
 
 protected:
 
-	STableViewBase( ETableViewMode::Type InTableViewMode );
+	SLATE_API STableViewBase( ETableViewMode::Type InTableViewMode );
 
 	/** Returns the "true" scroll offset where the list will ultimately settle (and may already be). */
-	double GetTargetScrollOffset() const;
+	SLATE_API double GetTargetScrollOffset() const;
 
 	/**
 	 * Scroll the list view by some number of screen units.
@@ -230,7 +230,7 @@ protected:
 	 *
 	 * @return The amount actually scrolled in items
 	 */
-	virtual float ScrollBy( const FGeometry& MyGeometry, float ScrollByAmount, EAllowOverscroll InAllowOverscroll );
+	SLATE_API virtual float ScrollBy( const FGeometry& MyGeometry, float ScrollByAmount, EAllowOverscroll InAllowOverscroll );
 
 	/**
 	 * Scroll the view to an offset and resets the inertial scroll velocity 
@@ -239,60 +239,60 @@ protected:
 	 *
 	 * @return The amount actually scrolled
 	 */
-	virtual float ScrollTo( float InScrollOffset);
+	SLATE_API virtual float ScrollTo( float InScrollOffset);
 
 	/** Insert WidgetToInsert at the top of the view. */
-	void InsertWidget( const TSharedRef<ITableRow> & WidgetToInset );
+	SLATE_API void InsertWidget( const TSharedRef<ITableRow> & WidgetToInset );
 
 	/** Add a WidgetToAppend to the bottom of the view. */
-	void AppendWidget( const TSharedRef<ITableRow>& WidgetToAppend );
+	SLATE_API void AppendWidget( const TSharedRef<ITableRow>& WidgetToAppend );
 
-	const FChildren* GetConstructedTableItems() const;
+	SLATE_API const FChildren* GetConstructedTableItems() const;
 
 	/**
 	 * Remove all the widgets from the view.
 	 */
-	void ClearWidgets();
+	SLATE_API void ClearWidgets();
 
 	/** Insert WidgetToInsert at the top of the pinned view. */
-	void InsertPinnedWidget(const TSharedRef<SWidget>& WidgetToInset);
+	SLATE_API void InsertPinnedWidget(const TSharedRef<SWidget>& WidgetToInset);
 
 	/** Add a WidgetToAppend to the bottom of the pinned view. */
-	void AppendPinnedWidget(const TSharedRef<SWidget>& WidgetToAppend);
+	SLATE_API void AppendPinnedWidget(const TSharedRef<SWidget>& WidgetToAppend);
 
 	/**
 	 * Remove all the pinned widgets from the view.
 	 */
-	void ClearPinnedWidgets();
+	SLATE_API void ClearPinnedWidgets();
 
 	/**
 	 * Get the uniform item width.
 	 */
-	float GetItemWidth() const;
+	SLATE_API float GetItemWidth() const;
 
 	/**
 	 * Get the uniform item height that is enforced by ListViews.
 	 */
-	float GetItemHeight() const;
+	SLATE_API float GetItemHeight() const;
 
 	/**
 	* Get the uniform item
 	*/
-	UE::Slate::FDeprecateVector2DResult GetItemSize() const;
+	SLATE_API UE::Slate::FDeprecateVector2DResult GetItemSize() const;
 
 	/** @return the number of items that can fit on the screen */
-	virtual float GetNumLiveWidgets() const;
+	SLATE_API virtual float GetNumLiveWidgets() const;
 
 	/**
 	 * Get the number of items that can fit in the view along the line axis (orthogonal to the scroll axis) before creating a new line.
 	 * Default is 1, but may be more in subclasses (like STileView)
 	 */
-	virtual int32 GetNumItemsPerLine() const;
+	SLATE_API virtual int32 GetNumItemsPerLine() const;
 
 	/**
 	 * Get the offset of the first list item.
 	 */
-	virtual float GetFirstLineScrollOffset() const;
+	SLATE_API virtual float GetFirstLineScrollOffset() const;
 
 	/*
 	 * Right click down
@@ -302,12 +302,12 @@ protected:
 	/**
 	 * Opens a context menu as the result of a right click if OnContextMenuOpening is bound and we are not right click scrolling.
 	 */
-	virtual void OnRightMouseButtonUp(const FPointerEvent& MouseEvent);
+	SLATE_API virtual void OnRightMouseButtonUp(const FPointerEvent& MouseEvent);
 	
 	/**
 	 * Get the scroll rate in items that best approximates a constant physical scroll rate.
 	 */
-	float GetScrollRateInItems() const;
+	SLATE_API float GetScrollRateInItems() const;
 
 	/**
 	 * Remove any items that are no longer in the list from the selection set.
@@ -315,10 +315,10 @@ protected:
 	virtual void UpdateSelectionSet() = 0;
 
 	/** Internal request for a layout update on the next tick (i.e. a refresh without implication that the source items have changed) */
-	void RequestLayoutRefresh();
+	SLATE_API void RequestLayoutRefresh();
 
 	/** Information about the outcome of the WidgetRegeneratePass */
-	struct SLATE_API FReGenerateResults
+	struct FReGenerateResults
 	{
 		FReGenerateResults(double InNewScrollOffset, double InLengthGenerated, double InItemsOnScreen, bool AtEndOfList)
 			: NewScrollOffset(InNewScrollOffset)
@@ -349,9 +349,9 @@ protected:
 	virtual int32 GetNumItemsBeingObserved() const = 0;
 
 	/** @return how many pinned items are in the table */
-	int32 GetNumPinnedItems() const;
+	SLATE_API int32 GetNumPinnedItems() const;
 
-	EVisibility GetPinnedItemsVisiblity() const;
+	SLATE_API EVisibility GetPinnedItemsVisiblity() const;
 
 	enum class EScrollIntoViewResult
 	{
@@ -381,10 +381,10 @@ protected:
 	virtual void NotifyFinishedScrolling() = 0;
 
 	/** Util Function so templates classes don't need to include SlateApplication */
-	void NavigateToWidget(const uint32 UserIndex, const TSharedPtr<SWidget>& NavigationDestination, ENavigationSource NavigationSource = ENavigationSource::FocusedWidget) const;
+	SLATE_API void NavigateToWidget(const uint32 UserIndex, const TSharedPtr<SWidget>& NavigationDestination, ENavigationSource NavigationSource = ENavigationSource::FocusedWidget) const;
 
 	/** Util function to find the child index under the given position. */
-	int32 FindChildUnderPosition(FArrangedChildren& ArrangedChildren, const FVector2D& ArrangedSpacePosition) const;
+	SLATE_API int32 FindChildUnderPosition(FArrangedChildren& ArrangedChildren, const FVector2D& ArrangedSpacePosition) const;
 
 	/** The panel which holds the visible widgets in this list */
 	TSharedPtr< SListPanel > ItemsPanel;
@@ -467,7 +467,7 @@ protected:
 	EOrientation Orientation = Orient_Vertical;
 
 	/** Passing over the clipping to SListPanel */
-	virtual void OnClippingChanged() override;
+	SLATE_API virtual void OnClippingChanged() override;
 
 	/** Brush resource representing the background area of the view */
 	FInvalidatableBrushAttribute BackgroundBrush;
@@ -475,13 +475,13 @@ protected:
 protected:
 
 	/** Check whether the current state of the table warrants inertial scroll by the specified amount */
-	bool CanUseInertialScroll( float ScrollAmount ) const;
+	SLATE_API bool CanUseInertialScroll( float ScrollAmount ) const;
 
 	/** Active timer to update the inertial scroll */
-	EActiveTimerReturnType UpdateInertialScroll(double InCurrentTime, float InDeltaTime);
+	SLATE_API EActiveTimerReturnType UpdateInertialScroll(double InCurrentTime, float InDeltaTime);
 
 	/** One-off active timer to refresh the contents of the table as needed */
-	EActiveTimerReturnType EnsureTickToRefresh(double InCurrentTime, float InDeltaTime);
+	SLATE_API EActiveTimerReturnType EnsureTickToRefresh(double InCurrentTime, float InDeltaTime);
 
 	/** Whether the active timer to update the inertial scrolling is currently registered */
 	bool bIsScrollingActiveTimerRegistered;
@@ -509,4 +509,4 @@ namespace TableViewHelpers
 	 * view, and therefore does not to know about parent wires.
 	 */
 	SLATE_API const TBitArray<>& GetEmptyBitArray();
-}
\ No newline at end of file
+}
diff --git a/Engine/Source/Runtime/SlateCore/Public/Widgets/SlateControlledConstruction.h b/Engine/Source/Runtime/SlateCore/Public/Widgets/SlateControlledConstruction.h
--- a/Engine/Source/Runtime/SlateCore/Public/Widgets/SlateControlledConstruction.h
+++ b/Engine/Source/Runtime/SlateCore/Public/Widgets/SlateControlledConstruction.h
@@ -17,6 +17,9 @@ namespace SharedPointerInternals
 
 /** */
 #define SLATE_DECLARE_WIDGET(WidgetType, ParentType) \
+	SLATE_DECLARE_WIDGET_API(WidgetType, ParentType, NO_API)
+
+#define SLATE_DECLARE_WIDGET_API(WidgetType, ParentType, ModuleApiDefine) \
 	private: \
 		using Super = ParentType; \
 		using ThisClass = WidgetType; \
@@ -31,7 +34,7 @@ namespace SharedPointerInternals
 			static FSlateWidgetClassData WidgetClassDataInstance = FSlateWidgetClassData(TIdentity<ParentType>(), #WidgetType, &WidgetType::PrivateRegisterAttributes); \
 			return WidgetClassDataInstance; \
 		} \
-		static void PrivateRegisterAttributes(FSlateAttributeInitializer&); \
+		static ModuleApiDefine void PrivateRegisterAttributes(FSlateAttributeInitializer&); \
 	public: \
 		static const FSlateWidgetClassData& StaticWidgetClass() { return GetPrivateWidgetClass(); } \
 		virtual const FSlateWidgetClassData& GetWidgetClass() const override { return GetPrivateWidgetClass(); } \
@@ -42,7 +45,7 @@ namespace SharedPointerInternals
 	FSlateWidgetClassRegistration ClassRegistration__##WidgetType = WidgetType::StaticWidgetClass();
 
 /** */
-class SLATECORE_API FSlateWidgetClassData
+class FSlateWidgetClassData
 {
 private:
 	friend FSlateControlledConstruction;
@@ -78,7 +81,7 @@ struct FSlateWidgetClassRegistration
 };
 
 /** */
-class SLATECORE_API FSlateControlledConstruction
+class FSlateControlledConstruction
 {
 public:
 	FSlateControlledConstruction() = default;
