Fixed transforms not being transacted when using the slider from the details panel
- As of cl 25403899, SNumericVectorInputBox::OnBeginSliderMovement isn't being triggered and has been replaced by 3 individual delegates for each axis.
- Undeprecated OnBegin/EndSliderMovement in SRotationInputBox. It was being routed to OnPitchBegin/EndSliderMovement meaning that it was ignored for the other axis.
- Supported both the global and component specific events for both SVectorInputBox and SRotationInputBox.
- Removed the specific delegate types that were added to SVectorInputBox so that it stays aligned with SRotationInputBox.

diff --git a/Engine/Source/Developer/AnimationWidgets/Public/SAdvancedTransformInputBox.h b/Engine/Source/Developer/AnimationWidgets/Public/SAdvancedTransformInputBox.h
--- a/Engine/Source/Developer/AnimationWidgets/Public/SAdvancedTransformInputBox.h
+++ b/Engine/Source/Developer/AnimationWidgets/Public/SAdvancedTransformInputBox.h
@@ -478,22 +478,22 @@ public:
 					); 
 				}
 
-				typename SNumericVectorInputBox3::FOnBeginSliderMovement XBeginSlide, YBeginSlide, ZBeginSlide;
+				FSimpleDelegate XBeginSlide, YBeginSlide, ZBeginSlide;
 				if(OnBeginSliderMovement.IsBound())
 				{
-					XBeginSlide = SNumericVectorInputBox3::FOnBeginSliderMovement::CreateLambda(
+					XBeginSlide = FSimpleDelegate::CreateLambda(
 						[OnBeginSliderMovement, InComponent]()
 						{
 							OnBeginSliderMovement.Execute(InComponent, ESlateRotationRepresentation::EulerXYZ, ESlateTransformSubComponent::X);
 						}
 					); 
-					YBeginSlide = SNumericVectorInputBox3::FOnBeginSliderMovement::CreateLambda(
+					YBeginSlide = FSimpleDelegate::CreateLambda(
 						[OnBeginSliderMovement, InComponent]()
 						{
 							OnBeginSliderMovement.Execute(InComponent, ESlateRotationRepresentation::EulerXYZ, ESlateTransformSubComponent::Y);
 						}
 					); 
-					ZBeginSlide = SNumericVectorInputBox3::FOnBeginSliderMovement::CreateLambda(
+					ZBeginSlide = FSimpleDelegate::CreateLambda(
 						[OnBeginSliderMovement, InComponent]()
 						{
 							OnBeginSliderMovement.Execute(InComponent, ESlateRotationRepresentation::EulerXYZ, ESlateTransformSubComponent::Z);
@@ -501,22 +501,22 @@ public:
 					); 
 				}
 
-				typename SNumericVectorInputBox3::FOnEndSliderMovement XEndSlide, YEndSlide, ZEndSlide;
+				typename SNumericVectorInputBox3::FOnNumericValueChanged XEndSlide, YEndSlide, ZEndSlide;
 				if(OnEndSliderMovement.IsBound())
 				{
-					XEndSlide = SNumericVectorInputBox3::FOnEndSliderMovement::CreateLambda(
+					XEndSlide = SNumericVectorInputBox3::FOnNumericValueChanged::CreateLambda(
 						[OnEndSliderMovement, InComponent](NumericType InValue)
 						{
 							OnEndSliderMovement.Execute(InComponent, ESlateRotationRepresentation::EulerXYZ, ESlateTransformSubComponent::X, InValue);
 						}
 					); 
-					YEndSlide = SNumericVectorInputBox3::FOnEndSliderMovement::CreateLambda(
+					YEndSlide = SNumericVectorInputBox3::FOnNumericValueChanged::CreateLambda(
 						[OnEndSliderMovement, InComponent](NumericType InValue)
 						{
 							OnEndSliderMovement.Execute(InComponent, ESlateRotationRepresentation::EulerXYZ, ESlateTransformSubComponent::Y, InValue);
 						}
 					); 
-					ZEndSlide = SNumericVectorInputBox3::FOnEndSliderMovement::CreateLambda(
+					ZEndSlide = SNumericVectorInputBox3::FOnNumericValueChanged::CreateLambda(
 						[OnEndSliderMovement, InComponent](NumericType InValue)
 						{
 							OnEndSliderMovement.Execute(InComponent, ESlateRotationRepresentation::EulerXYZ, ESlateTransformSubComponent::Z, InValue);
diff --git a/Engine/Source/Runtime/Slate/Public/Widgets/Input/SRotatorInputBox.h b/Engine/Source/Runtime/Slate/Public/Widgets/Input/SRotatorInputBox.h
--- a/Engine/Source/Runtime/Slate/Public/Widgets/Input/SRotatorInputBox.h
+++ b/Engine/Source/Runtime/Slate/Public/Widgets/Input/SRotatorInputBox.h
@@ -89,25 +89,11 @@ public:
 		/** Called when the roll value is committed */
 		SLATE_EVENT( FOnNumericValueCommitted, OnRollCommitted )
 
-		FSimpleDelegate ConvertOnBeginSliderMovement(const FSimpleDelegate& LegacyDelegate)
-		{
-			return FSimpleDelegate::CreateLambda([LegacyDelegate]()
-			{
-				return LegacyDelegate.Execute();
-			});
-		}
 		/** Called when the slider begins to move on any axis */
-		SLATE_EVENT_DEPRECATED(5.3, "Use the new OnPitchBeginSliderMovement()/OnYawBeginSliderMovement()/OnRollBeginSliderMovement() event instead", FSimpleDelegate, OnBeginSliderMovement, OnPitchBeginSliderMovement, ConvertOnBeginSliderMovement)
+		SLATE_EVENT(FSimpleDelegate, OnBeginSliderMovement)
 
-		FOnNumericValueChanged ConvertOnEndSliderMovement(const FOnNumericValueChanged& LegacyDelegate)
-		{
-			return FOnNumericValueChanged::CreateLambda([LegacyDelegate](const NumericType& Value)
-			{
-				return LegacyDelegate.Execute(Value);
-			});
-		}
 		/** Called when the slider for any axis is released */
-		SLATE_EVENT_DEPRECATED(5.3, "Use the new OnPitchEndSliderMovement()/OnYawEndSliderMovement()/OnRollEndSliderMovement() event instead", FOnNumericValueChanged, OnEndSliderMovement, OnPitchEndSliderMovement, ConvertOnEndSliderMovement)
+		SLATE_EVENT(FOnNumericValueChanged, OnEndSliderMovement)
 	
 		/** Called when the slider begins to move on pitch */
 		SLATE_EVENT( FSimpleDelegate, OnPitchBeginSliderMovement )
@@ -214,8 +200,8 @@ public:
 				.Value( InArgs._Pitch )
 				.OnValueChanged( InArgs._OnPitchChanged )
 				.OnValueCommitted( InArgs._OnPitchCommitted )
-				.OnBeginSliderMovement( InArgs._OnPitchBeginSliderMovement )
-				.OnEndSliderMovement( InArgs._OnPitchEndSliderMovement )
+				.OnBeginSliderMovement( CreatePerComponentSliderMovementEvent( InArgs._OnBeginSliderMovement, InArgs._OnPitchBeginSliderMovement ) )
+				.OnEndSliderMovement( CreatePerComponentSliderMovementEvent< FOnNumericValueChanged, NumericType >( InArgs._OnEndSliderMovement, InArgs._OnPitchEndSliderMovement ) )
 				.UndeterminedString( NSLOCTEXT("SRotatorInputBox", "MultipleValues", "Multiple Values") )
 				.ToolTipText_Lambda([PitchAttr = InArgs._Pitch]
 				{
@@ -264,7 +250,27 @@ public:
 			]
 		];
 
-	}	
+	}
+
+	/**
+	 * Creates a lambda to react to a begin/end slider movement event
+	 */
+	template<typename EventType, typename... ArgsType>
+	EventType CreatePerComponentSliderMovementEvent(
+		const EventType OnSliderMovement,
+		const EventType OnComponentSliderMovement)
+	{
+		if(OnSliderMovement.IsBound())
+		{
+			return EventType::CreateLambda(
+				[OnSliderMovement, OnComponentSliderMovement](ArgsType... Args)
+				{
+					OnSliderMovement.ExecuteIfBound(Args...);
+					OnComponentSliderMovement.ExecuteIfBound(Args...);
+				});
+		}
+		return OnComponentSliderMovement;
+	}
 };
 
 /**
diff --git a/Engine/Source/Runtime/Slate/Public/Widgets/Input/SVectorInputBox.h b/Engine/Source/Runtime/Slate/Public/Widgets/Input/SVectorInputBox.h
--- a/Engine/Source/Runtime/Slate/Public/Widgets/Input/SVectorInputBox.h
+++ b/Engine/Source/Runtime/Slate/Public/Widgets/Input/SVectorInputBox.h
@@ -40,12 +40,6 @@ public:
 	/** Delegate to constrain the vector during a change */
 	DECLARE_DELEGATE_ThreeParams(FOnConstrainVector, int32 /* Component */, VectorType /* old */ , VectorType& /* new */);
 
-	/** Notification for vector value slider began movement */
-	DECLARE_DELEGATE(FOnBeginSliderMovement);
-
-	/** Notification for vector value slider ended movement */
-	DECLARE_DELEGATE_OneParam(FOnEndSliderMovement, NumericType);
-
 	struct FArguments;
 
 private:
@@ -80,10 +74,10 @@ private:
 		SLATE_EVENT(FMenuExtensionDelegate, ContextMenuExtenderX)
 
 		/** Called when the x value of the vector slider began movement */
-		SLATE_EVENT(FOnBeginSliderMovement, OnXBeginSliderMovement)
+		SLATE_EVENT(FSimpleDelegate, OnXBeginSliderMovement)
 
 		/** Called when the x value of the vector slider ended movement */
-		SLATE_EVENT(FOnEndSliderMovement, OnXEndSliderMovement)
+		SLATE_EVENT(FOnNumericValueChanged, OnXEndSliderMovement)
 	};
 
 	struct FVectorYArgumentsEmpty {};
@@ -115,10 +109,10 @@ private:
 		SLATE_EVENT(FMenuExtensionDelegate, ContextMenuExtenderY)
 
 		/** Called when the y value of the vector slider began movement */
-		SLATE_EVENT(FOnBeginSliderMovement, OnYBeginSliderMovement)
+		SLATE_EVENT(FSimpleDelegate, OnYBeginSliderMovement)
 
 		/** Called when the y value of the vector slider ended movement */
-		SLATE_EVENT(FOnEndSliderMovement, OnYEndSliderMovement)
+		SLATE_EVENT(FOnNumericValueChanged, OnYEndSliderMovement)
 	};
 
 	struct FVectorZArgumentsEmpty {};
@@ -150,10 +144,10 @@ private:
 		SLATE_EVENT(FMenuExtensionDelegate, ContextMenuExtenderZ)
 
 		/** Called when the z value of the vector slider began movement */
-		SLATE_EVENT(FOnBeginSliderMovement, OnZBeginSliderMovement)
+		SLATE_EVENT(FSimpleDelegate, OnZBeginSliderMovement)
 
 		/** Called when the z value of the vector slider ended movement */
-		SLATE_EVENT(FOnEndSliderMovement, OnZEndSliderMovement)
+		SLATE_EVENT(FOnNumericValueChanged, OnZEndSliderMovement)
 	};
 
 	struct FVectorWArgumentsEmpty {};
@@ -185,10 +179,10 @@ private:
 		SLATE_EVENT(FMenuExtensionDelegate, ContextMenuExtenderW)
 
 		/** Called when the w value of the vector slider began movement */
-		SLATE_EVENT(FOnBeginSliderMovement, OnWBeginSliderMovement)
+		SLATE_EVENT(FSimpleDelegate, OnWBeginSliderMovement)
 
 		/** Called when the w value of the vector slider ended movement */
-		SLATE_EVENT(FOnEndSliderMovement, OnWEndSliderMovement)
+		SLATE_EVENT(FOnNumericValueChanged, OnWEndSliderMovement)
 	};
 
 public:
@@ -421,8 +415,8 @@ private:
 		const TAttribute<ECheckBoxState> ToggleChecked,
 		const FOnCheckStateChanged& OnToggleChanged,
 		const FMenuExtensionDelegate& OnContextMenuExtenderComponent,
-		const FSimpleDelegate& OnBeginSliderMovement,
-		const FOnNumericValueChanged& OnEndSliderMovement)
+		const FSimpleDelegate& OnComponentBeginSliderMovement,
+		const FOnNumericValueChanged& OnComponentEndSliderMovement)
 	{
 		TSharedRef<SWidget> LabelWidget = SNullWidget::NullWidget;
 		if (InArgs._bColorAxisLabels)
@@ -462,8 +456,8 @@ private:
 			.MaxSliderValue(CreatePerComponentGetter(ComponentIndex, TOptional<NumericType>(), InArgs._MaxSliderVector))
 			.LinearDeltaSensitivity(InArgs._LinearDeltaSensitivity)
 			.Delta(InArgs._SpinDelta)
-			.OnBeginSliderMovement(OnBeginSliderMovement)
-			.OnEndSliderMovement(OnEndSliderMovement)
+			.OnBeginSliderMovement(CreatePerComponentSliderMovementEvent(InArgs._OnBeginSliderMovement, OnComponentBeginSliderMovement))
+			.OnEndSliderMovement(CreatePerComponentSliderMovementEvent<FOnNumericValueChanged, NumericType>(InArgs._OnEndSliderMovement, OnComponentEndSliderMovement))
 			.LabelPadding(FMargin(3.f))
 			.LabelLocation(SNumericEntryBox<NumericType>::ELabelLocation::Inside)
 			.Label()
@@ -651,6 +645,26 @@ private:
 		return OnComponentCommitted;
 	}
 
+	/**
+	 * Creates a lambda to react to a begin/end slider movement event
+	 */
+	template<typename EventType, typename... ArgsType>
+	EventType CreatePerComponentSliderMovementEvent(
+		const EventType OnSliderMovement,
+		const EventType OnComponentSliderMovement)
+	{
+		if(OnSliderMovement.IsBound())
+		{
+			return EventType::CreateLambda(
+				[OnSliderMovement, OnComponentSliderMovement](ArgsType... Args)
+				{
+					OnSliderMovement.ExecuteIfBound(Args...);
+					OnComponentSliderMovement.ExecuteIfBound(Args...);
+				});
+		}
+		return OnComponentSliderMovement;
+	}
+
 	bool bUseVectorGetter = true;
 	TAttribute<TOptional<VectorType>> VectorAttribute;
 	FOnVectorValueChanged OnVectorValueChanged;
