TabManager: Remove the ability to enforce a main tab because it causes layout issues
Asset Editors now don't have a non closeable tab by default, instead we show a text hint  if the user accidentally closes all tabs

diff --git a/Engine/Source/Editor/SkeletalMeshEditor/Private/SkeletalMeshEditor.cpp b/Engine/Source/Editor/SkeletalMeshEditor/Private/SkeletalMeshEditor.cpp
--- a/Engine/Source/Editor/SkeletalMeshEditor/Private/SkeletalMeshEditor.cpp
+++ b/Engine/Source/Editor/SkeletalMeshEditor/Private/SkeletalMeshEditor.cpp
@@ -212,11 +212,6 @@ bool FSkeletalMeshEditor::OnRequestClose(EAssetEditorCloseReason InCloseReason)
 	return bAllowClose;
 }
 
-FName FSkeletalMeshEditor::GetMainTabName() const
-{
-	return SkeletalMeshEditorTabs::ViewportTab;
-}
-
 void FSkeletalMeshEditor::RegisterTabSpawners(const TSharedRef<class FTabManager>& InTabManager)
 {
 	WorkspaceMenuCategory = InTabManager->AddLocalWorkspaceMenuCategory(LOCTEXT("WorkspaceMenu_SkeletalMeshEditor", "Skeletal Mesh Editor"));
diff --git a/Engine/Source/Editor/SkeletalMeshEditor/Private/SkeletalMeshEditor.h b/Engine/Source/Editor/SkeletalMeshEditor/Private/SkeletalMeshEditor.h
--- a/Engine/Source/Editor/SkeletalMeshEditor/Private/SkeletalMeshEditor.h
+++ b/Engine/Source/Editor/SkeletalMeshEditor/Private/SkeletalMeshEditor.h
@@ -83,7 +83,6 @@ public:
 	
 	//~ Begin FAssetEditorToolkit Interface.
 	virtual bool OnRequestClose(EAssetEditorCloseReason InCloseReason) override;
-	virtual FName GetMainTabName() const override;
 	//~ End FAssetEditorToolkit Interface.
 
 	/** FEditorUndoClient interface */
diff --git a/Engine/Source/Editor/UnrealEd/Private/Toolkits/AssetEditorToolkit.cpp b/Engine/Source/Editor/UnrealEd/Private/Toolkits/AssetEditorToolkit.cpp
--- a/Engine/Source/Editor/UnrealEd/Private/Toolkits/AssetEditorToolkit.cpp
+++ b/Engine/Source/Editor/UnrealEd/Private/Toolkits/AssetEditorToolkit.cpp
@@ -190,13 +190,6 @@ void FAssetEditorToolkit::InitAssetEditor( const EToolkitMode::Type Mode, const
 		const TSharedRef<FTabManager> NewTabManager = FGlobalTabmanager::Get()->NewTabManager( NewMajorTab.ToSharedRef() );		
 		NewTabManager->SetOnPersistLayout(FTabManager::FOnPersistLayout::CreateRaw(this, &FAssetEditorToolkit::HandleTabManagerPersistLayout));
 		NewTabManager->SetAllowWindowMenuBar(true);
-		NewTabManager->SetEnforceMainTab(true);
-
-		FName MainTabName = GetMainTabName();
-		if(!MainTabName.IsNone())
-		{
-			NewTabManager->SetMainTab(MainTabName);
-		}
 
 		this->TabManager = NewTabManager;
 
diff --git a/Engine/Source/Editor/UnrealEd/Public/Toolkits/AssetEditorToolkit.h b/Engine/Source/Editor/UnrealEd/Public/Toolkits/AssetEditorToolkit.h
--- a/Engine/Source/Editor/UnrealEd/Public/Toolkits/AssetEditorToolkit.h
+++ b/Engine/Source/Editor/UnrealEd/Public/Toolkits/AssetEditorToolkit.h
@@ -244,12 +244,7 @@ public:
 	/** Returns the default extensibility managers, these are applied for all asset types */
 	static UNREALED_API TSharedPtr<FExtensibilityManager> GetSharedMenuExtensibilityManager();
 	static UNREALED_API TSharedPtr<FExtensibilityManager> GetSharedToolBarExtensibilityManager();
-
-	/** Override this to manually specify the tab that cannot be closed for your asset editor, it is the first tab
-	 *  registered otherwise
-	 */
-	virtual FName GetMainTabName() const { return NAME_None; }
-
+	
 	/**
 	 * Allows the caller to set a menu overlay, displayed to the far right of the editor's menu bar
 	 *
diff --git a/Engine/Source/Runtime/Slate/Private/Framework/Docking/SDockingArea.cpp b/Engine/Source/Runtime/Slate/Private/Framework/Docking/SDockingArea.cpp
--- a/Engine/Source/Runtime/Slate/Private/Framework/Docking/SDockingArea.cpp
+++ b/Engine/Source/Runtime/Slate/Private/Framework/Docking/SDockingArea.cpp
@@ -8,7 +8,11 @@
 #include "Framework/Docking/FDockingDragOperation.h"
 #include "HAL/PlatformApplicationMisc.h"
 #include "Framework/Docking/STabSidebar.h"
+#include "Styling/SlateColor.h"
+#include "Widgets/Images/SImage.h"
+#include "Widgets/Text/STextBlock.h"
 
+#define LOCTEXT_NAMESPACE "SDockingArea"
 
 void SDockingArea::Construct( const FArguments& InArgs, const TSharedRef<FTabManager>& InTabManager, const TSharedRef<FTabManager::FArea>& PersistentNode )
 {
@@ -25,6 +29,36 @@ void SDockingArea::Construct( const FArguments& InArgs, const TSharedRef<FTabMan
 
 	const TSharedRef<SOverlay> SidebarDrawersOverlay = SNew(SOverlay);
 
+	TSharedRef<SWidget> NoOpenTabsWidget =
+			SNew(SHorizontalBox)
+			.Visibility_Lambda([this]()
+			{
+				// This text is visible when we have no child tabs
+				return GetNumTabs() == 0 ? EVisibility::Visible : EVisibility::Collapsed;
+			})
+			
+			+SHorizontalBox::Slot()
+			.FillWidth(1.0f)
+			.HAlign(HAlign_Center)
+			.VAlign(VAlign_Center)
+			[
+				SNew(STextBlock)
+				.Text(LOCTEXT("NoTabsDockedText", "This asset editor has no docked tabs."))
+				.TextStyle( FAppStyle::Get(), "HintText")
+			]
+			+SHorizontalBox::Slot()
+			.AutoWidth()
+			.Padding(4.0f)
+			.HAlign(HAlign_Center)
+			.VAlign(VAlign_Center)
+			[
+				SNew(SImage)
+				.ColorAndOpacity(FSlateColor::UseSubduedForeground())
+				.Image(FAppStyle::GetBrush("Icons.Help"))
+				.ToolTipText(LOCTEXT("NoTabsDockedTooltip", "To recover your tabs, you can reopen them from the Window menu, or drag and drop them back from floating windows.\n"
+												   "You can also reset your editor layout completely from the Window > Load Layout menu, but this affects all editor windows."))
+			];
+
 	// In DockSplitter mode we just act as a thin shell around a Splitter widget
 	this->ChildSlot
 	[
@@ -42,8 +76,19 @@ void SDockingArea::Construct( const FArguments& InArgs, const TSharedRef<FTabMan
 			]
 			+ SHorizontalBox::Slot()
 			[
-				SAssignNew(Splitter, SSplitter)
-				.Orientation(PersistentNode->GetOrientation())
+				SNew(SOverlay)
+				+SOverlay::Slot()
+				[
+					SAssignNew(Splitter, SSplitter)
+					.Orientation(PersistentNode->GetOrientation())
+				]
+				+SOverlay::Slot()
+				.HAlign(HAlign_Center)
+				.VAlign(VAlign_Center)
+				[
+					NoOpenTabsWidget
+				]
+				
 			]
 			+SHorizontalBox::Slot()
 			.AutoWidth()
@@ -674,3 +719,4 @@ void SDockingArea::UpdateWindowChromeAndSidebar()
 
 }
 
+#undef LOCTEXT_NAMESPACE
diff --git a/Engine/Source/Runtime/Slate/Private/Framework/Docking/TabManager.cpp b/Engine/Source/Runtime/Slate/Private/Framework/Docking/TabManager.cpp
--- a/Engine/Source/Runtime/Slate/Private/Framework/Docking/TabManager.cpp
+++ b/Engine/Source/Runtime/Slate/Private/Framework/Docking/TabManager.cpp
@@ -851,22 +851,9 @@ void FTabManager::SetMainTab(const TSharedRef<const SDockTab>& InTab)
 	
 }
 
-void FTabManager::SetEnforceMainTab(bool bInEnforceMainTab)
-{
-	bEnforceMainTab = bInEnforceMainTab;
-
-	static const FTabId InvalidTab;
-	// If we already have tabs registered, set the main tab to whichever first tab we can find if we don't have one
-	if(bEnforceMainTab && MainNonCloseableTabID == InvalidTab && !TabSpawner.IsEmpty())
-	{
-		FTabSpawner::TIterator SpawnerIterator(TabSpawner);
-		MainNonCloseableTabID = SpawnerIterator.Key();
-	}
-}
-
 bool FTabManager::IsTabCloseable(const TSharedRef<const SDockTab>& InTab) const
 {
-	return !(MainNonCloseableTabID == InTab->GetLayoutIdentifier());
+	return MainNonCloseableTabID != InTab->GetLayoutIdentifier();
 }
 
 const TSharedRef<FWorkspaceItem> FTabManager::GetLocalWorkspaceMenuRoot() const
@@ -1020,13 +1007,6 @@ FTabSpawnerEntry& FTabManager::RegisterTabSpawner(const FName TabId, const FOnSp
 
 	TSharedRef<FTabSpawnerEntry> NewSpawnerEntry = MakeShareable(new FTabSpawnerEntry(TabId, OnSpawnTab, CanSpawnTab));
 	TabSpawner.Add(TabId, NewSpawnerEntry);
-
-	static const FTabId InvalidTab;
-	// If this tab manager always wants a non closeable tab, and we don't already have one set this tab to be the main tab
-	if(bEnforceMainTab && MainNonCloseableTabID == InvalidTab && !PendingMainNonClosableTab)
-	{
-		MainNonCloseableTabID = TabId;
-	}
 	
 	return NewSpawnerEntry.Get();
 }
diff --git a/Engine/Source/Runtime/Slate/Public/Framework/Docking/TabManager.h b/Engine/Source/Runtime/Slate/Public/Framework/Docking/TabManager.h
--- a/Engine/Source/Runtime/Slate/Public/Framework/Docking/TabManager.h
+++ b/Engine/Source/Runtime/Slate/Public/Framework/Docking/TabManager.h
@@ -947,13 +947,7 @@ class FTabManager : public TSharedFromThis<FTabManager>
 
 		/** Provide a tab that will be the main tab and cannot be closed. */
 		SLATE_API void SetMainTab(const FTabId& InMainTabID);
-
-		/** If true, this tab manager will always have one main tab that cannot be closed
-		 *  Automatically set to the first tab registered if there are no tabs registered already (set on creation)
-		 *  or the first tab found otherwise
-		 */
-		SLATE_API void SetEnforceMainTab(bool bInEnforceMainTab);
-
+	
 		/* Prevent or allow all tabs to be drag */
 		void SetCanDoDragOperation(bool CanDoDragOperation) { bCanDoDragOperation = CanDoDragOperation; }
 		
