[PIE] improve commandline to force touch simulation in PIE
#rnx

diff --git a/Engine/Source/Runtime/Slate/Private/Framework/Application/SlateApplication.cpp b/Engine/Source/Runtime/Slate/Private/Framework/Application/SlateApplication.cpp
--- a/Engine/Source/Runtime/Slate/Private/Framework/Application/SlateApplication.cpp
+++ b/Engine/Source/Runtime/Slate/Private/Framework/Application/SlateApplication.cpp
@@ -5991,6 +5991,12 @@ bool FSlateApplication::ProcessMouseMoveEvent( const FPointerEvent& MouseEvent,
 {
 	SCOPE_CYCLE_COUNTER(STAT_ProcessMouseMove);
 
+	if (IsFakingTouchEvents() && !MouseEvent.IsMouseButtonDown(EKeys::LeftMouseButton))
+	{
+		// If we're faking touch events and the left mouse button is not down, do not process the mouse move event
+		return false;
+	}
+   
 #if WITH_SLATE_DEBUGGING
 	FSlateDebugging::FScopeProcessInputEvent Scope(ESlateDebuggingInputEvent::MouseMove, MouseEvent);
 #endif
