Avoid Text Style copies occuring each frame, expensive due to atomics.

diff --git a/Engine/Source/Runtime/Slate/Private/Widgets/Text/STextBlock.cpp b/Engine/Source/Runtime/Slate/Private/Widgets/Text/STextBlock.cpp
--- a/Engine/Source/Runtime/Slate/Private/Widgets/Text/STextBlock.cpp
+++ b/Engine/Source/Runtime/Slate/Private/Widgets/Text/STextBlock.cpp
@@ -129,6 +129,11 @@ void STextBlock::Construct( const FArguments& InArgs )
 }
 
 FSlateFontInfo STextBlock::GetFont() const
+{
+	return GetFontRef();
+}
+
+const FSlateFontInfo& STextBlock::GetFontRef() const
 {
 	return bIsAttributeFontSet ? Font.Get() : TextStyle.Font;
 }
@@ -139,6 +144,11 @@ const FSlateBrush* STextBlock::GetStrikeBrush() const
 }
 
 FSlateColor STextBlock::GetColorAndOpacity() const
+{
+	return GetColorAndOpacityRef();
+}
+
+const FSlateColor& STextBlock::GetColorAndOpacityRef() const
 {
 	return bIsAttributeColorAndOpacitySet ? ColorAndOpacity.Get() : TextStyle.ColorAndOpacity;
 }
@@ -149,6 +159,11 @@ FVector2f STextBlock::GetShadowOffset() const
 }
 
 FLinearColor STextBlock::GetShadowColorAndOpacity() const
+{
+	return GetShadowColorAndOpacityRef();
+}
+
+const FLinearColor& STextBlock::GetShadowColorAndOpacityRef() const
 {
 	return bIsAttributeShadowColorAndOpacitySet ? ShadowColorAndOpacity.Get() : TextStyle.ShadowColorAndOpacity;
 }
@@ -296,6 +311,17 @@ FVector2D STextBlock::ComputeDesiredSize(float LayoutScaleMultiplier) const
 	}
 	else
 	{
+		FTextBlockStyle::CompareParams StyleParams(TextStyle
+			, GetFontRef()
+			, GetColorAndOpacityRef()
+			, GetShadowOffset()
+			, GetShadowColorAndOpacityRef()
+			, GetHighlightColor()
+			, GetHighlightShape()
+			, GetStrikeBrush());
+
+		TextLayoutCache->ConditionallyUpdateTextStyle(MoveTemp(StyleParams));
+
 		// ComputeDesiredSize will also update the text layout cache if required
 		const FVector2D TextSize = TextLayoutCache->ComputeDesiredSize(
 			FSlateTextBlockLayout::FWidgetDesiredSizeArgs(
@@ -308,7 +334,7 @@ FVector2D STextBlock::ComputeDesiredSize(float LayoutScaleMultiplier) const
 				Margin.Get(),
 				LineHeightPercentage.Get(),
 				Justification.Get()),
-			LayoutScaleMultiplier, GetComputedTextStyle());
+			LayoutScaleMultiplier);
 
 		return FVector2D(FMath::Max(MinDesiredWidth.Get(), TextSize.X), TextSize.Y);
 	}
diff --git a/Engine/Source/Runtime/Slate/Private/Widgets/Text/SlateTextBlockLayout.cpp b/Engine/Source/Runtime/Slate/Private/Widgets/Text/SlateTextBlockLayout.cpp
--- a/Engine/Source/Runtime/Slate/Private/Widgets/Text/SlateTextBlockLayout.cpp
+++ b/Engine/Source/Runtime/Slate/Private/Widgets/Text/SlateTextBlockLayout.cpp
@@ -25,7 +25,42 @@ FSlateTextBlockLayout::FSlateTextBlockLayout(SWidget* InOwner, FTextBlockStyle I
 	TextLayout->SetLineBreakIterator(MoveTemp(InLineBreakPolicy));
 }
 
-FVector2D FSlateTextBlockLayout::ComputeDesiredSize(const FWidgetDesiredSizeArgs& InWidgetArgs, const float InScale, const FTextBlockStyle& InTextStyle)
+void FSlateTextBlockLayout::ConditionallyUpdateTextStyle(const FTextBlockStyle& InTextStyle)
+{
+	// Has the style used for this text block changed?
+	if (!IsStyleUpToDate(InTextStyle))
+	{
+		TextLayout->SetDefaultTextStyle(InTextStyle);
+		Marshaller->MakeDirty(); // will regenerate the text using the new default style
+	}
+}
+
+void FSlateTextBlockLayout::ConditionallyUpdateTextStyle(const FTextBlockStyle::CompareParams& InNewStyleParams)
+{
+	// Has the style used for this text block changed?
+	if (!IsStyleUpToDate(InNewStyleParams))
+	{
+		FTextBlockStyle ComputedStyle = InNewStyleParams.StyleBase;
+		ComputedStyle.SetFont(InNewStyleParams.Font);
+		if (InNewStyleParams.StrikeBrush)
+		{
+			ComputedStyle.SetStrikeBrush(*InNewStyleParams.StrikeBrush);
+		}
+		ComputedStyle.SetColorAndOpacity(InNewStyleParams.ColorAndOpacity);
+		ComputedStyle.SetShadowOffset(InNewStyleParams.ShadowOffset);
+		ComputedStyle.SetShadowColorAndOpacity(InNewStyleParams.ShadowColorAndOpacity);
+		ComputedStyle.SetHighlightColor(InNewStyleParams.HighlightColor);
+		if (InNewStyleParams.HighlightShape)
+		{
+			ComputedStyle.SetHighlightShape(*InNewStyleParams.HighlightShape);
+		}
+		TextLayout->SetDefaultTextStyle(MoveTemp(ComputedStyle));
+		Marshaller->MakeDirty(); // will regenerate the text using the new default style
+	}
+}
+
+
+FVector2D FSlateTextBlockLayout::ComputeDesiredSize(const FWidgetDesiredSizeArgs& InWidgetArgs, const float InScale)
 {
 	// Cache the wrapping rules so that we can recompute the wrap at width in paint.
 	CachedWrapTextAt = InWidgetArgs.WrapTextAt;
@@ -48,13 +83,6 @@ FVector2D FSlateTextBlockLayout::ComputeDesiredSize(const FWidgetDesiredSizeArgs
 		Marshaller->MakeDirty();
 	}
 
-	// Has the style used for this text block changed?
-	if (!IsStyleUpToDate(InTextStyle))
-	{
-		TextLayout->SetDefaultTextStyle(InTextStyle);
-		Marshaller->MakeDirty(); // will regenerate the text using the new default style
-	}
-
 	{
 		bool bRequiresTextUpdate = false;
 		const FText& TextToSet = InWidgetArgs.Text;
@@ -98,6 +126,12 @@ FVector2D FSlateTextBlockLayout::ComputeDesiredSize(const FWidgetDesiredSizeArgs
 	return TextLayout->GetSize();
 }
 
+FVector2D FSlateTextBlockLayout::ComputeDesiredSize(const FWidgetDesiredSizeArgs& InWidgetArgs, const float InScale, const FTextBlockStyle& InTextStyle)
+{
+	ConditionallyUpdateTextStyle(InTextStyle);
+	return ComputeDesiredSize(InWidgetArgs, InScale);
+}
+
 PRAGMA_DISABLE_DEPRECATION_WARNINGS
 FVector2D FSlateTextBlockLayout::ComputeDesiredSize(const FWidgetArgs& InWidgetArgs, const float InScale, const FTextBlockStyle& InTextStyle)
 {
@@ -337,6 +371,12 @@ bool FSlateTextBlockLayout::IsStyleUpToDate(const FTextBlockStyle& NewStyle) con
 	return CurrentStyle.IsIdenticalTo(NewStyle);
 }
 
+bool FSlateTextBlockLayout::IsStyleUpToDate(const FTextBlockStyle::CompareParams& InNewStyleParams) const
+{
+	const FTextBlockStyle& CurrentStyle = TextLayout->GetDefaultTextStyle();
+	return CurrentStyle.IsIdenticalTo(InNewStyleParams);
+}
+
 float FSlateTextBlockLayout::CalculateWrappingWidth() const
 {
 	// Text wrapping can either be used defined (WrapTextAt), automatic (bAutoWrapText and CachedSize), 
diff --git a/Engine/Source/Runtime/Slate/Public/Widgets/Text/STextBlock.h b/Engine/Source/Runtime/Slate/Public/Widgets/Text/STextBlock.h
--- a/Engine/Source/Runtime/Slate/Public/Widgets/Text/STextBlock.h
+++ b/Engine/Source/Runtime/Slate/Public/Widgets/Text/STextBlock.h
@@ -259,9 +259,15 @@ public:
 	/** Gets the current foreground color */
 	SLATE_API FSlateColor GetColorAndOpacity() const;
 
+	/** Gets the current foreground color as ref */
+	const FSlateColor& GetColorAndOpacityRef() const;
+
 	/** Gets the current font */
 	SLATE_API FSlateFontInfo GetFont() const;
 
+	/** Gets the current font as ref */
+	const FSlateFontInfo& GetFontRef() const;
+
 	/** Gets the current strike brush */
 	SLATE_API const FSlateBrush* GetStrikeBrush() const;
 
@@ -274,6 +280,9 @@ public:
 	/** Gets the current shadow color and opacity */
 	SLATE_API FLinearColor GetShadowColorAndOpacity() const;
 
+	/** Gets the current shadow color and opacity as ref */
+	const FLinearColor& GetShadowColorAndOpacityRef() const;
+
 	/** Gets the current highlight color */
 	SLATE_API FSlateColor GetHighlightColor() const;
 
diff --git a/Engine/Source/Runtime/Slate/Public/Widgets/Text/SlateTextBlockLayout.h b/Engine/Source/Runtime/Slate/Public/Widgets/Text/SlateTextBlockLayout.h
--- a/Engine/Source/Runtime/Slate/Public/Widgets/Text/SlateTextBlockLayout.h
+++ b/Engine/Source/Runtime/Slate/Public/Widgets/Text/SlateTextBlockLayout.h
@@ -99,6 +99,21 @@ public:
 	 */
 	SLATE_API FSlateTextBlockLayout(SWidget* InOwner, FTextBlockStyle InDefaultTextStyle, const TOptional<ETextShapingMethod> InTextShapingMethod, const TOptional<ETextFlowDirection> InTextFlowDirection, const FCreateSlateTextLayout& InCreateSlateTextLayout, TSharedRef<ITextLayoutMarshaller> InMarshaller, TSharedPtr<IBreakIterator> InLineBreakPolicy);
 
+	/**
+	 * Conditionally update the text style if needed
+	 */
+	SLATE_API void ConditionallyUpdateTextStyle(const FTextBlockStyle& InTextStyle);
+
+	/**
+	 * Conditionally update the text style if needed
+	 */
+	SLATE_API void ConditionallyUpdateTextStyle(const FTextBlockStyle::CompareParams& InNewStyleParams);
+
+	/**
+	 * Get the computed desired size for this layout, updating the internal cache as required
+	 */
+	SLATE_API FVector2D ComputeDesiredSize(const FWidgetDesiredSizeArgs& InWidgetArgs, const float InScale);
+
 	/**
 	 * Get the computed desired size for this layout, updating the internal cache as required
 	 */
@@ -190,6 +205,9 @@ private:
 	/** Is the style used by the text marshaller up-to-date? */
 	bool IsStyleUpToDate(const FTextBlockStyle& NewStyle) const;
 
+	/** Is the style used by the text marshaller up-to-date? */
+	bool IsStyleUpToDate(const FTextBlockStyle::CompareParams& InNewStyleParams) const;
+
 	/** Calculate the wrapping width based on the given fixed wrap width, and whether we're auto-wrapping */
 	float CalculateWrappingWidth() const;
 
diff --git a/Engine/Source/Runtime/SlateCore/Public/Styling/SlateTypes.h b/Engine/Source/Runtime/SlateCore/Public/Styling/SlateTypes.h
--- a/Engine/Source/Runtime/SlateCore/Public/Styling/SlateTypes.h
+++ b/Engine/Source/Runtime/SlateCore/Public/Styling/SlateTypes.h
@@ -412,6 +412,59 @@ struct FTextBlockStyle : public FSlateWidgetStyle
 			&& OverflowPolicy == InOther.OverflowPolicy;
 	}
 
+
+	/** Helper struct to compare two text styles without constructing a temporary text style */
+	struct CompareParams
+	{
+		CompareParams(const FTextBlockStyle& InStyleBase
+					, const FSlateFontInfo& InFont
+					, const FSlateColor& InColorAndOpacity
+					, const FVector2f InShadowOffset
+					, const FLinearColor& InShadowColorAndOpacity
+					, const FSlateColor InHighlightColor
+					, const FSlateBrush* InHighlightShape
+					, const FSlateBrush* InStrikeBrush)
+			: StyleBase(InStyleBase)
+			, Font(InFont)
+			, ColorAndOpacity(InColorAndOpacity)
+			, ShadowOffset(InShadowOffset)
+			, ShadowColorAndOpacity(InShadowColorAndOpacity)
+			, HighlightColor(InHighlightColor)
+			, HighlightShape(InHighlightShape)
+			, StrikeBrush(InStrikeBrush)
+		{}
+
+		const FTextBlockStyle& StyleBase;
+		const FSlateFontInfo& Font;
+		const FSlateColor& ColorAndOpacity;
+		const FVector2f ShadowOffset;
+		const FLinearColor& ShadowColorAndOpacity;
+		const FSlateColor HighlightColor;
+		const FSlateBrush* HighlightShape;
+		const FSlateBrush* StrikeBrush;
+	};
+
+	/**
+	 * Checks to see whether this style is identical to another.
+	 */
+	bool IsIdenticalTo(const FTextBlockStyle::CompareParams& InNewStyleParams) const
+	{
+		const FSlateBrush& NewStrikeBrush = InNewStyleParams.StrikeBrush ? *InNewStyleParams.StrikeBrush : InNewStyleParams.StyleBase.StrikeBrush;
+		const FSlateBrush& NewHighlightShape = InNewStyleParams.HighlightShape ? *InNewStyleParams.HighlightShape : InNewStyleParams.StyleBase.HighlightShape;
+
+		return Font.IsIdenticalTo(InNewStyleParams.Font)
+			&& ColorAndOpacity == InNewStyleParams.ColorAndOpacity
+			&& ShadowOffset == InNewStyleParams.ShadowOffset
+			&& ShadowColorAndOpacity == InNewStyleParams.ShadowColorAndOpacity
+			&& SelectedBackgroundColor == InNewStyleParams.StyleBase.SelectedBackgroundColor
+			&& HighlightColor == InNewStyleParams.HighlightColor
+			&& HighlightShape == NewHighlightShape
+			&& StrikeBrush == NewStrikeBrush
+			&& UnderlineBrush == InNewStyleParams.StyleBase.UnderlineBrush
+			&& TransformPolicy == InNewStyleParams.StyleBase.TransformPolicy
+			&& OverflowPolicy == InNewStyleParams.StyleBase.OverflowPolicy;
+	}
+
 	/**
 	 * Unlinks all colors in this style.
 	 * @see FSlateColor::Unlink
