Fix typo in comment of TabManager.h

diff --git a/Engine/Source/Runtime/Slate/Public/Framework/Docking/TabManager.h b/Engine/Source/Runtime/Slate/Public/Framework/Docking/TabManager.h
--- a/Engine/Source/Runtime/Slate/Public/Framework/Docking/TabManager.h
+++ b/Engine/Source/Runtime/Slate/Public/Framework/Docking/TabManager.h
@@ -1224,12 +1224,11 @@ public:
 	SLATE_API void SetActiveTab( const TSharedPtr<SDockTab>& NewActiveTab );
 
 	/**
-	 * Register a new normad tab spawner with the tab manager.  The spawner will be called when anyone calls
-	 * InvokeTab().
-	 * A nomad tab is a tab that can be placed with major tabs or minor tabs in any tab well
+	 * Register a new nomad tab spawner with the tab manager. The spawner will be called when anyone calls InvokeTab().
+	 * A nomad tab is a tab that can be placed with major tabs or minor tabs in any tab well.
 	 * @param TabId The TabId to register the spawner for.
 	 * @param OnSpawnTab The callback that will be used to spawn the tab.
-	 * @param CanSpawnTab The callback that will be used to ask if spawning the tab is allowed
+	 * @param CanSpawnTab The callback that will be used to ask if spawning the tab is allowed.
 	 * @return The registration entry for the spawner.
 	 */
 	SLATE_API FTabSpawnerEntry& RegisterNomadTabSpawner( const FName TabId, const FOnSpawnTab& OnSpawnTab, const FCanSpawnTab& CanSpawnTab = FCanSpawnTab());
