SSuggestionTextBox now presents the suggestions from the top of the list.

diff --git a/Engine/Source/Runtime/Slate/Private/Widgets/Input/SSuggestionTextBox.cpp b/Engine/Source/Runtime/Slate/Private/Widgets/Input/SSuggestionTextBox.cpp
--- a/Engine/Source/Runtime/Slate/Private/Widgets/Input/SSuggestionTextBox.cpp
+++ b/Engine/Source/Runtime/Slate/Private/Widgets/Input/SSuggestionTextBox.cpp
@@ -173,7 +173,9 @@ void SSuggestionTextBox::SetSuggestions( TArray<FString>& SuggestionStrings, boo
 	{
 		// @todo Slate: make the window title not flicker when the box toggles visibility
 		MenuAnchor->SetIsOpen(true, false);
-		SuggestionListView->RequestScrollIntoView(Suggestions.Last());
+
+		int32 SuggestionIndexToScrollIntoView = FMath::Max(0, SelectedSuggestion);
+		SuggestionListView->RequestScrollIntoView(Suggestions[SuggestionIndexToScrollIntoView]);
 
 		FocusTextBox();
 	}
