Added AutoWrapText attribute to SInlineEditableTextBlock

diff --git a/Engine/Source/Runtime/Slate/Private/Widgets/Text/SInlineEditableTextBlock.cpp b/Engine/Source/Runtime/Slate/Private/Widgets/Text/SInlineEditableTextBlock.cpp
--- a/Engine/Source/Runtime/Slate/Private/Widgets/Text/SInlineEditableTextBlock.cpp
+++ b/Engine/Source/Runtime/Slate/Private/Widgets/Text/SInlineEditableTextBlock.cpp
@@ -41,6 +41,7 @@ void SInlineEditableTextBlock::Construct( const FArguments& InArgs )
 			.HighlightText( InArgs._HighlightText )
 			.ToolTipText( InArgs._ToolTipText )
 			.WrapTextAt( InArgs._WrapTextAt )
+			.AutoWrapText( InArgs._AutoWrapNonEditText )
 			.Justification( InArgs._Justification )
 			.LineBreakPolicy( InArgs._LineBreakPolicy )
 			.OverflowPolicy(InArgs._OverflowPolicy)
@@ -58,6 +59,7 @@ void SInlineEditableTextBlock::Construct( const FArguments& InArgs )
 			.OnTextChanged(this, &SInlineEditableTextBlock::OnTextChanged)
 			.OnTextCommitted(this, &SInlineEditableTextBlock::OnTextBoxCommitted)
 			.WrapTextAt(InArgs._WrapTextAt)
+			.AutoWrapText(InArgs._AutoWrapMultilineEditText)
 			.Justification(InArgs._Justification)
 			.SelectAllTextWhenFocused(true)
 			.ClearKeyboardFocusOnCommit(true)
diff --git a/Engine/Source/Runtime/Slate/Public/Widgets/Text/SInlineEditableTextBlock.h b/Engine/Source/Runtime/Slate/Public/Widgets/Text/SInlineEditableTextBlock.h
--- a/Engine/Source/Runtime/Slate/Public/Widgets/Text/SInlineEditableTextBlock.h
+++ b/Engine/Source/Runtime/Slate/Public/Widgets/Text/SInlineEditableTextBlock.h
@@ -39,6 +39,8 @@ class SInlineEditableTextBlock: public SCompoundWidget
 		, _ShadowColorAndOpacity()
 		, _HighlightText()
 		, _WrapTextAt(0.0f)
+		, _AutoWrapNonEditText(false)
+		, _AutoWrapMultilineEditText(false)
 		, _Justification(ETextJustify::Left)
 		, _LineBreakPolicy()
 		, _IsReadOnly(false)
@@ -72,6 +74,12 @@ class SInlineEditableTextBlock: public SCompoundWidget
 		/** Whether text wraps onto a new line when it's length exceeds this width; if this value is zero or negative, no wrapping occurs. */
 		SLATE_ATTRIBUTE( float, WrapTextAt )
 
+		/** Whether the text should wrap automatically when in non-edit mode */
+		SLATE_ATTRIBUTE( bool, AutoWrapNonEditText )
+
+		/** Whether the multiline text should wrap automatically when in edit mode */
+		SLATE_ATTRIBUTE( bool, AutoWrapMultilineEditText )
+
 		/** How the text should be aligned with the margin. */
 		SLATE_ATTRIBUTE( ETextJustify::Type, Justification )
 
