Add missing SLATE_API to STextBlock.

diff --git a/Engine/Source/Runtime/Slate/Public/Widgets/Text/STextBlock.h b/Engine/Source/Runtime/Slate/Public/Widgets/Text/STextBlock.h
--- a/Engine/Source/Runtime/Slate/Public/Widgets/Text/STextBlock.h
+++ b/Engine/Source/Runtime/Slate/Public/Widgets/Text/STextBlock.h
@@ -260,13 +260,13 @@ public:
 	SLATE_API FSlateColor GetColorAndOpacity() const;
 
 	/** Gets the current foreground color as ref */
-	const FSlateColor& GetColorAndOpacityRef() const;
+	SLATE_API const FSlateColor& GetColorAndOpacityRef() const;
 
 	/** Gets the current font */
 	SLATE_API FSlateFontInfo GetFont() const;
 
 	/** Gets the current font as ref */
-	const FSlateFontInfo& GetFontRef() const;
+	SLATE_API const FSlateFontInfo& GetFontRef() const;
 
 	/** Gets the current strike brush */
 	SLATE_API const FSlateBrush* GetStrikeBrush() const;
@@ -281,7 +281,7 @@ public:
 	SLATE_API FLinearColor GetShadowColorAndOpacity() const;
 
 	/** Gets the current shadow color and opacity as ref */
-	const FLinearColor& GetShadowColorAndOpacityRef() const;
+	SLATE_API const FLinearColor& GetShadowColorAndOpacityRef() const;
 
 	/** Gets the current highlight color */
 	SLATE_API FSlateColor GetHighlightColor() const;
