Updated FNavigationConfig with an option to only send key events if no modifiers are held
Updated the EditorNavigationConfig to use this new setting, preventing an issue where the Ctrl+Space content drawer hotkey would be consumed by some widgets

diff --git a/Engine/Source/Runtime/Slate/Private/Framework/Application/NavigationConfig.cpp b/Engine/Source/Runtime/Slate/Private/Framework/Application/NavigationConfig.cpp
--- a/Engine/Source/Runtime/Slate/Private/Framework/Application/NavigationConfig.cpp
+++ b/Engine/Source/Runtime/Slate/Private/Framework/Application/NavigationConfig.cpp
@@ -9,6 +9,7 @@ FNavigationConfig::FNavigationConfig()
 	: bTabNavigation(true)
 	, bKeyNavigation(true)
 	, bAnalogNavigation(true)
+	, bIgnoreModifiersForNavigationActions(true)
 	, AnalogNavigationHorizontalThreshold(0.50f)
 	, AnalogNavigationVerticalThreshold(0.50f)
 {
@@ -153,10 +154,16 @@ float FNavigationConfig::GetRepeatRateForPressure(float InPressure, int32 InRepe
 
 EUINavigationAction FNavigationConfig::GetNavigationActionFromKey(const FKeyEvent& InKeyEvent) const
 {
+	const bool bModifierHeld = InKeyEvent.IsControlDown() || InKeyEvent.IsAltDown() || InKeyEvent.IsCommandDown() || InKeyEvent.IsShiftDown();
+	if (bIgnoreModifiersForNavigationActions || !bModifierHeld)
+	{
 PRAGMA_DISABLE_DEPRECATION_WARNINGS
-	// Call raw key version for back compatibility, subclasses should override this function
-	return GetNavigationActionForKey(InKeyEvent.GetKey());
+		// Call raw key version for back compatibility, subclasses should override this function
+		return GetNavigationActionForKey(InKeyEvent.GetKey());
 PRAGMA_ENABLE_DEPRECATION_WARNINGS
+
+	}
+	return EUINavigationAction::Invalid;
 }
 
 EUINavigationAction FNavigationConfig::GetNavigationActionForKey(const FKey& InKey) const
diff --git a/Engine/Source/Runtime/Slate/Private/Framework/Application/SlateApplication.cpp b/Engine/Source/Runtime/Slate/Private/Framework/Application/SlateApplication.cpp
--- a/Engine/Source/Runtime/Slate/Private/Framework/Application/SlateApplication.cpp
+++ b/Engine/Source/Runtime/Slate/Private/Framework/Application/SlateApplication.cpp
@@ -852,6 +852,7 @@ FSlateApplication::FSlateApplication()
 
 	NavigationConfig->OnRegister();
 #if WITH_EDITOR
+	EditorNavigationConfig->bIgnoreModifiersForNavigationActions = false;
 	EditorNavigationConfig->OnRegister();
 #endif
 
diff --git a/Engine/Source/Runtime/Slate/Public/Framework/Application/NavigationConfig.h b/Engine/Source/Runtime/Slate/Public/Framework/Application/NavigationConfig.h
--- a/Engine/Source/Runtime/Slate/Public/Framework/Application/NavigationConfig.h
+++ b/Engine/Source/Runtime/Slate/Public/Framework/Application/NavigationConfig.h
@@ -96,6 +96,8 @@ public:
 	bool bKeyNavigation;
 	/** Should we respect the analog stick for navigation. */
 	bool bAnalogNavigation;
+	/** Should we ignore modifier keys when checking for navigation actions. If false, only unmodified keys will be processed. */
+	bool bIgnoreModifiersForNavigationActions;
 
 	/**  */
 	float AnalogNavigationHorizontalThreshold;
