Solve a bug where Android users are able to scroll SListView widgets even when scrolling options are turned off

-Added a property that allows disabling touch scrolling on STableViewBase

diff --git a/Engine/Source/Runtime/Slate/Private/Widgets/Views/STableViewBase.cpp b/Engine/Source/Runtime/Slate/Private/Widgets/Views/STableViewBase.cpp
--- a/Engine/Source/Runtime/Slate/Private/Widgets/Views/STableViewBase.cpp
+++ b/Engine/Source/Runtime/Slate/Private/Widgets/Views/STableViewBase.cpp
@@ -398,7 +398,7 @@ void STableViewBase::ScrollBar_OnUserScrolled( float InScrollOffsetFraction )
 
 FReply STableViewBase::OnPreviewMouseButtonDown( const FGeometry& MyGeometry, const FPointerEvent& MouseEvent )
 {
-	if (MouseEvent.IsTouchEvent())
+	if (bEnableTouchScrolling && MouseEvent.IsTouchEvent())
 	{
 		// Clear any inertia 
 		this->InertialScrollManager.ClearScrollVelocity();
@@ -545,7 +545,7 @@ FReply STableViewBase::OnMouseMove( const FGeometry& MyGeometry, const FPointerE
 
 void STableViewBase::OnMouseEnter( const FGeometry& MyGeometry, const FPointerEvent& MouseEvent )
 {
-	if ( MouseEvent.IsTouchEvent() )
+	if (bEnableTouchScrolling && MouseEvent.IsTouchEvent())
 	{
 		if ( !bStartedTouchInteraction )
 		{
@@ -636,7 +636,7 @@ FReply STableViewBase::OnTouchStarted( const FGeometry& MyGeometry, const FPoint
 
 FReply STableViewBase::OnTouchMoved( const FGeometry& MyGeometry, const FPointerEvent& InTouchEvent )
 {
-	if (bStartedTouchInteraction)
+	if (bEnableTouchScrolling && bStartedTouchInteraction)
 	{
 		// We only care about deltas along the scroll axis
 		FTableViewDimensions CursorDeltaDimensions(Orientation, InTouchEvent.GetCursorDelta());
@@ -917,6 +917,12 @@ void STableViewBase::SetIsRightClickScrollingEnabled(const bool bInEnableRightCl
 	bEnableRightClickScrolling = bInEnableRightClickScrolling;
 }
 
+void STableViewBase::SetIsTouchScrollingEnabled(const bool bInEnableTouchScrolling)
+{
+	bEnableTouchScrolling = bInEnableTouchScrolling;
+	ensureMsgf(!bStartedTouchInteraction, TEXT("TouchScrollingEnabled flag should not be changed while scrolling."));
+}
+
 void STableViewBase::SetWheelScrollMultiplier(float NewWheelScrollMultiplier)
 {
 	WheelScrollMultiplier = NewWheelScrollMultiplier;
diff --git a/Engine/Source/Runtime/Slate/Public/Widgets/Views/STableViewBase.h b/Engine/Source/Runtime/Slate/Public/Widgets/Views/STableViewBase.h
--- a/Engine/Source/Runtime/Slate/Public/Widgets/Views/STableViewBase.h
+++ b/Engine/Source/Runtime/Slate/Public/Widgets/Views/STableViewBase.h
@@ -176,6 +176,9 @@ public:
 	/** Enables/disables being able to scroll with the right mouse button. */
 	SLATE_API void SetIsRightClickScrollingEnabled(const bool bInEnableRightClickScrolling);
 
+	/** Enables/disables being able to scroll using touch input. */
+	SLATE_API void SetIsTouchScrollingEnabled(const bool bInEnableTouchScrolling);
+
 	/** Sets the multiplier applied when wheel scrolling. Higher numbers will cover more distance per click of the wheel. */
 	SLATE_API void SetWheelScrollMultiplier(float NewWheelScrollMultiplier);
 
@@ -411,6 +414,9 @@ protected:
 	/** True to allow right click drag scrolling. */
 	bool bEnableRightClickScrolling = true;
 
+	/** True to allow scrolling by using touch input. */
+	bool bEnableTouchScrolling = true;
+
 	/** The currently displayed scroll offset from the beginning of the list in items. */
 	double CurrentScrollOffset = 0.;
 
diff --git a/Engine/Source/Runtime/UMG/Private/Components/ListViewBase.cpp b/Engine/Source/Runtime/UMG/Private/Components/ListViewBase.cpp
--- a/Engine/Source/Runtime/UMG/Private/Components/ListViewBase.cpp
+++ b/Engine/Source/Runtime/UMG/Private/Components/ListViewBase.cpp
@@ -137,6 +137,7 @@ TSharedRef<SWidget> UListViewBase::RebuildWidget()
 	MyTableViewBase = RebuildListWidget();
 	MyTableViewBase->SetIsScrollAnimationEnabled(bEnableScrollAnimation);
 	MyTableViewBase->SetIsRightClickScrollingEnabled(bEnableRightClickScrolling);
+	MyTableViewBase->SetIsTouchScrollingEnabled(bEnableTouchScrolling);
 	MyTableViewBase->SetFixedLineScrollOffset(bEnableFixedLineOffset ? TOptional<double>(FixedLineScrollOffset) : TOptional<double>());
 	MyTableViewBase->SetWheelScrollMultiplier(GetGlobalScrollAmount() * WheelScrollMultiplier);
 
@@ -160,6 +161,7 @@ void UListViewBase::SynchronizeProperties()
 	{
 		MyTableViewBase->SetIsScrollAnimationEnabled(bEnableScrollAnimation);
 		MyTableViewBase->SetIsRightClickScrollingEnabled(bEnableRightClickScrolling);
+		MyTableViewBase->SetIsTouchScrollingEnabled(bEnableTouchScrolling);
 		MyTableViewBase->SetAllowOverscroll(AllowOverscroll ? EAllowOverscroll::Yes : EAllowOverscroll::No);
 		MyTableViewBase->SetFixedLineScrollOffset(bEnableFixedLineOffset ? TOptional<double>(FixedLineScrollOffset) : TOptional<double>());
 		MyTableViewBase->SetWheelScrollMultiplier(GetGlobalScrollAmount() * WheelScrollMultiplier);
diff --git a/Engine/Source/Runtime/UMG/Public/Components/ListViewBase.h b/Engine/Source/Runtime/UMG/Public/Components/ListViewBase.h
--- a/Engine/Source/Runtime/UMG/Public/Components/ListViewBase.h
+++ b/Engine/Source/Runtime/UMG/Public/Components/ListViewBase.h
@@ -677,6 +677,10 @@ protected:
 	UPROPERTY(EditAnywhere, Category = Scrolling)
 	bool bEnableRightClickScrolling = true;
 
+	/** True to allow scrolling using touch input. */
+	UPROPERTY(EditAnywhere, Category = Scrolling)
+	bool bEnableTouchScrolling = true;
+
 	UPROPERTY(EditAnywhere, Category = Scrolling)
 	bool bEnableFixedLineOffset = false;
 
