Slate: Fix invalidation of SSlider when the Min or Max value changes, when Enable changed when Hovered changes.

diff --git a/Engine/Source/Runtime/Slate/Private/Widgets/Input/SSlider.cpp b/Engine/Source/Runtime/Slate/Private/Widgets/Input/SSlider.cpp
--- a/Engine/Source/Runtime/Slate/Private/Widgets/Input/SSlider.cpp
+++ b/Engine/Source/Runtime/Slate/Private/Widgets/Input/SSlider.cpp
@@ -15,6 +15,9 @@ void SSlider::PrivateRegisterAttributes(FSlateAttributeInitializer& AttributeIni
 	SLATE_ADD_MEMBER_ATTRIBUTE_DEFINITION_WITH_NAME(AttributeInitializer, "Locked", LockedSlateAttribute, EInvalidateWidgetReason::Paint);
 	SLATE_ADD_MEMBER_ATTRIBUTE_DEFINITION_WITH_NAME(AttributeInitializer, "SliderBarColor", SliderBarColorSlateAttribute, EInvalidateWidgetReason::Paint);
 	SLATE_ADD_MEMBER_ATTRIBUTE_DEFINITION_WITH_NAME(AttributeInitializer, "SliderHandleColor", SliderHandleColorSlateAttribute, EInvalidateWidgetReason::Paint);
+
+	AttributeInitializer.OverrideInvalidationReason("EnabledState", FSlateAttributeDescriptor::FInvalidateWidgetReasonAttribute{EInvalidateWidgetReason::Paint});
+	AttributeInitializer.OverrideInvalidationReason("Hovered", FSlateAttributeDescriptor::FInvalidateWidgetReasonAttribute{EInvalidateWidgetReason::Paint});
 }
 
 SSlider::SSlider()
@@ -517,11 +520,15 @@ void SSlider::SetValue(TAttribute<float> InValueAttribute)
 
 void SSlider::SetMinAndMaxValues(float InMinValue, float InMaxValue)
 {
-	MinValue = InMinValue;
-	MaxValue = InMaxValue;
-	if (MinValue > MaxValue)
+	if (MinValue != InMinValue || MaxValue != InMaxValue)
 	{
-		MaxValue = MinValue;
+		MinValue = InMinValue;
+		MaxValue = InMaxValue;
+		if (MinValue > MaxValue)
+		{
+			MaxValue = MinValue;
+		}
+		Invalidate(EInvalidateWidgetReason::Paint);
 	}
 }
 
