Fix rounding error in spinbox calculation when scrolling the values.

diff --git a/Engine/Source/Runtime/Slate/Public/Widgets/Input/SSpinBox.h b/Engine/Source/Runtime/Slate/Public/Widgets/Input/SSpinBox.h
--- a/Engine/Source/Runtime/Slate/Public/Widgets/Input/SSpinBox.h
+++ b/Engine/Source/Runtime/Slate/Public/Widgets/Input/SSpinBox.h
@@ -648,18 +648,17 @@ public:
 					if (LinearDeltaSensitivity.IsSet() && LinearDeltaSensitivity.Get() != 0 && Delta.IsSet() && Delta.Get() > 0)
 					{
 						const double MouseDelta = FMath::Abs(MouseEvent.GetCursorDelta().X / (float)LinearDeltaSensitivity.Get());
-						NewValue = InternalValue + (Sign * MouseDelta * FMath::Pow((double)Delta.Get(), SliderExponent.Get()));
+						NewValue = InternalValue + (Sign * MouseDelta * FMath::Pow((double)Delta.Get(), (double)SliderExponent.Get()));
 					}
 					else
 					{
 						const double MouseDelta = FMath::Abs(MouseEvent.GetCursorDelta().X / SliderWidthInSlateUnits);
 						const double CurrentValue = FMath::Clamp<double>(FMath::Abs(InternalValue), 1.0, (double)std::numeric_limits<NumericType>::max());
-						NewValue = InternalValue + (Sign * MouseDelta * FMath::Pow((float)CurrentValue, SliderExponent.Get()));
+						NewValue = InternalValue + (Sign * MouseDelta * FMath::Pow((double)CurrentValue, (double)SliderExponent.Get()));
 					}
 				}
 
 				NumericType RoundedNewValue = RoundIfIntegerValue(NewValue);
-				//UE_LOG(LogSlate, Warning, TEXT("A - %llx B - %f  C - %llx"), RoundedNewValue, NewValue, static_cast<NumericType>(NewValue));
 				CommitValue(RoundedNewValue, NewValue, CommittedViaSpin, ETextCommit::OnEnter);
 			}
 
