PR #10680: Fixed issue with tab well animations constantly playing

diff --git a/Engine/Source/Runtime/Slate/Private/Framework/Docking/SDockingTabStack.cpp b/Engine/Source/Runtime/Slate/Private/Framework/Docking/SDockingTabStack.cpp
--- a/Engine/Source/Runtime/Slate/Private/Framework/Docking/SDockingTabStack.cpp
+++ b/Engine/Source/Runtime/Slate/Private/Framework/Docking/SDockingTabStack.cpp
@@ -1153,6 +1153,12 @@ SSplitter::ESizeRule SDockingTabStack::GetSizeRule() const
 
 void SDockingTabStack::SetTabWellHidden( bool bShouldHideTabWell )
 {
+	// If the tab well is already hidden or visible, don't replay the animations.
+	if ( (bShouldHideTabWell && IsTabWellHidden()) || (!bShouldHideTabWell && !IsTabWellHidden()))
+	{
+		return;
+	}
+	
 	if (bShouldHideTabWell)
 	{
 		ShowHideTabWell.PlayReverse( this->AsShared() );
