Added new keybind that combines Close Minor Tab and Close Major Tab. Unbound by default, this keybind will let users press a single key to close the focused tab regardless of its type, prioritizing minor tabs.

diff --git a/Engine/Source/Runtime/Slate/Private/Framework/Docking/SDockingTabStack.cpp b/Engine/Source/Runtime/Slate/Private/Framework/Docking/SDockingTabStack.cpp
--- a/Engine/Source/Runtime/Slate/Private/Framework/Docking/SDockingTabStack.cpp
+++ b/Engine/Source/Runtime/Slate/Private/Framework/Docking/SDockingTabStack.cpp
@@ -1328,6 +1328,7 @@ void SDockingTabStack::BindTabCommands()
 	const FTabCommands& Commands = FTabCommands::Get();
 	ActionList->MapAction(Commands.CloseMajorTab, FExecuteAction::CreateSP(this, &SDockingTabStack::ExecuteCloseMajorTabCommand), FCanExecuteAction::CreateSP(this, &SDockingTabStack::CanExecuteCloseMajorTabCommand));
 	ActionList->MapAction(Commands.CloseMinorTab, FExecuteAction::CreateSP(this, &SDockingTabStack::ExecuteCloseMinorTabCommand), FCanExecuteAction::CreateSP(this, &SDockingTabStack::CanExecuteCloseMinorTabCommand));
+	ActionList->MapAction(Commands.CloseFocusedTab, FExecuteAction::CreateSP(this, &SDockingTabStack::ExecuteCloseFocusedTabCommand), FCanExecuteAction::CreateSP(this, &SDockingTabStack::CanExecuteCloseFocusedTabCommand));
 }
 
 void SDockingTabStack::ExecuteCloseMajorTabCommand()
@@ -1373,6 +1374,23 @@ bool SDockingTabStack::CanExecuteCloseMinorTabCommand()
 	return false;
 }
 
+void SDockingTabStack::ExecuteCloseFocusedTabCommand()
+{
+	if (CanExecuteCloseMinorTabCommand())
+	{
+		ExecuteCloseMinorTabCommand();
+	}
+	else
+	{
+		ExecuteCloseMajorTabCommand();
+	}
+}
+
+bool SDockingTabStack::CanExecuteCloseFocusedTabCommand()
+{
+	return CanExecuteCloseMinorTabCommand() || CanExecuteCloseMajorTabCommand();
+}
+
 void SDockingTabStack::OnResized()
 {
 	if (TSharedPtr<SDockingArea> DockArea = GetDockArea())
diff --git a/Engine/Source/Runtime/Slate/Private/Framework/Docking/SDockingTabStack.h b/Engine/Source/Runtime/Slate/Private/Framework/Docking/SDockingTabStack.h
--- a/Engine/Source/Runtime/Slate/Private/Framework/Docking/SDockingTabStack.h
+++ b/Engine/Source/Runtime/Slate/Private/Framework/Docking/SDockingTabStack.h
@@ -317,6 +317,12 @@ private:
 	/** Attempts to find the active tab that can be closed by the CloseMinorTab command */
 	SLATE_API bool CanExecuteCloseMinorTabCommand();
 
+	/** Attempts to close the focused tab when the CloseFocusedTab command is executed */
+	SLATE_API void ExecuteCloseFocusedTabCommand();
+
+	/** Attempts to find the active tab that can be closed by the CloseFocusedTab command */
+	SLATE_API bool CanExecuteCloseFocusedTabCommand();
+
 	/** Called when this stack is resized. */
 	SLATE_API virtual void OnResized() override;
 
diff --git a/Engine/Source/Runtime/Slate/Private/Framework/Docking/TabCommands.cpp b/Engine/Source/Runtime/Slate/Private/Framework/Docking/TabCommands.cpp
--- a/Engine/Source/Runtime/Slate/Private/Framework/Docking/TabCommands.cpp
+++ b/Engine/Source/Runtime/Slate/Private/Framework/Docking/TabCommands.cpp
@@ -8,6 +8,7 @@ void FTabCommands::RegisterCommands()
 {
 	UI_COMMAND(CloseMajorTab, "Close Major Tab", "Closes the focused major tab", EUserInterfaceActionType::Button, FInputChord(EModifierKey::Control, EKeys::W))
 	UI_COMMAND(CloseMinorTab, "Close Minor Tab", "Closes the current window's active minor tab", EUserInterfaceActionType::Button, FInputChord())
+	UI_COMMAND(CloseFocusedTab, "Close Focused Tab", "Closes the focused tab, prioritizing minor tabs over major tabs", EUserInterfaceActionType::Button, FInputChord())
 }
 
 #undef LOCTEXT_NAMESPACE
diff --git a/Engine/Source/Runtime/Slate/Public/Framework/Docking/TabCommands.h b/Engine/Source/Runtime/Slate/Public/Framework/Docking/TabCommands.h
--- a/Engine/Source/Runtime/Slate/Public/Framework/Docking/TabCommands.h
+++ b/Engine/Source/Runtime/Slate/Public/Framework/Docking/TabCommands.h
@@ -23,4 +23,5 @@ public:
 
 	TSharedPtr<FUICommandInfo> CloseMajorTab;
 	TSharedPtr<FUICommandInfo> CloseMinorTab;
+	TSharedPtr<FUICommandInfo> CloseFocusedTab;
 };
