Changed the way SVirtualJoystick handles Pressed/Released inputs so it can be used as a button without requiring sliding the finger.

diff --git a/Engine/Source/Runtime/Engine/Classes/GameFramework/TouchInterface.h b/Engine/Source/Runtime/Engine/Classes/GameFramework/TouchInterface.h
--- a/Engine/Source/Runtime/Engine/Classes/GameFramework/TouchInterface.h
+++ b/Engine/Source/Runtime/Engine/Classes/GameFramework/TouchInterface.h
@@ -17,6 +17,8 @@ struct FTouchInputControl
 	GENERATED_USTRUCT_BODY()
 
 	// basically mirroring SVirtualJoystick::FControlInfo but as an editable class
+	UPROPERTY(EditAnywhere, Category="Control", meta=(ToolTip="Set this to true to treat the joystick as a simple button"))
+	bool bTreatAsButton = false;
 	UPROPERTY(EditAnywhere, Category="Control", meta=(ToolTip="For sticks, this is the Thumb"))
 	TObjectPtr<UTexture2D> Image1;
 	UPROPERTY(EditAnywhere, Category="Control", meta=(ToolTip="For sticks, this is the Background"))
@@ -37,7 +39,8 @@ struct FTouchInputControl
 	FKey AltInputKey;
 
 	FTouchInputControl()
-		: Image1(nullptr)
+		: bTreatAsButton(false)
+		, Image1(nullptr)
 		, Image2(nullptr)
 		, Center(ForceInitToZero)
 		, VisualSize(ForceInitToZero)
diff --git a/Engine/Source/Runtime/Engine/Private/GameFramework/TouchInterface.cpp b/Engine/Source/Runtime/Engine/Private/GameFramework/TouchInterface.cpp
--- a/Engine/Source/Runtime/Engine/Private/GameFramework/TouchInterface.cpp
+++ b/Engine/Source/Runtime/Engine/Private/GameFramework/TouchInterface.cpp
@@ -35,6 +35,7 @@ void UTouchInterface::Activate(TSharedPtr<SVirtualJoystick> VirtualJoystick)
 			FTouchInputControl Control = Controls[ControlIndex];
 			SVirtualJoystick::FControlInfo* SlateControl = new(SlateControls)SVirtualJoystick::FControlInfo;
 
+			SlateControl->bTreatAsButton = Control.bTreatAsButton;
 			SlateControl->Image1 = Control.Image1 ? StaticCastSharedRef<ISlateBrushSource>(FDeferredCleanupSlateBrush::CreateBrush(Control.Image1)) : TSharedPtr<ISlateBrushSource>();
 			SlateControl->Image2 = Control.Image2 ? StaticCastSharedRef<ISlateBrushSource>(FDeferredCleanupSlateBrush::CreateBrush(Control.Image2)) : TSharedPtr<ISlateBrushSource>();
 			SlateControl->Center = Control.Center;
diff --git a/Engine/Source/Runtime/Slate/Private/Widgets/Input/SVirtualJoystick.cpp b/Engine/Source/Runtime/Slate/Private/Widgets/Input/SVirtualJoystick.cpp
--- a/Engine/Source/Runtime/Slate/Private/Widgets/Input/SVirtualJoystick.cpp
+++ b/Engine/Source/Runtime/Slate/Private/Widgets/Input/SVirtualJoystick.cpp
@@ -392,6 +392,9 @@ void SVirtualJoystick::Tick(const FGeometry& AllottedGeometry, const double InCu
 
 		if (Control.CapturedPointerIndex >= 0 || Control.bSendOneMoreEvent)
 		{
+			// cache the key released state so we can send input pressed/released events later
+			bool bButtonPressed = !Control.bSendOneMoreEvent;
+
 			Control.bSendOneMoreEvent = false;
 
 			// Get the corrected thumb offset scale (now allows ellipse instead of assuming square)
@@ -420,13 +423,24 @@ void SVirtualJoystick::Tick(const FGeometry& AllottedGeometry, const double InCu
 			
 			FInputDeviceId PrimaryInputDevice = IPlatformInputDeviceMapper::Get().GetPrimaryInputDeviceForUser(FSlateApplicationBase::SlateAppPrimaryPlatformUser);
 
-			auto ApplyInput = [PrimaryInputDevice](const FGamepadKeyNames::Type KeyName, float Delta)
+			auto ApplyInput = [PrimaryInputDevice](const FGamepadKeyNames::Type KeyName, float Delta, bool bTreatAsButton, bool bPressed)
 			{
 				FKey Key(KeyName);
 				if (Key.IsAnalog())
 				{
 					FSlateApplication::Get().OnControllerAnalog(KeyName, FSlateApplicationBase::SlateAppPrimaryPlatformUser, PrimaryInputDevice, Delta);
 				}
+				else if (bTreatAsButton)
+				{
+					if (bPressed)
+					{
+						FSlateApplication::Get().OnControllerButtonPressed(KeyName, FSlateApplicationBase::SlateAppPrimaryPlatformUser, PrimaryInputDevice, false);
+					}
+					else
+					{
+						FSlateApplication::Get().OnControllerButtonReleased(KeyName, FSlateApplicationBase::SlateAppPrimaryPlatformUser, PrimaryInputDevice, false);
+					}
+				}
 				else if (Delta != 0.0f)
 				{
 					FSlateApplication::Get().OnControllerButtonPressed(KeyName, FSlateApplicationBase::SlateAppPrimaryPlatformUser, PrimaryInputDevice, false);
@@ -437,8 +451,13 @@ void SVirtualJoystick::Tick(const FGeometry& AllottedGeometry, const double InCu
 				}
 			};
 
-			ApplyInput(XAxis, NormalizedOffset.X);
-			ApplyInput(YAxis, -NormalizedOffset.Y);
+			ApplyInput(XAxis, NormalizedOffset.X, Control.Info.bTreatAsButton, bButtonPressed);
+
+			// ignore the Y axis if this is a button
+			if (!Control.Info.bTreatAsButton)
+			{
+				ApplyInput(YAxis, -NormalizedOffset.Y, false, bButtonPressed);
+			}
 		}
 		
 
diff --git a/Engine/Source/Runtime/Slate/Public/Widgets/Input/SVirtualJoystick.h b/Engine/Source/Runtime/Slate/Public/Widgets/Input/SVirtualJoystick.h
--- a/Engine/Source/Runtime/Slate/Public/Widgets/Input/SVirtualJoystick.h
+++ b/Engine/Source/Runtime/Slate/Public/Widgets/Input/SVirtualJoystick.h
@@ -24,6 +24,9 @@ public:
 	{
 		// Set by the game
 
+		/** If true, this joystick will be treated as a simple press-release button */
+		bool bTreatAsButton = false;
+
 		/** The brush to use to draw the background for joysticks, or unclicked for buttons */
 		TSharedPtr< ISlateBrushSource > Image1;
 
