Fix Outliner not handling Shift selection when Ctrl is held
Fix Outliner scrollbar jumping when Shift selecting a row higher than the starting row

diff --git a/Engine/Source/Runtime/Slate/Public/Widgets/Views/SListView.h b/Engine/Source/Runtime/Slate/Public/Widgets/Views/SListView.h
--- a/Engine/Source/Runtime/Slate/Public/Widgets/Views/SListView.h
+++ b/Engine/Source/Runtime/Slate/Public/Widgets/Views/SListView.h
@@ -966,15 +966,16 @@ public:
 			RangeStartIndex = FMath::Clamp(RangeStartIndex, 0, ItemsSourceRef.Num()-1);
 			RangeEndIndex = FMath::Clamp(RangeEndIndex, 0, ItemsSourceRef.Num()-1);
 
-			if (RangeEndIndex < RangeStartIndex)
-			{
-				Swap( RangeStartIndex, RangeEndIndex );
-			}
+			// Respect the direction of selection when ordering, ie if selecting upwards then make sure the top element is last-selected
+			const int32 Direction = (RangeEndIndex > RangeStartIndex) ? 1 : -1;
 
-			for (int32 ItemIndex = RangeStartIndex; ItemIndex <= RangeEndIndex; ++ItemIndex)
+			int32 ItemIndex = RangeStartIndex;
+			for (; ItemIndex != RangeEndIndex; ItemIndex += Direction)
 			{
 				SelectedItems.Add(ItemsSourceRef[ItemIndex]);
 			}
+			// The above loop won't add the last item, so manually add it here
+			SelectedItems.Add(ItemsSourceRef[ItemIndex]);
 		}
 
 		this->InertialScrollManager.ClearScrollVelocity();
diff --git a/Engine/Source/Runtime/Slate/Public/Widgets/Views/STableRow.h b/Engine/Source/Runtime/Slate/Public/Widgets/Views/STableRow.h
--- a/Engine/Source/Runtime/Slate/Public/Widgets/Views/STableRow.h
+++ b/Engine/Source/Runtime/Slate/Public/Widgets/Views/STableRow.h
@@ -457,18 +457,18 @@ public:
 
 					if (SelectionMode == ESelectionMode::Multi)
 					{
-						if (MouseEvent.IsControlDown())
+						if (MouseEvent.IsShiftDown())
 						{
-							OwnerTable->Private_SetItemSelection(MyItem, !bIsSelected, true);
+							OwnerTable->Private_SelectRangeFromCurrentTo(MyItem);
 							bChangedSelectionOnMouseDown = true;
 							if (SignalSelectionMode == ETableRowSignalSelectionMode::Instantaneous)
 							{
 								OwnerTable->Private_SignalSelectionChanged(ESelectInfo::OnMouseClick);
 							}
 						}
-						else if (MouseEvent.IsShiftDown())
+						else if (MouseEvent.IsControlDown())
 						{
-							OwnerTable->Private_SelectRangeFromCurrentTo(MyItem);
+							OwnerTable->Private_SetItemSelection(MyItem, !bIsSelected, true);
 							bChangedSelectionOnMouseDown = true;
 							if (SignalSelectionMode == ETableRowSignalSelectionMode::Instantaneous)
 							{
