Add Mac support for additional content on the right side of the title bar

This is already supported on Windows where it is used to show build info.

diff --git a/Engine/Source/Runtime/Slate/Public/Framework/Application/SWindowTitleBar.h b/Engine/Source/Runtime/Slate/Public/Framework/Application/SWindowTitleBar.h
--- a/Engine/Source/Runtime/Slate/Public/Framework/Application/SWindowTitleBar.h
+++ b/Engine/Source/Runtime/Slate/Public/Framework/Application/SWindowTitleBar.h
@@ -168,6 +168,11 @@ public:
 
 	virtual void UpdateWindowMenu(TSharedPtr<SWidget> MenuContent) override
 	{
+		if (!WindowMenuSlot)
+		{
+			return;
+		}
+
 		if(MenuContent.IsValid() && bAllowMenuBar)
 		{
 			(*WindowMenuSlot)
@@ -319,13 +324,10 @@ protected:
 		}
 
 #if PLATFORM_MAC
-
-		// On Mac we use real window buttons drawn by the OS
+	// The Mac has no need for custom content on the left side of the title bar as there is no main menu and the close, minimize, and maximize buttons are drawn by macOS.
 		OutLeftContent = SNew(SSpacer);
-		OutRightContent = SNew(SSpacer);
-
+		WindowMenuSlot = nullptr;
 #else // PLATFORM_MAC
-
 		// Windows UI layout
 		if (ShowAppIcon && bHasWindowButtons)
 		{
@@ -344,7 +346,7 @@ protected:
 				.FillWidth(1)
 				.Expose(WindowMenuSlot);
 
-			// Default everything to use thge small icon unless specifically set to use the large icon
+			// Default everything to use the small icon unless specifically set to use the large icon.
 			SetAllowMenuBar(false);
 		}
 		else
@@ -353,6 +355,7 @@ protected:
 
 			OutLeftContent = SNew(SSpacer);
 		}
+#endif //PLATFORM_MAC
 
 		if (bHasWindowButtons)
 		{
@@ -360,11 +363,14 @@ protected:
 				.Visibility(EVisibility::SelfHitTestInvisible)
 				.Padding(FMargin(2.0f, 0.0f, 0.0f, 0.0f))
 				[
-					// Minimize
+					// Expose a slot for optional content inside the title bar on the right side.
 					SNew(SHorizontalBox)
 					+ SHorizontalBox::Slot()
 					.Expose(RightSideContentSlot)
 
+// We don't need to add these as macOS draws them for us.
+#if !PLATFORM_MAC
+					// Minimize
 					+ SHorizontalBox::Slot()
 					.VAlign(VAlign_Top)
 					.AutoWidth()
@@ -387,6 +393,7 @@ protected:
 					[
 						CloseButton.ToSharedRef()
 					]
+#endif //!PLATFORM_MAC
 				];
 		}
 		else
@@ -395,8 +402,6 @@ protected:
 
 			OutRightContent = SNew(SSpacer);
 		}
-
-#endif // PLATFORM_MAC
 	}
 
 	/**
