Avoid a crash when creating a FUIAction in an environment where Slate isn't initialized, eg, a commandlet

#rnx

diff --git a/Engine/Source/Runtime/Slate/Private/Framework/Commands/UIAction.cpp b/Engine/Source/Runtime/Slate/Private/Framework/Commands/UIAction.cpp
--- a/Engine/Source/Runtime/Slate/Private/Framework/Commands/UIAction.cpp
+++ b/Engine/Source/Runtime/Slate/Private/Framework/Commands/UIAction.cpp
@@ -15,11 +15,16 @@ FUIAction::FUIAction()
 
 FUIAction::FUIAction(FExecuteAction InitExecuteAction, EUIActionRepeatMode InitRepeatMode)
 	: ExecuteAction(InitExecuteAction)
-	, CanExecuteAction(FCanExecuteAction::CreateRaw(&FSlateApplication::Get(), &FSlateApplication::IsNormalExecution))
+	, CanExecuteAction()
 	, GetActionCheckState()
 	, IsActionVisibleDelegate()
 	, RepeatMode(InitRepeatMode)
-{ 
+{
+	CanExecuteAction = FCanExecuteAction::CreateStatic([]()
+	{
+		return FSlateApplication::IsInitialized()
+			&& FSlateApplication::Get().IsNormalExecution();
+	});
 }
 
 FUIAction::FUIAction(FExecuteAction InitExecuteAction, FCanExecuteAction InitCanExecuteAction, EUIActionRepeatMode InitRepeatMode)
