Add CVAR to move Mac main menu into UE window

* Add the Mac-only CVAR Slate.MacMainMenuInsideUnrealWindow (ignored on other platforms). When it's false (default) there should be no user-visible difference. When true, Slate will draw (most of) the main menu inside the Unreal windows instead of using the macOS system main menu bar at the top of the screen.
* Expose the value of the CVAR via a new function FPlatformMisc::CanShowMenusInWindows(). This always returns true on non-Mac platforms.
* Remove unused code in SDockingTabStack.cpp.
* Improve spacing between tabs and close/min/max buttons in the title bar of smaller windows on the Mac.

diff --git a/Engine/Source/Runtime/Core/Private/Mac/MacPlatformMisc.cpp b/Engine/Source/Runtime/Core/Private/Mac/MacPlatformMisc.cpp
--- a/Engine/Source/Runtime/Core/Private/Mac/MacPlatformMisc.cpp
+++ b/Engine/Source/Runtime/Core/Private/Mac/MacPlatformMisc.cpp
@@ -77,6 +77,14 @@ static TAutoConsoleVariable<int32> CVarMacPlatformDumpAllThreadsOnHang(
 	TEXT("Mac.DumpAllThreadsOnHang"),
 	1,
 	TEXT("If > 0, then when reporting a hang generate a backtrace for all threads."));
+static bool GMacMainMenuInsideUnrealWindow = false;
+static FAutoConsoleVariableRef CVarMacMainMenuInsideUnrealWindow(
+	TEXT("Slate.MacMainMenuInsideUnrealWindow"),
+	GMacMainMenuInsideUnrealWindow,
+	TEXT("When enabled, will draw the main menu inside the Unreal windows on the Mac.\n")
+	TEXT("When disabled (default, old behavior) use the macOS system menu for all main menus."),
+	ECVF_Default
+);
 
 /*------------------------------------------------------------------------------
  Platform property discovery.
@@ -2663,6 +2671,11 @@ void ReportHang(const TCHAR* ErrorMessage, const uint64* StackFrames, int32 NumS
 	EnsureLock.Unlock();
 }
 
+bool FMacPlatformMisc::CanShowMenusInWindows()
+{
+	return GMacMainMenuInsideUnrealWindow;
+}
+
 typedef NSArray* (*MTLCopyAllDevices)(void);
 
 bool FMacPlatformMisc::HasPlatformFeature(const TCHAR* FeatureName)
diff --git a/Engine/Source/Runtime/Core/Public/GenericPlatform/GenericPlatformMisc.h b/Engine/Source/Runtime/Core/Public/GenericPlatform/GenericPlatformMisc.h
--- a/Engine/Source/Runtime/Core/Public/GenericPlatform/GenericPlatformMisc.h
+++ b/Engine/Source/Runtime/Core/Public/GenericPlatform/GenericPlatformMisc.h
@@ -1541,6 +1541,14 @@ public:
 		return false;
 	}
 
+	/*
+	 * Returns whether the platform can show main menus inside window title bars.
+	 */
+	static CORE_API bool CanShowMenusInWindows()
+	{
+		return true;
+	}
+
 	/*
 	 * Returns whether the volume buttons are handled by the system
 	 */
diff --git a/Engine/Source/Runtime/Core/Public/Mac/MacPlatformMisc.h b/Engine/Source/Runtime/Core/Public/Mac/MacPlatformMisc.h
--- a/Engine/Source/Runtime/Core/Public/Mac/MacPlatformMisc.h
+++ b/Engine/Source/Runtime/Core/Public/Mac/MacPlatformMisc.h
@@ -43,6 +43,7 @@ struct CORE_API FMacPlatformMisc : public FApplePlatformMisc
 	static void GetOSVersions( FString& out_OSVersionLabel, FString& out_OSSubVersionLabel );
 	static FString GetOSVersion();
 	static NSOperatingSystemVersion GetNSOperatingSystemVersion();
+	static bool CanShowMenusInWindows();
 	static bool HasPlatformFeature(const TCHAR* FeatureName);
 	static bool GetDiskTotalAndFreeSpace(const FString& InPath, uint64& TotalNumberOfBytes, uint64& NumberOfFreeBytes);
 	static bool HasSeparateChannelForDebugOutput();
diff --git a/Engine/Source/Runtime/Slate/Private/Framework/Docking/SDockingTabStack.cpp b/Engine/Source/Runtime/Slate/Private/Framework/Docking/SDockingTabStack.cpp
--- a/Engine/Source/Runtime/Slate/Private/Framework/Docking/SDockingTabStack.cpp
+++ b/Engine/Source/Runtime/Slate/Private/Framework/Docking/SDockingTabStack.cpp
@@ -3,6 +3,7 @@
 #include "Framework/Docking/SDockingTabStack.h"
 #include "Framework/Commands/UIAction.h"
 #include "Framework/Commands/UICommandList.h"
+#include "HAL/PlatformMisc.h"
 #include "Widgets/Text/STextBlock.h"
 #include "Framework/MultiBox/MultiBoxBuilder.h"
 #include "Layout/WidgetPath.h"
@@ -867,8 +868,21 @@ void SDockingTabStack::ClearReservedSpace()
 void SDockingTabStack::ReserveSpaceForWindowChrome(EChromeElement Element, bool bIncludePaddingForMenuBar, bool bOnlyMinorTabs)
 {
 	#if PLATFORM_MAC
-		static const FMargin ControlsPadding = FMargin(64, 0, 0, 0);
+		FMargin ControlsPadding;
 		static const FMargin IconPadding = FMargin(0);
+
+		if (FPlatformMisc::CanShowMenusInWindows() && bIncludePaddingForMenuBar)
+		{
+			static const float TopPaddingForTrafficLightsAndMenuBar = 30.0f;
+			// Always add padding on top, because on the Mac there is always either a main menu bar or the "traffic light" buttons (close, minimize, and maximize) above controls.
+			// Always add padding to the left, because on the Mac there's no Unreal icon to the left of controls, only the window edge, so we need some space.
+			ControlsPadding = FMargin(8.0f, TopPaddingForTrafficLightsAndMenuBar, 0, 0);
+		}
+		else
+		{
+			// Without a main menu bar in the title bar, we just need to pad on the left to avoid overlapping with the "traffic light" buttons (close, minimize, and maximize).
+			ControlsPadding = FMargin(67.0f, 0, 0, 0);
+		}
 	#else
 		static const float TopPaddingForMenuBar = 27.0f;
 		static const float LeftPaddingForIcon = FSlateApplication::Get().GetAppIcon()->GetImageSize().X;
@@ -916,10 +930,6 @@ void SDockingTabStack::SetParentNode( TSharedRef<class SDockingSplitter> InParen
 {
 	SDockingNode::SetParentNode(InParent);
 
-	// OK, if this docking area has a parent window, we'll assume the window was created with no title bar, and we'll
-	// place the title bar widgets into our content instead!
-	const TSharedPtr<SDockingArea>& DockArea = GetDockArea();
-
 	TitleBarSlot->AttachWidget(TitleBarContent.ToSharedRef());
 }
 
diff --git a/Engine/Source/Runtime/Slate/Private/Framework/Docking/TabManager.cpp b/Engine/Source/Runtime/Slate/Private/Framework/Docking/TabManager.cpp
--- a/Engine/Source/Runtime/Slate/Private/Framework/Docking/TabManager.cpp
+++ b/Engine/Source/Runtime/Slate/Private/Framework/Docking/TabManager.cpp
@@ -21,6 +21,7 @@
 #include "Misc/NamePermissionList.h"
 #include "Trace/SlateMemoryTags.h"
 #include "HAL/PlatformApplicationMisc.h"
+#include "HAL/PlatformMisc.h"
 #if PLATFORM_MAC
 #include "Framework/MultiBox/Mac/MacMenu.h"
 #endif
@@ -789,48 +790,48 @@ void FTabManager::UpdateMainMenu(TSharedPtr<SDockTab> ForTab, const bool bForce)
 		ParentWindowOfOwningTab = MainNonCloseableTabPinned->GetParentWindow();
 	}
 
-	// We only use the platform native global menu bar on Mac
 #if PLATFORM_MAC
-	if (MenuMultiBox.IsValid())
+	if (!FPlatformMisc::CanShowMenusInWindows())
 	{
-		bool bUpdate = bForce;
-		// On OS X opening the tab will set the multi-box and take key focus, but not seemingly send a keyboard focus event into Slate.
-		if (ParentWindowOfOwningTab.IsValid())
+		if (MenuMultiBox.IsValid())
 		{
-			bUpdate |= ParentWindowOfOwningTab->GetNativeWindow()->IsForegroundWindow();
-		}
+			bool bUpdate = bForce;
+			// On OS X opening the tab will set the multi-box and take key focus, but not seemingly send a keyboard focus event into Slate.
+			if (ParentWindowOfOwningTab.IsValid())
+			{
+				bUpdate |= ParentWindowOfOwningTab->GetNativeWindow()->IsForegroundWindow();
+			}
 
-		if (bUpdate)
+			if (bUpdate)
+			{
+				FSlateMacMenu::UpdateWithMultiBox(MenuMultiBox.ToSharedRef());
+			}
+		}
+		else
 		{
-			FSlateMacMenu::UpdateWithMultiBox(MenuMultiBox.ToSharedRef());
+			FSlateMacMenu::UpdateWithMultiBox(nullptr);
 		}
 	}
 	else
+#endif // PLATFORM_MAC
 	{
-		FSlateMacMenu::UpdateWithMultiBox(nullptr);
-	}
-#else
-
-
-	if (bAllowPerWindowMenu)
-	{
-		if (ParentWindowOfOwningTab)
-		{	
-			ParentWindowOfOwningTab->GetTitleBar()->UpdateWindowMenu(MenuWidget);
+		if (bAllowPerWindowMenu)
+		{
+			if (ParentWindowOfOwningTab)
+			{
+				ParentWindowOfOwningTab->GetTitleBar()->UpdateWindowMenu(MenuWidget);
+			}
 		}
-	}
-	else
-	{
-		MenuMultiBox.Reset();
-		MenuWidget.Reset();
-		if (ParentWindowOfOwningTab)
+		else
 		{
-			ParentWindowOfOwningTab->GetTitleBar()->UpdateWindowMenu(nullptr);
+			MenuMultiBox.Reset();
+			MenuWidget.Reset();
+			if (ParentWindowOfOwningTab)
+			{
+				ParentWindowOfOwningTab->GetTitleBar()->UpdateWindowMenu(nullptr);
+			}
 		}
 	}
-
-
-#endif
 }
 
 void FTabManager::SetMainTab(const FTabId& InMainTabID)
diff --git a/Engine/Source/Runtime/Slate/Private/Framework/MultiBox/MultiBox.cpp b/Engine/Source/Runtime/Slate/Private/Framework/MultiBox/MultiBox.cpp
--- a/Engine/Source/Runtime/Slate/Private/Framework/MultiBox/MultiBox.cpp
+++ b/Engine/Source/Runtime/Slate/Private/Framework/MultiBox/MultiBox.cpp
@@ -3,6 +3,7 @@
 #include "Framework/MultiBox/MultiBox.h"
 
 #include "SClippingVerticalBox.h"
+#include "HAL/PlatformMisc.h"
 #include "Misc/ConfigCacheIni.h"
 #include "Widgets/SBoxPanel.h"
 #include "Widgets/SOverlay.h"
@@ -448,9 +449,10 @@ TSharedRef< SMultiBoxWidget > FMultiBox::MakeWidget( bool bSearchable, FOnMakeMu
 		// Build up the widget
 		NewMultiBoxWidget->BuildMultiBoxWidget();
 	}
-	
+
 #if PLATFORM_MAC
-	if(Type == EMultiBoxType::MenuBar)
+	// Hide the main menu on the Mac if we are using the macOS system menu instead.
+	if (Type == EMultiBoxType::MenuBar && !FPlatformMisc::CanShowMenusInWindows())
 	{
 		NewMultiBoxWidget->SetVisibility(EVisibility::Collapsed);
 	}
diff --git a/Engine/Source/Runtime/Slate/Public/Framework/Application/SWindowTitleBar.h b/Engine/Source/Runtime/Slate/Public/Framework/Application/SWindowTitleBar.h
--- a/Engine/Source/Runtime/Slate/Public/Framework/Application/SWindowTitleBar.h
+++ b/Engine/Source/Runtime/Slate/Public/Framework/Application/SWindowTitleBar.h
@@ -3,6 +3,7 @@
 #pragma once
 
 #include "CoreMinimal.h"
+#include "HAL/PlatformMisc.h"
 #include "Misc/Attribute.h"
 #include "Layout/Visibility.h"
 #include "Styling/SlateColor.h"
@@ -272,6 +273,8 @@ protected:
 
 		const bool bHasWindowButtons = OwnerWindow->HasCloseBox() || OwnerWindow->HasMinimizeBox() || OwnerWindow->HasMaximizeBox();
 
+// We don't need these on the Mac where the system adds the "traffic light" buttons (close, minimize, and maximize).
+#if !PLATFORM_MAC
 		if (bHasWindowButtons)
 		{
 			MinimizeButton = SNew(SButton)
@@ -322,17 +325,24 @@ protected:
 					]
 				;
 		}
+#endif //!PLATFORM_MAC
 
-#if PLATFORM_MAC
-	// The Mac has no need for custom content on the left side of the title bar as there is no main menu and the close, minimize, and maximize buttons are drawn by macOS.
-		OutLeftContent = SNew(SSpacer);
-		WindowMenuSlot = nullptr;
-#else // PLATFORM_MAC
-		// Windows UI layout
-		if (ShowAppIcon && bHasWindowButtons)
+		if (ShowAppIcon && bHasWindowButtons && FPlatformMisc::CanShowMenusInWindows())
 		{
 			OutLeftContent = 
 				SNew(SHorizontalBox)
+#if PLATFORM_MAC
+				// The Mac has traffic-light buttons to the left, so we need some empty space under them (but no icon).
+				+ SHorizontalBox::Slot()
+				.AutoWidth()
+				.HAlign(HAlign_Left)
+				.VAlign(VAlign_Top)
+				[
+					// This leaves spaces for the macOS "traffic light" close/minimize/maximize buttons. Without this space, the traffic lights would render on top of the main menu on the Mac.
+					SNew(SSpacer)
+					.Size(FVector2D(64, 10))
+				]
+#else
 				+ SHorizontalBox::Slot()
 				.AutoWidth()
 				.HAlign(HAlign_Left)
@@ -341,9 +351,14 @@ protected:
 					SAssignNew(AppIconWidget, SAppIconWidget)
 					.IconColorAndOpacity(this, &SWindowTitleBar::GetWindowTitleContentColor)
 				]
+#endif //PLATFORM_MAC
 				+ SHorizontalBox::Slot()
 				.VAlign(VAlign_Top)
 				.FillWidth(1)
+#if PLATFORM_MAC
+				// Align the main menu with the macOS "traffic light" buttons (close, minimize, and maximize).
+				.Padding(FMargin(0, 3.0f, 0, 0))
+#endif // PLATFORM_MAC
 				.Expose(WindowMenuSlot);
 
 			// Default everything to use the small icon unless specifically set to use the large icon.
@@ -355,7 +370,6 @@ protected:
 
 			OutLeftContent = SNew(SSpacer);
 		}
-#endif //PLATFORM_MAC
 
 		if (bHasWindowButtons)
 		{
