Fix uninitialized members.

diff --git a/Engine/Source/Runtime/Slate/Private/Framework/Application/SlateApplication.cpp b/Engine/Source/Runtime/Slate/Private/Framework/Application/SlateApplication.cpp
--- a/Engine/Source/Runtime/Slate/Private/Framework/Application/SlateApplication.cpp
+++ b/Engine/Source/Runtime/Slate/Private/Framework/Application/SlateApplication.cpp
@@ -801,6 +801,7 @@ FSlateApplication::FSlateApplication()
 	, AverageDeltaTimeForResponsiveness( 1.0f / 30.0f )
 	, OnExitRequested()
 	, NumExternalModalWindowsActive( 0 )
+	, RootStyleNode(nullptr)
 	, bRequestLeaveDebugMode( false )
 	, bLeaveDebugForSingleStep( false )
 	, bIsExternalUIOpened( false )
@@ -811,6 +812,7 @@ FSlateApplication::FSlateApplication()
 	, bTouchFallbackToMouse( true )
 	, bSoftwareCursorAvailable( false )	
 	, bMenuAnimationsEnabled( false )
+	, AppIcon(nullptr)
 	, VirtualDesktopRect( 0,0,0,0 )
 	, NavigationConfig(MakeShared<FNavigationConfig>())
 #if WITH_EDITOR
