Update Slider widget to use current cursor position instead of previous position

diff --git a/Engine/Source/Runtime/Slate/Private/Widgets/Input/SSlider.cpp b/Engine/Source/Runtime/Slate/Private/Widgets/Input/SSlider.cpp
--- a/Engine/Source/Runtime/Slate/Private/Widgets/Input/SSlider.cpp
+++ b/Engine/Source/Runtime/Slate/Private/Widgets/Input/SSlider.cpp
@@ -293,7 +293,7 @@ FReply SSlider::OnMouseButtonDown( const FGeometry& MyGeometry, const FPointerEv
 	{
 		CachedCursor = GetCursor().Get(EMouseCursor::Default);
 		OnMouseCaptureBegin.ExecuteIfBound();
-		CommitValue(PositionToValue(MyGeometry, MouseEvent.GetLastScreenSpacePosition()));
+		CommitValue(PositionToValue(MyGeometry, MouseEvent.GetScreenSpacePosition()));
 		
 		// Release capture for controller/keyboard when switching to mouse.
 		ResetControllerState();
@@ -330,7 +330,7 @@ FReply SSlider::OnMouseMove( const FGeometry& MyGeometry, const FPointerEvent& M
 	if (HasMouseCaptureByUser(MouseEvent.GetUserIndex(), MouseEvent.GetPointerIndex()) && !IsLocked())
 	{
 		SetCursor((Orientation == Orient_Horizontal) ? EMouseCursor::ResizeLeftRight : EMouseCursor::ResizeUpDown);
-		CommitValue(PositionToValue(MyGeometry, MouseEvent.GetLastScreenSpacePosition()));
+		CommitValue(PositionToValue(MyGeometry, MouseEvent.GetScreenSpacePosition()));
 		
 		// Release capture for controller/keyboard when switching to mouse
 		ResetControllerState();
