Free standing Mode Panels populate empty when reopening the Editor.

Currently with any windows, if they have any b;ank windows we just show them. This commit catches the blank windows and does not open them

diff --git a/Engine/Source/Editor/LevelEditor/Private/SLevelEditor.cpp b/Engine/Source/Editor/LevelEditor/Private/SLevelEditor.cpp
--- a/Engine/Source/Editor/LevelEditor/Private/SLevelEditor.cpp
+++ b/Engine/Source/Editor/LevelEditor/Private/SLevelEditor.cpp
@@ -115,7 +115,7 @@ void SLevelEditor::BindCommands()
 		Actions.EditAssetNoConfirmMultiple, 
 		FExecuteAction::CreateStatic( &FLevelEditorActionCallbacks::EditAsset_Clicked, EToolkitMode::Standalone, TWeakPtr< SLevelEditor >( SharedThis( this ) ), false ),
 		FCanExecuteAction::CreateStatic( &FLevelEditorActionCallbacks::EditAsset_CanExecute ) );
-	
+
 	LevelEditorCommands->MapAction(
 		Actions.OpenSelectionInPropertyMatrix,
 		FExecuteAction::CreateStatic( &FLevelEditorActionCallbacks::OpenSelectionInPropertyMatrix_Clicked ),
@@ -1640,6 +1640,9 @@ TSharedRef<SWidget> SLevelEditor::RestoreContentArea( const TSharedRef<SDockTab>
 				)
 			)
 		));
+
+	FGlobalTabmanager::Get()->SetInitialLayoutSP(DefaultLayout);
+	
 	const EOutputCanBeNullptr OutputCanBeNullptr = EOutputCanBeNullptr::IfNoTabValid;
 	TArray<FString> RemovedOlderLayoutVersions;
 	const TSharedRef<FTabManager::FLayout> Layout = FLayoutSaveRestore::LoadFromConfig(GEditorLayoutIni,
diff --git a/Engine/Source/Runtime/Slate/Private/Framework/Docking/TabManager.cpp b/Engine/Source/Runtime/Slate/Private/Framework/Docking/TabManager.cpp
--- a/Engine/Source/Runtime/Slate/Private/Framework/Docking/TabManager.cpp
+++ b/Engine/Source/Runtime/Slate/Private/Framework/Docking/TabManager.cpp
@@ -1022,7 +1022,7 @@ FTabSpawnerEntry& FTabManager::RegisterTabSpawner(const FName TabId, const FOnSp
 
 	TSharedRef<FTabSpawnerEntry> NewSpawnerEntry = MakeShareable(new FTabSpawnerEntry(TabId, OnSpawnTab, CanSpawnTab));
 	TabSpawner.Add(TabId, NewSpawnerEntry);
-	
+
 	return NewSpawnerEntry.Get();
 }
 
@@ -1381,7 +1381,7 @@ void FTabManager::RestoreDocumentTab( FName PlaceholderId, ESearchPreference::Ty
 
 TSharedPtr<SDockTab> FTabManager::TryInvokeTab(const FTabId& TabId, bool bInvokeAsInactive)
 {
-	TSharedPtr<SDockTab> NewTab = InvokeTab_Internal(TabId, bInvokeAsInactive);
+	TSharedPtr<SDockTab> NewTab = InvokeTab_Internal(TabId, bInvokeAsInactive, true);
 	if (!NewTab.IsValid())
 	{
 		return NewTab;
@@ -1398,7 +1398,7 @@ TSharedPtr<SDockTab> FTabManager::TryInvokeTab(const FTabId& TabId, bool bInvoke
 	return NewTab;
 }
 
-TSharedPtr<SDockTab> FTabManager::InvokeTab_Internal(const FTabId& TabId, bool bInvokeAsInactive)
+TSharedPtr<SDockTab> FTabManager::InvokeTab_Internal(const FTabId& TabId, bool bInvokeAsInactive, bool bForceOpenWindowIfNeeded)
 {
 	// Tab Spawning Rules:
 	// 
@@ -1458,7 +1458,7 @@ TSharedPtr<SDockTab> FTabManager::InvokeTab_Internal(const FTabId& TabId, bool b
 	}
 
 	// Tab is not live. Figure out where to spawn it.
-	TSharedPtr<SDockingTabStack> StackToSpawnIn = FindPotentiallyClosedTab( TabId );
+	TSharedPtr<SDockingTabStack> StackToSpawnIn = bForceOpenWindowIfNeeded ? AttemptToOpenTab( TabId, true ) : FindPotentiallyClosedTab( TabId );
 
 	if (StackToSpawnIn.IsValid())
 	{
@@ -1480,8 +1480,10 @@ TSharedPtr<SDockTab> FTabManager::InvokeTab_Internal(const FTabId& TabId, bool b
 	}
 	else
 	{
-		// No layout info about this tab found; start 
-		TSharedRef<FArea> NewAreaForTab = FTabManager::NewArea(FTabManager::GetDefaultTabWindowSize(TabId))
+		const TSharedPtr<FArea> Area = FGlobalTabmanager::Get()->GetFAreaFromInitialLayoutWithTabType(TabId);
+		const TSharedRef<FArea> NewAreaForTab = Area.IsValid() ? Area.ToSharedRef() : FTabManager::NewArea(FTabManager::GetDefaultTabWindowSize(TabId));
+
+		NewAreaForTab
 		->Split
 		(
 			FTabManager::NewStack()
@@ -1503,6 +1505,11 @@ TSharedPtr<SDockTab> FTabManager::InvokeTab_Internal(const FTabId& TabId, bool b
 }
 
 TSharedPtr<SDockingTabStack> FTabManager::FindPotentiallyClosedTab( const FTabId& ClosedTabId )
+{
+	return AttemptToOpenTab( ClosedTabId );
+}
+
+TSharedPtr<SDockingTabStack> FTabManager::AttemptToOpenTab( const FTabId& ClosedTabId, bool bForceOpenWindowIfNeeded )
 {
 	TSharedPtr<SDockingTabStack> StackWithClosedTab;
 
@@ -1513,14 +1520,13 @@ TSharedPtr<SDockingTabStack> FTabManager::FindPotentiallyClosedTab( const FTabId
 	if ( CollapsedAreaWithMatchingTabIndex != INDEX_NONE )
 	{
 		TSharedRef<FTabManager::FArea> CollapsedAreaWithMatchingTab = CollapsedDockAreas[CollapsedAreaWithMatchingTabIndex];
-
-		TSharedPtr<SDockingArea> RestoredArea = RestoreArea(CollapsedAreaWithMatchingTab, GetPrivateApi().GetParentWindow());
+		
+		TSharedPtr<SDockingArea> RestoredArea = RestoreArea(CollapsedDockAreas[CollapsedAreaWithMatchingTabIndex],
+			GetPrivateApi().GetParentWindow(), false, EOutputCanBeNullptr::Never, bForceOpenWindowIfNeeded);
 		check(RestoredArea.IsValid());
-
 		// We have just un-collapsed this dock area.
 		// Don't rely on the collapsed tab index: RestoreArea() can end up kicking the task graph which could do other tab work and modify the CollapsedDockAreas array.
 		CollapsedDockAreas.Remove(CollapsedAreaWithMatchingTab);
-
 		if (RestoredArea.IsValid())
 		{
 			StackWithClosedTab = FindTabInLiveArea(TabMatcher, StaticCastSharedRef<SDockingArea>(RestoredArea->AsShared()));
@@ -1558,7 +1564,7 @@ FUIAction FTabManager::GetUIActionForTabSpawnerMenuEntry(TSharedPtr<FTabSpawnerE
 
 void FTabManager::InvokeTabForMenu( FName TabId )
 {
-	TryInvokeTab(TabId);
+	TryInvokeTab(TabId, true);
 }
 
 void FTabManager::InsertDocumentTab(FName PlaceholderId, const FSearchPreference& SearchPreference, const TSharedRef<SDockTab>& UnmanagedTab, bool bPlaySpawnAnim)
@@ -1594,7 +1600,7 @@ void FTabManager::OpenUnmanagedTab(FName PlaceholderId, const FSearchPreference&
 	}
 	else
 	{
-		TSharedPtr<SDockingTabStack> StackToSpawnIn = FindPotentiallyClosedTab( PlaceholderId );
+		TSharedPtr<SDockingTabStack> StackToSpawnIn = AttemptToOpenTab( PlaceholderId, true );
 		if (StackToSpawnIn.IsValid())
 		{
 			StackToSpawnIn->OpenTab(UnmanagedTab);
@@ -1621,14 +1627,14 @@ FTabManager::FTabManager( const TSharedPtr<SDockTab>& InOwnerTab, const TSharedR
 	LocalWorkspaceMenuRoot = FWorkspaceItem::NewGroup(LOCTEXT("LocalWorkspaceRoot", "Local Workspace Root"));
 }
 
-TSharedPtr<SDockingArea> FTabManager::RestoreArea(const TSharedRef<FArea>& AreaToRestore, const TSharedPtr<SWindow>& InParentWindow, const bool bEmbedTitleAreaContent, const EOutputCanBeNullptr OutputCanBeNullptr)
+TSharedPtr<SDockingArea> FTabManager::RestoreArea(const TSharedRef<FArea>& AreaToRestore, const TSharedPtr<SWindow>& InParentWindow, const bool bEmbedTitleAreaContent, const EOutputCanBeNullptr OutputCanBeNullptr, bool bForceOpenWindowIfNeeded)
 {
 	// Sidebar tabs for this area
 	FSidebarTabLists SidebarTabs;
 
 	TemporarilySidebaredTabs.Empty();
 
-	if (TSharedPtr<SDockingNode> RestoredNode = RestoreArea_Helper(AreaToRestore, InParentWindow, bEmbedTitleAreaContent, SidebarTabs, OutputCanBeNullptr))
+	if (TSharedPtr<SDockingNode> RestoredNode = RestoreArea_Helper(AreaToRestore, InParentWindow, bEmbedTitleAreaContent, SidebarTabs, OutputCanBeNullptr, bForceOpenWindowIfNeeded))
 	{
 		TSharedRef<SDockingArea> RestoredArea = StaticCastSharedRef<SDockingArea>(RestoredNode->AsShared());
 
@@ -1655,7 +1661,8 @@ TSharedPtr<SDockingArea> FTabManager::RestoreArea(const TSharedRef<FArea>& AreaT
 	}
 }
 
-TSharedPtr<SDockingNode> FTabManager::RestoreArea_Helper(const TSharedRef<FLayoutNode>& LayoutNode, const TSharedPtr<SWindow>& ParentWindow, const bool bEmbedTitleAreaContent, FSidebarTabLists& OutSidebarTabs, const EOutputCanBeNullptr OutputCanBeNullptr)
+TSharedPtr<SDockingNode> FTabManager::RestoreArea_Helper(const TSharedRef<FLayoutNode>& LayoutNode, const TSharedPtr<SWindow>& ParentWindow, const bool bEmbedTitleAreaContent,
+	FSidebarTabLists& OutSidebarTabs, const EOutputCanBeNullptr OutputCanBeNullptr, bool bForceOpenWindowIfNeeded)
 {
 	TSharedPtr<FTabManager::FStack> NodeAsStack = LayoutNode->AsStack();
 	TSharedPtr<FTabManager::FSplitter> NodeAsSplitter = LayoutNode->AsSplitter();
@@ -1791,28 +1798,31 @@ TSharedPtr<SDockingNode> FTabManager::RestoreArea_Helper(const TSharedRef<FLayou
 			TArray<TSharedRef<SDockingNode>> DockingNodes;
 			if (CanRestoreSplitterContent(DockingNodes, NodeAsArea.ToSharedRef(), NewWindow, OutSidebarTabs, OutputCanBeNullptr))
 			{
-				// Create SplitterWidget only if it will be filled with at least 1 DockingNodes
-				// Any windows that were "pulled out" of a dock area should be children of the window in which the parent dock area resides.
-				if (bIsChildWindow)
+				NewWindow->SetContent(SAssignNew(NewDockAreaWidget, SDockingArea, SharedThis(this), NodeAsArea.ToSharedRef()).ParentWindow(NewWindow));
+
+				// Restore content
+				if (!bCanOutputBeNullptr)
 				{
-					FSlateApplication::Get().AddWindowAsNativeChild(NewWindow, ParentWindow.ToSharedRef())->SetContent(
-						SAssignNew(NewDockAreaWidget, SDockingArea, SharedThis(this), NodeAsArea.ToSharedRef()).ParentWindow(NewWindow)
-					);
+					RestoreSplitterContent(NodeAsArea.ToSharedRef(), NewDockAreaWidget.ToSharedRef(), NewWindow, OutSidebarTabs);
 				}
 				else
 				{
-					FSlateApplication::Get().AddWindow(NewWindow)->SetContent(
-						SAssignNew(NewDockAreaWidget, SDockingArea, SharedThis(this), NodeAsArea.ToSharedRef()).ParentWindow(NewWindow)
-					);
+					RestoreSplitterContent(DockingNodes, NewDockAreaWidget.ToSharedRef());
 				}
-				// Restore content
-				if (!bCanOutputBeNullptr)
+
+				if (bIsChildWindow)
 				{
-					RestoreSplitterContent(NodeAsArea.ToSharedRef(), NewDockAreaWidget.ToSharedRef(), NewWindow, OutSidebarTabs);
+					// Recursively check to see how many actually spawned tabs there are in this dock area. If there are none we will not spawn a useless window
+					const int32 TotalNumTabs = NewDockAreaWidget->GetNumTabs();
+
+					if (TotalNumTabs > 0 || bForceOpenWindowIfNeeded)
+					{
+						FSlateApplication::Get().AddWindowAsNativeChild(NewWindow, ParentWindow.ToSharedRef());
+					}
 				}
 				else
 				{
-					RestoreSplitterContent(DockingNodes, NewDockAreaWidget.ToSharedRef());
+					FSlateApplication::Get().AddWindow(NewWindow);
 				}
 			}
 		}
@@ -2235,6 +2245,50 @@ FVector2D FTabManager::GetDefaultTabWindowSize(const FTabId& TabId)
 	return WindowSize;
 }
 
+bool FTabManager::HasAnyTabWithTabId( const TSharedRef<FLayoutNode>& SomeNode, const FName& InTabTypeToMatch ) const
+{
+	struct TabWithIdMatcher
+	{
+		TabWithIdMatcher(const FTabManager* InTabManager, FName& InTabTypeToMatch) : TabManager(InTabManager), TabTypeToMatch(InTabTypeToMatch)
+		{
+		}
+		
+		const FTabManager* TabManager;
+		const FName& TabTypeToMatch;
+
+		bool operator()(const FTab& Candidate) const
+		{
+			return TabManager->IsValidTabForSpawning(Candidate) && Candidate.TabId.TabType == TabTypeToMatch;
+		}
+	};
+	FName TabTypeToMatch = InTabTypeToMatch;
+	const TabWithIdMatcher FindTabWithId(this, TabTypeToMatch);
+
+	return HasAnyMatchingTabs(SomeNode, FindTabWithId);
+}
+
+TSharedPtr<FTabManager::FArea> FTabManager::GetFAreaFromInitialLayoutWithTabType( const FTabId& InTabIdToMatch ) const
+{
+	for (const TSharedRef<FArea>& Area : FGlobalTabmanager::Get()->GetInitialLayoutSP()->Areas)
+	{
+		if (HasAnyTabWithTabId(Area, InTabIdToMatch.TabType))
+		{
+			return Area.ToSharedPtr();
+		}
+	}
+	return nullptr;
+}
+
+void FGlobalTabmanager::SetInitialLayoutSP(TSharedPtr<FTabManager::FLayout> InLayout)
+{
+	InitialLayoutSP = InLayout;
+}
+
+TSharedPtr<FTabManager::FLayout> FGlobalTabmanager::GetInitialLayoutSP()
+{
+	return InitialLayoutSP;
+}
+
 template<typename MatchFunctorType>
 bool FTabManager::HasAnyMatchingTabs( const TSharedRef<FTabManager::FLayoutNode>& SomeNode, const MatchFunctorType& Matcher )
 {
@@ -2954,8 +3008,10 @@ void FProxyTabmanager::OpenUnmanagedTab(FName PlaceholderId, const FSearchPrefer
 	TSharedPtr<SWindow> ParentWindowPtr = ParentWindow.Pin();
 	if (ensure(ParentWindowPtr.IsValid()))
 	{
-		// No layout info about this tab found; start 
-		TSharedRef<FArea> NewAreaForTab = FTabManager::NewPrimaryArea()
+		const TSharedPtr<FArea> Area = FGlobalTabmanager::Get()->GetFAreaFromInitialLayoutWithTabType(UnmanagedTab->GetLayoutIdentifier());
+		const TSharedRef<FArea> NewAreaForTab =  Area.IsValid() ? Area.ToSharedRef() : NewPrimaryArea();
+
+		NewAreaForTab
 			->Split
 			(
 				FTabManager::NewStack()
diff --git a/Engine/Source/Runtime/Slate/Public/Framework/Docking/TabManager.h b/Engine/Source/Runtime/Slate/Public/Framework/Docking/TabManager.h
--- a/Engine/Source/Runtime/Slate/Public/Framework/Docking/TabManager.h
+++ b/Engine/Source/Runtime/Slate/Public/Framework/Docking/TabManager.h
@@ -1041,7 +1041,7 @@ class FTabManager : public TSharedFromThis<FTabManager>
 
 		SLATE_API void MakeSpawnerMenuEntry( FMenuBuilder &PopulateMe, const TSharedPtr<FTabSpawnerEntry> &InSpawnerNode );
 
-		SLATE_API TSharedPtr<SDockTab> InvokeTab_Internal(const FTabId& TabId, bool bInvokeAsInactive = false);
+		SLATE_API TSharedPtr<SDockTab> InvokeTab_Internal(const FTabId& TabId, bool bInvokeAsInactive = false, bool bForceOpenWindowIfNeeded = false);
 
 		/** Finds the last major or nomad tab in a particular window. */
 		SLATE_API TSharedPtr<SDockTab> FindLastTabInWindow(TSharedPtr<SWindow> Window) const;
@@ -1057,10 +1057,12 @@ class FTabManager : public TSharedFromThis<FTabManager>
 
 		SLATE_API FTabManager( const TSharedPtr<SDockTab>& InOwnerTab, const TSharedRef<FTabManager::FTabSpawner> & InNomadTabSpawner );
 
-		SLATE_API TSharedPtr<SDockingArea> RestoreArea(
-			const TSharedRef<FArea>& AreaToRestore, const TSharedPtr<SWindow>& InParentWindow, const bool bEmbedTitleAreaContent = false, const EOutputCanBeNullptr OutputCanBeNullptr = EOutputCanBeNullptr::Never);
+	SLATE_API TSharedPtr<SDockingArea> RestoreArea(
+		const TSharedRef<FArea>& AreaToRestore, const TSharedPtr<SWindow>& InParentWindow, const bool bEmbedTitleAreaContent = false,
+		const EOutputCanBeNullptr OutputCanBeNullptr = EOutputCanBeNullptr::Never, bool bForceOpenWindowIfNeeded = false);
 
-		SLATE_API TSharedPtr<class SDockingNode> RestoreArea_Helper(const TSharedRef<FLayoutNode>& LayoutNode, const TSharedPtr<SWindow>& ParentWindow, const bool bEmbedTitleAreaContent, FSidebarTabLists& OutSidebarTabs, const EOutputCanBeNullptr OutputCanBeNullptr = EOutputCanBeNullptr::Never);
+	SLATE_API TSharedPtr<class SDockingNode> RestoreArea_Helper(const TSharedRef<FLayoutNode>& LayoutNode, const TSharedPtr<SWindow>& ParentWindow, const bool bEmbedTitleAreaContent,
+		FSidebarTabLists& OutSidebarTabs, const EOutputCanBeNullptr OutputCanBeNullptr = EOutputCanBeNullptr::Never, bool bForceOpenWindowIfNeeded = false);
 
 		/**
 		 * Use CanRestoreSplitterContent + RestoreSplitterContent when the output of its internal RestoreArea_Helper can be a nullptr.
@@ -1100,6 +1102,13 @@ class FTabManager : public TSharedFromThis<FTabManager>
 		 */
 		SLATE_API bool HasValidOpenTabs( const TSharedRef<FTabManager::FLayoutNode>& SomeNode ) const;
 
+		/**
+		 * Gets a TSharedPtr<FArea> with the given FTabId if present, else it return nullptr
+		 *
+		 * @param InTabIdToMatch the const &FTabId for which to find the FArea
+		 */
+	    SLATE_API TSharedPtr<FArea> GetFAreaFromInitialLayoutWithTabType(const FTabId& InTabIdToMatch ) const;
+
 	protected:
 		SLATE_API bool HasValidTabs( const TSharedRef<FTabManager::FLayoutNode>& SomeNode ) const;
 
@@ -1140,6 +1149,16 @@ class FTabManager : public TSharedFromThis<FTabManager>
 		/** Checks all dock areas and adds up the number of open tabs and unique parent windows in the manager */
 		SLATE_API void GetRecordableStats( int32& OutTabCount, TArray<TSharedPtr<SWindow>>& OutUniqueParentWindows ) const;
 
+	TSharedPtr<SDockingTabStack> AttemptToOpenTab( const FTabId& ClosedTabId, bool bForceOpenWindowIfNeeded = false );
+
+	/**
+	 * Returns true if the given FLayoutNode contains a tab with the given FTabId.TabType
+	 *
+	 * @param InTabTypeToMatch the given FTabId.TabType to look to see if this tab manager manages it
+	 * @param SomeNode the TSharedRef<FLayoutNode> in which to look for the tab with FTabId.TabType == InTabTypeToMatch
+	 */
+	bool HasAnyTabWithTabId( const TSharedRef<FTabManager::FLayoutNode>& SomeNode, const FName& InTabTypeToMatch ) const;
+
 	protected:
 		FTabSpawner TabSpawner;
 		TSharedRef<FTabSpawner> NomadTabSpawner;
@@ -1346,7 +1365,20 @@ protected:
 public:
 	SLATE_API virtual void OnTabManagerClosing() override;
 
+	/**
+	 * Sets the initial layout shared pointer, which can be used later to get the current layout for tabs which were not
+	 * spawned on initialization
+	 *
+	 * @param InLayout the FLayout which was loaded at layout load time
+	 */
+	SLATE_API void SetInitialLayoutSP(TSharedPtr<FTabManager::FLayout> InLayout);
 
+	/**
+	 * Gets the initial layout shared pointer, which can be used later to get the current layout for tabs which were not
+	 * spawned on initialization
+	 */
+	SLATE_API TSharedPtr<FTabManager::FLayout> GetInitialLayoutSP();
+	
 private:
 	
 	/** Pairs of Major Tab and the TabManager that manages tabs within it. */
@@ -1424,6 +1456,12 @@ private:
 
 	/**  */
 	TSharedPtr<FProxyTabmanager> ProxyTabManager;
+
+	/**
+	 * the initial layout shared pointer, which can be used to get the current layout for tabs which were not
+	 * spawned on initialization
+	 */
+	TSharedPtr<FTabManager::FLayout> InitialLayoutSP;
 };
 
 //#HACK VREDITOR - Had to introduce the proxy tab manager to steal asset tabs.
