Crash upon opening content browser window (and some others)

We had some signatures change with the commit of UE-183468 and this one should have been updated and was missed. Any tabs that are launched into windows on demand now have to pass the new bForceOpenWindowIfNeeded argument to the InvokeTab_Internal method. This is what allows us to detect when we may have invalid content versus content that is just not yet loaded. If we don't pass that as true for such a window then the window may not launch and unpredictable behavior (such as a crash) could ensue.

diff --git a/Engine/Source/Runtime/Slate/Private/Framework/Docking/TabManager.cpp b/Engine/Source/Runtime/Slate/Private/Framework/Docking/TabManager.cpp
--- a/Engine/Source/Runtime/Slate/Private/Framework/Docking/TabManager.cpp
+++ b/Engine/Source/Runtime/Slate/Private/Framework/Docking/TabManager.cpp
@@ -1476,7 +1476,7 @@ TSharedPtr<SDockTab> FTabManager::InvokeTab_Internal(const FTabId& TabId, bool b
 	else if ( FGlobalTabmanager::Get() != SharedThis(this) && NomadTabSpawner->Contains(TabId.TabType) )
 	{
 		// This tab could have been spawned in the global tab manager since it has a nomad tab spawner
-		return FGlobalTabmanager::Get()->InvokeTab_Internal(TabId, bInvokeAsInactive);
+		return FGlobalTabmanager::Get()->InvokeTab_Internal(TabId, bInvokeAsInactive, bForceOpenWindowIfNeeded);
 	}
 	else
 	{
