Plugins window does not populate with its tab in focus

With the commit of UE-183468 we had some signatures change, and this was an overzealous correction. This reverts it to the code that was here before the fix.

diff --git a/Engine/Source/Runtime/Slate/Private/Framework/Docking/TabManager.cpp b/Engine/Source/Runtime/Slate/Private/Framework/Docking/TabManager.cpp
--- a/Engine/Source/Runtime/Slate/Private/Framework/Docking/TabManager.cpp
+++ b/Engine/Source/Runtime/Slate/Private/Framework/Docking/TabManager.cpp
@@ -1564,7 +1564,7 @@ FUIAction FTabManager::GetUIActionForTabSpawnerMenuEntry(TSharedPtr<FTabSpawnerE
 
 void FTabManager::InvokeTabForMenu( FName TabId )
 {
-	TryInvokeTab(TabId, true);
+	TryInvokeTab(TabId);
 }
 
 void FTabManager::InsertDocumentTab(FName PlaceholderId, const FSearchPreference& SearchPreference, const TSharedRef<SDockTab>& UnmanagedTab, bool bPlaySpawnAnim)
