TabManager: Remove ensure checking for stale tabs that was accidentally added in 27874281

diff --git a/Engine/Source/Runtime/Slate/Private/Framework/Docking/TabManager.cpp b/Engine/Source/Runtime/Slate/Private/Framework/Docking/TabManager.cpp
--- a/Engine/Source/Runtime/Slate/Private/Framework/Docking/TabManager.cpp
+++ b/Engine/Source/Runtime/Slate/Private/Framework/Docking/TabManager.cpp
@@ -2797,7 +2797,6 @@ void FGlobalTabmanager::SaveAllVisualState()
 		const TSharedPtr<FTabManager> SubManagerTab = SubTabManagers[ManagerIndex].TabManager.Pin();		
 		if (SubManagerTab.IsValid())
 		{
-			ensure(SubTabManagers[ManagerIndex].MajorTab.IsValid());
 			SubManagerTab->SavePersistentLayout();
 		}
 	}
