Running Automation in Unreal Insights can result in a crash

In this case when we run this code outside of UE, the initial layout can be null. A null check can fix this, and a nullptr is a perfectly valid return value from here as this just assumes the user didn't have an initial layout saved (aside from the default).

diff --git a/Engine/Source/Runtime/Slate/Private/Framework/Docking/TabManager.cpp b/Engine/Source/Runtime/Slate/Private/Framework/Docking/TabManager.cpp
--- a/Engine/Source/Runtime/Slate/Private/Framework/Docking/TabManager.cpp
+++ b/Engine/Source/Runtime/Slate/Private/Framework/Docking/TabManager.cpp
@@ -2269,11 +2269,15 @@ bool FTabManager::HasAnyTabWithTabId( const TSharedRef<FLayoutNode>& SomeNode, c
 
 TSharedPtr<FTabManager::FArea> FTabManager::GetFAreaFromInitialLayoutWithTabType( const FTabId& InTabIdToMatch ) const
 {
-	for (const TSharedRef<FArea>& Area : FGlobalTabmanager::Get()->GetInitialLayoutSP()->Areas)
+	const TSharedPtr<FTabManager::FLayout> InitialLayoutSP = FGlobalTabmanager::Get()->GetInitialLayoutSP();
+	if (InitialLayoutSP.IsValid())
 	{
-		if (HasAnyTabWithTabId(Area, InTabIdToMatch.TabType))
+		for (const TSharedRef<FArea>& Area : InitialLayoutSP->Areas)
 		{
-			return Area.ToSharedPtr();
+			if (HasAnyTabWithTabId(Area, InTabIdToMatch.TabType))
+			{
+				return Area.ToSharedPtr();
+			}
 		}
 	}
 	return nullptr;
