Fix VirtualKeyboard reopening automatically

diff --git a/Engine/Source/Runtime/Slate/Private/Widgets/Text/SlateEditableTextLayout.cpp b/Engine/Source/Runtime/Slate/Private/Widgets/Text/SlateEditableTextLayout.cpp
--- a/Engine/Source/Runtime/Slate/Private/Widgets/Text/SlateEditableTextLayout.cpp
+++ b/Engine/Source/Runtime/Slate/Private/Widgets/Text/SlateEditableTextLayout.cpp
@@ -1271,21 +1271,6 @@ FReply FSlateEditableTextLayout::HandleMouseButtonDown(const FGeometry& MyGeomet
 				// should reset the selection range to the caret's position.
 				bWasFocusedByLastMouseDown = true;
 			}
-			else
-			{
-				// On platforms using a virtual keyboard open the virtual keyboard again 
-				if (FPlatformApplicationMisc::RequiresVirtualKeyboard())
-				{
-					if (!OwnerWidget->IsTextReadOnly())
-					{
-						if (OwnerWidget->GetVirtualKeyboardTrigger() == EVirtualKeyboardTrigger::OnAllFocusEvents ||
-							OwnerWidget->GetVirtualKeyboardTrigger() == EVirtualKeyboardTrigger::OnFocusByPointer)
-						{
-							FSlateApplication::Get().ShowVirtualKeyboard(true, InMouseEvent.GetUserIndex(), VirtualKeyboardEntry.ToSharedRef());
-						}
-					}
-				}
-			}
 
 			if (InMouseEvent.GetEffectingButton() == EKeys::LeftMouseButton)
 			{
@@ -1332,6 +1317,28 @@ FReply FSlateEditableTextLayout::HandleMouseButtonUp(const FGeometry& MyGeometry
 	// The mouse must have been captured by either left or right button down before we'll process mouse ups
 	if (OwnerWidget->GetSlateWidget()->HasMouseCapture())
 	{
+		if (InMouseEvent.GetEffectingButton() == EKeys::LeftMouseButton ||
+			InMouseEvent.GetEffectingButton() == EKeys::RightMouseButton)
+		{
+			// Am I getting focus right now?
+			const bool bIsGettingFocus = !OwnerWidget->GetSlateWidget()->HasAnyUserFocus().IsSet();
+			if (!bIsGettingFocus)
+			{
+				// On platforms using a virtual keyboard open the virtual keyboard again 
+				if (FPlatformApplicationMisc::RequiresVirtualKeyboard())
+				{
+					if (!OwnerWidget->IsTextReadOnly())
+					{
+						if (OwnerWidget->GetVirtualKeyboardTrigger() == EVirtualKeyboardTrigger::OnAllFocusEvents ||
+							OwnerWidget->GetVirtualKeyboardTrigger() == EVirtualKeyboardTrigger::OnFocusByPointer)
+						{
+								FSlateApplication::Get().ShowVirtualKeyboard(true, InMouseEvent.GetUserIndex(), VirtualKeyboardEntry.ToSharedRef());
+						}
+					}
+				}
+			}
+		}
+
 		if (InMouseEvent.GetEffectingButton() == EKeys::LeftMouseButton && bIsDragSelecting)
 		{
 			// No longer drag-selecting
