Add FNotificationInfo::ShowCopyToClipboadHyperlink helper method to set a "Copy to Clipboard" hyperlink to the notification that copies the notif text and subtext on the operating system clipboard
#rnx

diff --git a/Engine/Source/Runtime/Slate/Private/Widgets/Notifications/SNotificationList.cpp b/Engine/Source/Runtime/Slate/Private/Widgets/Notifications/SNotificationList.cpp
--- a/Engine/Source/Runtime/Slate/Private/Widgets/Notifications/SNotificationList.cpp
+++ b/Engine/Source/Runtime/Slate/Private/Widgets/Notifications/SNotificationList.cpp
@@ -4,9 +4,10 @@
 #include "Animation/CurveHandle.h"
 #include "Animation/CurveSequence.h"
 #include "Application/ThrottleManager.h"
+#include "Framework/Application/SlateApplication.h"
+#include "HAL/PlatformApplicationMisc.h"
 #include "Widgets/SBoxPanel.h"
 #include "Widgets/SOverlay.h"
-#include "Framework/Application/SlateApplication.h"
 #include "Widgets/Layout/SBorder.h"
 #include "Widgets/Images/SImage.h"
 #include "Widgets/Text/STextBlock.h"
@@ -19,6 +20,8 @@
 #include "Styling/StyleColors.h"
 #include "Widgets/Notifications/SNotificationBackground.h"
 
+#define LOCTEXT_NAMESPACE "SNotificationList"
+
 /////////////////////////////////////////////////
 // SNotificationExtendable
 
@@ -739,8 +742,36 @@ private:
 	TSharedPtr<INotificationWidget> NotificationWidget;
 };
 
+/////////////////////////////////////////////////
+// FNotificationInfo
+
+void FNotificationInfo::ShowCopyToClipboadHyperlink()
+{
+	HyperlinkText = LOCTEXT("CopyToClipboard", "Copy to Clipboard");
+	Hyperlink = FSimpleDelegate::CreateLambda(
+		[Text = Text, SubText = SubText]()
+		{
+			TStringBuilder<1024> StrBuilder;
+			if (Text.IsSet())
+			{
+				StrBuilder << Text.Get().ToString();
+			}
+			if (SubText.IsSet())
+			{
+				if (StrBuilder.Len() > 0)
+				{
+					StrBuilder << TEXT("\n");
+				}
+				StrBuilder << SubText.Get().ToString();
+			}
+
+			FPlatformApplicationMisc::ClipboardCopy(StrBuilder.ToString());
+		});
+}
+
 ///////////////////////////////////////////////////
 //// SNotificationList
+
 TSharedRef<SNotificationItem> SNotificationList::AddNotification(const FNotificationInfo& Info)
 {
 	TSharedPtr<SNotificationExtendable> NewItem;
@@ -834,3 +865,5 @@ void SNotificationList::Construct(const FArguments& InArgs)
 		SAssignNew(MessageItemBoxPtr, SVerticalBox)
 	];
 }
+
+#undef LOCTEXT_NAMESPACE 
diff --git a/Engine/Source/Runtime/Slate/Public/Widgets/Notifications/SNotificationList.h b/Engine/Source/Runtime/Slate/Public/Widgets/Notifications/SNotificationList.h
--- a/Engine/Source/Runtime/Slate/Public/Widgets/Notifications/SNotificationList.h
+++ b/Engine/Source/Runtime/Slate/Public/Widgets/Notifications/SNotificationList.h
@@ -159,6 +159,12 @@ struct FNotificationInfo
 		bAllowThrottleWhenFrameRateIsLow(true)
 	{ };
 
+	/**
+	 * Shows a "Copy to Clipboard" hyperlink that when clicked copies the text and sub-text to the operating system clipboard
+	 * @note Text and SubText must be set prior to calling this method for them to be included
+	 */
+	SLATE_API void ShowCopyToClipboadHyperlink();
+
 	/** If set, overrides the entire content of the notification with this widget */
 	TSharedPtr<INotificationWidget> ContentWidget;
 
