Problem: Curve editor window fails to open after closing and cannot be moved within a new User Widget Blueprint

Root cause: CL 27874281 brought in Some changes  that modified the flow of control for tabs in order to fix UE-196963 but it caused this issue, and a similar issue would happen in the same way for any tabs opening as the curve editor does

Solution: Revert CL 27874281 and we can find a different fix for UE-196963

diff --git a/Engine/Source/Runtime/Slate/Private/Framework/Docking/TabManager.cpp b/Engine/Source/Runtime/Slate/Private/Framework/Docking/TabManager.cpp
--- a/Engine/Source/Runtime/Slate/Private/Framework/Docking/TabManager.cpp
+++ b/Engine/Source/Runtime/Slate/Private/Framework/Docking/TabManager.cpp
@@ -2089,11 +2089,6 @@ TSharedPtr<SDockTab> FTabManager::SpawnTab(const FTabId& TabId, const TSharedPtr
 			bSpawningAllowedBySpawner = false;
 		}
 	}
-	else
-	{
-	
-		bSpawningAllowedBySpawner = false;
-	}
 
 	// The tab was allowed to be spawned but failed for some reason
 	if (bSpawningAllowedBySpawner && !NewTabWidget.IsValid())
