Problem:
Undo history fails to open after closing with it docked.

Root cause: We shouldn't be specifying saved positions for anything but positionally specified windows while  trying to invoke a tab. This code path was only put in to support repositioning specified windows from the layout, so any other type of EWindowPlacement is untested and unpredictable at this point.

Solution: Added a new method to TabManager which gets the proper FArea to load a tab. If it should be in a positionally specified window according to the initial layout, it is loaded there, otherwise it returns a new FArea to load it.

diff --git a/Engine/Source/Runtime/Slate/Private/Framework/Docking/TabManager.cpp b/Engine/Source/Runtime/Slate/Private/Framework/Docking/TabManager.cpp
--- a/Engine/Source/Runtime/Slate/Private/Framework/Docking/TabManager.cpp
+++ b/Engine/Source/Runtime/Slate/Private/Framework/Docking/TabManager.cpp
@@ -318,6 +318,7 @@ TSharedRef<FTabManager::FLayoutNode> FTabManager::FLayout::NewFromString_Helper(
 	}	
 }
 
+
 TSharedPtr<FTabManager::FLayout> FTabManager::FLayout::NewFromString( const FString& LayoutAsText )
 {
 	TSharedPtr<FJsonObject> JsonObject;
@@ -1459,7 +1460,7 @@ TSharedPtr<SDockTab> FTabManager::InvokeTab_Internal(const FTabId& TabId, bool b
 
 	// Tab is not live. Figure out where to spawn it.
 	TSharedPtr<SDockingTabStack> StackToSpawnIn = bForceOpenWindowIfNeeded ? AttemptToOpenTab( TabId, true ) : FindPotentiallyClosedTab( TabId );
-
+	
 	if (StackToSpawnIn.IsValid())
 	{
 		const TSharedPtr<SDockTab> NewTab = SpawnTab(TabId, TSharedPtr<SWindow>());
@@ -1480,8 +1481,7 @@ TSharedPtr<SDockTab> FTabManager::InvokeTab_Internal(const FTabId& TabId, bool b
 	}
 	else
 	{
-		const TSharedPtr<FArea> Area = FGlobalTabmanager::Get()->GetFAreaFromInitialLayoutWithTabType(TabId);
-		const TSharedRef<FArea> NewAreaForTab = Area.IsValid() ? Area.ToSharedRef() : FTabManager::NewArea(FTabManager::GetDefaultTabWindowSize(TabId));
+		const TSharedRef<FArea> NewAreaForTab = GetFAreaForFTabId(TabId);
 
 		NewAreaForTab
 		->Split
@@ -2278,6 +2278,20 @@ TSharedPtr<FTabManager::FArea> FTabManager::GetFAreaFromInitialLayoutWithTabType
 	return nullptr;
 }
 
+TSharedRef<FTabManager::FArea> FTabManager::GetFAreaForFTabId(const FTabId& TabId)
+{
+	if (const TSharedPtr<FArea> AreaFromInitiallyLoadedLayout = FGlobalTabmanager::Get()->GetFAreaFromInitialLayoutWithTabType(TabId))
+	{
+		/* we must reuse positions from the initial layout for positionally specified floating windows. If we don't
+		* do this then any persisted floating windows load in a big cluster in the middle on top of one another */
+		if ( AreaFromInitiallyLoadedLayout->DefinesPositionallySpecifiedFloatingWindow() )
+		{
+			return AreaFromInitiallyLoadedLayout.ToSharedRef();
+		}
+	}
+	return NewArea( GetDefaultTabWindowSize(TabId) );
+}
+
 void FGlobalTabmanager::SetInitialLayoutSP(TSharedPtr<FTabManager::FLayout> InLayout)
 {
 	InitialLayoutSP = InLayout;
diff --git a/Engine/Source/Runtime/Slate/Public/Framework/Docking/TabManager.h b/Engine/Source/Runtime/Slate/Public/Framework/Docking/TabManager.h
--- a/Engine/Source/Runtime/Slate/Public/Framework/Docking/TabManager.h
+++ b/Engine/Source/Runtime/Slate/Public/Framework/Docking/TabManager.h
@@ -615,16 +615,21 @@ class FTabManager : public TSharedFromThis<FTabManager>
 				TArray< TSharedRef<FLayoutNode> > ChildNodes;
 		};
 
-
 		class FArea : public FSplitter
 		{
 				friend class FTabManager;
 		
-			public:			
+			public:
+				/** An enum which specifies how content appears within a window. */
 				enum EWindowPlacement
 				{
+					/** The content is docked within a primary window. */
 					Placement_NoWindow,
+
+					/** The content is docked within a floating window that is initially automatically positioned. */
 					Placement_Automatic,
+
+					/** The content is docked within a floating window with positioning specified by an FArea. */
 					Placement_Specified
 				};
 
@@ -675,6 +680,14 @@ class FTabManager : public TSharedFromThis<FTabManager>
 				{
 				}
 
+			/**
+			 * Returns true if this FArea defines a positionally specified window, else it returns false
+			*/
+			bool DefinesPositionallySpecifiedFloatingWindow()const
+			{
+				return WindowPlacement == Placement_Specified;
+			}
+
 			protected:
 				FArea( const float InWidth, const float InHeight )
 				: WindowPlacement(Placement_Automatic)
@@ -1109,7 +1122,7 @@ class FTabManager : public TSharedFromThis<FTabManager>
 		 */
 	    SLATE_API TSharedPtr<FArea> GetFAreaFromInitialLayoutWithTabType(const FTabId& InTabIdToMatch ) const;
 
-	protected:
+protected:
 		SLATE_API bool HasValidTabs( const TSharedRef<FTabManager::FLayoutNode>& SomeNode ) const;
 
 		/**
@@ -1159,6 +1172,13 @@ class FTabManager : public TSharedFromThis<FTabManager>
 	 */
 	bool HasAnyTabWithTabId( const TSharedRef<FTabManager::FLayoutNode>& SomeNode, const FName& InTabTypeToMatch ) const;
 
+	/**
+	 * Given FTabId returns the proper TSharedRef<FArea> to load it
+	 *
+	 * @param TabId the FTabId for which to find the FArea 
+	 */
+	TSharedRef<FArea> GetFAreaForFTabId(const FTabId& TabId);
+
 	protected:
 		FTabSpawner TabSpawner;
 		TSharedRef<FTabSpawner> NomadTabSpawner;
