Adding event to fire when scroll box scroll bar visibility changes

diff --git a/Engine/Source/Runtime/Slate/Private/Widgets/Layout/SScrollBar.cpp b/Engine/Source/Runtime/Slate/Private/Widgets/Layout/SScrollBar.cpp
--- a/Engine/Source/Runtime/Slate/Private/Widgets/Layout/SScrollBar.cpp
+++ b/Engine/Source/Runtime/Slate/Private/Widgets/Layout/SScrollBar.cpp
@@ -11,6 +11,7 @@
 void SScrollBar::Construct(const FArguments& InArgs)
 {
 	OnUserScrolled = InArgs._OnUserScrolled;
+	OnScrollBarVisibilityChanged = InArgs._OnScrollBarVisibilityChanged;
 	Orientation = InArgs._Orientation;
 	DragFocusCause = InArgs._DragFocusCause;
 	UserVisibility = InArgs._Visibility;
@@ -91,6 +92,11 @@ void SScrollBar::SetOnUserScrolled( const FOnUserScrolled& InHandler )
 	OnUserScrolled = InHandler;
 }
 
+void SScrollBar::SetOnScrollBarVisibilityChanged( const FOnScrollBarVisibilityChanged& InHandler )
+{
+	OnScrollBarVisibilityChanged = InHandler;
+}
+
 void SScrollBar::SetState( float InOffsetFraction, float InThumbSizeFraction, bool bCallOnUserScrolled )
 {
 	if ( Track->DistanceFromTop() != InOffsetFraction || Track->GetThumbSizeFraction() != InThumbSizeFraction )
@@ -316,18 +322,24 @@ const FSlateBrush* SScrollBar::GetDragThumbImage() const
 
 EVisibility SScrollBar::ShouldBeVisible() const
 {
+	const EVisibility CurrentVisibility = GetVisibility();
+    EVisibility NewVisibility = EVisibility::Collapsed;
+    
 	if ( this->HasMouseCapture() )
 	{
-		return EVisibility::Visible;
+		NewVisibility = EVisibility::Visible;
 	}
 	else if( Track->IsNeeded() )
 	{
-		return UserVisibility.Get();
+		NewVisibility = UserVisibility.Get();
 	}
-	else
+
+	if (NewVisibility != CurrentVisibility)
 	{
-		return EVisibility::Collapsed;
+		OnScrollBarVisibilityChanged.ExecuteIfBound(NewVisibility);
 	}
+	
+	return NewVisibility;
 }
 
 bool SScrollBar::IsScrolling() const
diff --git a/Engine/Source/Runtime/Slate/Private/Widgets/Layout/SScrollBox.cpp b/Engine/Source/Runtime/Slate/Private/Widgets/Layout/SScrollBox.cpp
--- a/Engine/Source/Runtime/Slate/Private/Widgets/Layout/SScrollBox.cpp
+++ b/Engine/Source/Runtime/Slate/Private/Widgets/Layout/SScrollBox.cpp
@@ -126,6 +126,7 @@ void SScrollBox::Construct( const FArguments& InArgs )
 	bShowSoftwareCursor = false;
 	SoftwareCursorPosition = FVector2f::ZeroVector;
 	OnUserScrolled = InArgs._OnUserScrolled;
+	OnScrollBarVisibilityChanged = InArgs._OnScrollBarVisibilityChanged;
 	Orientation = InArgs._Orientation;
 	bScrollToEnd = false;
 	bIsScrollingActiveTimerRegistered = false;
@@ -148,6 +149,7 @@ void SScrollBox::Construct( const FArguments& InArgs )
 		// An external scroll bar was specified by the user
 		ScrollBar = InArgs._ExternalScrollbar;
 		ScrollBar->SetOnUserScrolled(FOnUserScrolled::CreateSP(this, &SScrollBox::ScrollBar_OnUserScrolled));
+		ScrollBar->SetOnScrollBarVisibilityChanged(FOnScrollBarVisibilityChanged::CreateSP(this, &SScrollBox::ScrollBar_OnScrollBarVisibilityChanged));
 		bScrollBarIsExternal = true;
 	}
 	else
@@ -158,6 +160,7 @@ void SScrollBox::Construct( const FArguments& InArgs )
 		ScrollBar->SetThickness(InArgs._ScrollBarThickness);
 		ScrollBar->SetUserVisibility(InArgs._ScrollBarVisibility);
 		ScrollBar->SetScrollBarAlwaysVisible(InArgs._ScrollBarAlwaysVisible);
+		ScrollBar->SetOnScrollBarVisibilityChanged(FOnScrollBarVisibilityChanged::CreateSP(this, &SScrollBox::ScrollBar_OnScrollBarVisibilityChanged));
 		ScrollBarSlotPadding = InArgs._ScrollBarPadding;
 
 		bScrollBarIsExternal = false;
@@ -1138,6 +1141,11 @@ void SScrollBox::ScrollBar_OnUserScrolled( float InScrollOffsetFraction )
 	Invalidate(EInvalidateWidget::Layout);
 }
 
+void SScrollBox::ScrollBar_OnScrollBarVisibilityChanged( EVisibility NewVisibility )
+{
+	OnScrollBarVisibilityChanged.ExecuteIfBound(NewVisibility);
+}
+
 const float ShadowFadeDistance = 32.0f;
 
 FSlateColor SScrollBox::GetStartShadowOpacity() const
diff --git a/Engine/Source/Runtime/Slate/Public/Widgets/Layout/SScrollBar.h b/Engine/Source/Runtime/Slate/Public/Widgets/Layout/SScrollBar.h
--- a/Engine/Source/Runtime/Slate/Public/Widgets/Layout/SScrollBar.h
+++ b/Engine/Source/Runtime/Slate/Public/Widgets/Layout/SScrollBar.h
@@ -21,6 +21,7 @@ DECLARE_DELEGATE_OneParam(
 	FOnUserScrolled,
 	float );	/** ScrollOffset as a fraction between 0 and 1 */
 
+DECLARE_DELEGATE_OneParam(FOnScrollBarVisibilityChanged, EVisibility);	/** changed scroll bar visibility */
 
 class SScrollBarTrack;
 
@@ -48,6 +49,7 @@ public:
 		/** The style to use for this scrollbar */
 		SLATE_STYLE_ARGUMENT( FScrollBarStyle, Style )
 		SLATE_EVENT( FOnUserScrolled, OnUserScrolled )
+		SLATE_EVENT( FOnScrollBarVisibilityChanged, OnScrollBarVisibilityChanged )
 		SLATE_ARGUMENT( bool, AlwaysShowScrollbar )
 		SLATE_ARGUMENT( bool, AlwaysShowScrollbarTrack )
 		SLATE_ARGUMENT( bool, HideWhenNotInUse )
@@ -74,6 +76,13 @@ public:
 	 */
 	SLATE_API void SetOnUserScrolled( const FOnUserScrolled& InHandler );
 
+	/**
+	 * Set the handler to be invoked when scroll bar visibility changes.
+	 *
+	 * @param InHandler   Method to execute when scroll bar visibility changed
+	 */
+	SLATE_API void SetOnScrollBarVisibilityChanged( const FOnScrollBarVisibilityChanged& InHandler );
+
 	/**
 	 * Set the offset and size of the track's thumb.
 	 * Note that the maximum offset is 1.0-ThumbSizeFraction.
@@ -171,6 +180,7 @@ protected:
 	bool bDraggingThumb;
 	TSharedPtr<SScrollBarTrack> Track;
 	FOnUserScrolled OnUserScrolled;
+	FOnScrollBarVisibilityChanged OnScrollBarVisibilityChanged;
 	float DragGrabOffset;
 	EOrientation Orientation;
 	bool bAlwaysShowScrollbar;
diff --git a/Engine/Source/Runtime/Slate/Public/Widgets/Layout/SScrollBox.h b/Engine/Source/Runtime/Slate/Public/Widgets/Layout/SScrollBox.h
--- a/Engine/Source/Runtime/Slate/Public/Widgets/Layout/SScrollBox.h
+++ b/Engine/Source/Runtime/Slate/Public/Widgets/Layout/SScrollBox.h
@@ -192,6 +192,7 @@ public:
 		, _NavigationScrollPadding(0.0f)
 		, _ScrollWhenFocusChanges(EScrollWhenFocusChanges::NoScroll)
 		, _OnUserScrolled()
+	    , _OnScrollBarVisibilityChanged()
 		, _ConsumeMouseWheel(EConsumeMouseWheel::WhenScrollingPossible)
 		{
 			_Clipping = EWidgetClipping::ClipToBounds;
@@ -245,6 +246,9 @@ public:
 		/** Called when the button is clicked */
 		SLATE_EVENT(FOnUserScrolled, OnUserScrolled)
 
+		/** Fired when scroll bar visibility changed */
+		SLATE_EVENT(FOnScrollBarVisibilityChanged, OnScrollBarVisibilityChanged)
+
 		SLATE_ARGUMENT(EConsumeMouseWheel, ConsumeMouseWheel);
 
 	SLATE_END_ARGS()
@@ -405,6 +409,8 @@ private:
 	/** Invoked when the user scroll via the scrollbar */
 	void ScrollBar_OnUserScrolled( float InScrollOffsetFraction );
 
+	void ScrollBar_OnScrollBarVisibilityChanged( EVisibility NewVisibility);
+
 	/** Does the user need a hint that they can scroll to the start of the list? */
 	FSlateColor GetStartShadowOpacity() const;
 	
@@ -492,6 +498,9 @@ protected:
 
 	/** Fired when the user scrolls the scrollbox */
 	FOnUserScrolled OnUserScrolled;
+	
+	/** Fired when scroll bar visibility changed */
+	FOnScrollBarVisibilityChanged OnScrollBarVisibilityChanged;
 
 	/** The scrolling and stacking orientation. */
 	EOrientation Orientation;
diff --git a/Engine/Source/Runtime/UMG/Private/Components/ScrollBox.cpp b/Engine/Source/Runtime/UMG/Private/Components/ScrollBox.cpp
--- a/Engine/Source/Runtime/UMG/Private/Components/ScrollBox.cpp
+++ b/Engine/Source/Runtime/UMG/Private/Components/ScrollBox.cpp
@@ -106,7 +106,8 @@ PRAGMA_DISABLE_DEPRECATION_WARNINGS
 		.FrontPadScrolling(FrontPadScrolling)
 		.AnimateWheelScrolling(bAnimateWheelScrolling)
 		.WheelScrollMultiplier(WheelScrollMultiplier)
-		.OnUserScrolled(BIND_UOBJECT_DELEGATE(FOnUserScrolled, SlateHandleUserScrolled));
+		.OnUserScrolled(BIND_UOBJECT_DELEGATE(FOnUserScrolled, SlateHandleUserScrolled))
+		.OnScrollBarVisibilityChanged(BIND_UOBJECT_DELEGATE(FOnScrollBarVisibilityChanged, SlateHandleScrollBarVisibilityChanged));
 PRAGMA_ENABLE_DEPRECATION_WARNINGS
 	for ( UPanelSlot* PanelSlot : Slots )
 	{
@@ -534,6 +535,11 @@ void UScrollBox::SlateHandleUserScrolled(float CurrentOffset)
 	OnUserScrolled.Broadcast(CurrentOffset);
 }
 
+void UScrollBox::SlateHandleScrollBarVisibilityChanged(EVisibility NewVisibility)
+{
+	OnScrollBarVisibilityChanged.Broadcast(ConvertRuntimeToSerializedVisibility(NewVisibility));
+}
+
 #if WITH_EDITOR
 
 const FText UScrollBox::GetPaletteCategory()
diff --git a/Engine/Source/Runtime/UMG/Public/Components/ScrollBox.h b/Engine/Source/Runtime/UMG/Public/Components/ScrollBox.h
--- a/Engine/Source/Runtime/UMG/Public/Components/ScrollBox.h
+++ b/Engine/Source/Runtime/UMG/Public/Components/ScrollBox.h
@@ -12,6 +12,7 @@
 #include "ScrollBox.generated.h"
 
 DECLARE_DYNAMIC_MULTICAST_DELEGATE_OneParam(FOnUserScrolledEvent, float, CurrentOffset);
+DECLARE_DYNAMIC_MULTICAST_DELEGATE_OneParam(FOnScrollBarVisibilityChangedEvent, ESlateVisibility , NewVisibility);
 
 /**
  * An arbitrary scrollable collection of widgets.  Great for presenting 10-100 widgets in a list.  Doesn't support virtualization.
@@ -202,6 +203,10 @@ public:
 	/** Called when the scroll has changed */
 	UPROPERTY(BlueprintAssignable, Category = "Button|Event")
 	FOnUserScrolledEvent OnUserScrolled;
+	
+	/** Called when the scrollbar visibility has changed */
+	UPROPERTY(BlueprintAssignable, Category = "Button|Event")
+	FOnScrollBarVisibilityChangedEvent OnScrollBarVisibilityChanged;
 
 	/**
 	 * Updates the scroll offset of the scrollbox.
@@ -271,6 +276,8 @@ protected:
 
 	UMG_API void SlateHandleUserScrolled(float CurrentOffset);
 
+	UMG_API void SlateHandleScrollBarVisibilityChanged(EVisibility NewVisibility);
+
 	// Initialize IsFocusable in the constructor before the SWidget is constructed.
 	UMG_API void InitBackPadScrolling(bool InBackPadScrolling);
 	// Initialize IsFocusable in the constructor before the SWidget is constructed.
