Remove unnecessary assertion - sometimes menu items are added with unbound executable actions as grey placeholders if those features are not available.

diff --git a/Engine/Source/Runtime/Slate/Public/Framework/MultiBox/MultiBox.h b/Engine/Source/Runtime/Slate/Public/Framework/MultiBox/MultiBox.h
--- a/Engine/Source/Runtime/Slate/Public/Framework/MultiBox/MultiBox.h
+++ b/Engine/Source/Runtime/Slate/Public/Framework/MultiBox/MultiBox.h
@@ -129,7 +129,6 @@ public:
 		, bSearchable(true)
 		, bIsPartOfHeading(InMultiBlockParams.bIsPartOfHeading)
 	{
-		check(Action.IsValid() || DirectActions.IsBound());
 	}
 
 	virtual ~FMultiBlock()
