Move macOS main menu into UE window by default

* On the Mac, this change moves the Unreal Engine editor's main menu from the macOS system main menu into the editor window.
* Remove the CVAR Slate.MacMainMenuInsideUnrealWindow that previously turned this main menu move on or off on the Mac.
* Remove FGenericPlatformMisc::CanShowMenusInWindows because menus in the UE window is default on all platforms now. Note that this function was never released so we don't need to deprecate it (added after 5.3).

diff --git a/Engine/Source/Runtime/Core/Private/Mac/MacPlatformMisc.cpp b/Engine/Source/Runtime/Core/Private/Mac/MacPlatformMisc.cpp
--- a/Engine/Source/Runtime/Core/Private/Mac/MacPlatformMisc.cpp
+++ b/Engine/Source/Runtime/Core/Private/Mac/MacPlatformMisc.cpp
@@ -77,14 +77,6 @@ static TAutoConsoleVariable<int32> CVarMacPlatformDumpAllThreadsOnHang(
 	TEXT("Mac.DumpAllThreadsOnHang"),
 	1,
 	TEXT("If > 0, then when reporting a hang generate a backtrace for all threads."));
-static bool GMacMainMenuInsideUnrealWindow = false;
-static FAutoConsoleVariableRef CVarMacMainMenuInsideUnrealWindow(
-	TEXT("Slate.MacMainMenuInsideUnrealWindow"),
-	GMacMainMenuInsideUnrealWindow,
-	TEXT("When enabled, will draw the main menu inside the Unreal windows on the Mac.\n")
-	TEXT("When disabled (default, old behavior) use the macOS system menu for all main menus."),
-	ECVF_Default
-);
 
 /*------------------------------------------------------------------------------
  Platform property discovery.
@@ -2671,11 +2663,6 @@ void ReportHang(const TCHAR* ErrorMessage, const uint64* StackFrames, int32 NumS
 	EnsureLock.Unlock();
 }
 
-bool FMacPlatformMisc::CanShowMenusInWindows()
-{
-	return GMacMainMenuInsideUnrealWindow;
-}
-
 typedef NSArray* (*MTLCopyAllDevices)(void);
 
 bool FMacPlatformMisc::HasPlatformFeature(const TCHAR* FeatureName)
diff --git a/Engine/Source/Runtime/Core/Public/GenericPlatform/GenericPlatformMisc.h b/Engine/Source/Runtime/Core/Public/GenericPlatform/GenericPlatformMisc.h
--- a/Engine/Source/Runtime/Core/Public/GenericPlatform/GenericPlatformMisc.h
+++ b/Engine/Source/Runtime/Core/Public/GenericPlatform/GenericPlatformMisc.h
@@ -1541,14 +1541,6 @@ public:
 		return false;
 	}
 
-	/*
-	 * Returns whether the platform can show main menus inside window title bars.
-	 */
-	static CORE_API bool CanShowMenusInWindows()
-	{
-		return true;
-	}
-
 	/*
 	 * Returns whether the volume buttons are handled by the system
 	 */
diff --git a/Engine/Source/Runtime/Core/Public/Mac/MacPlatformMisc.h b/Engine/Source/Runtime/Core/Public/Mac/MacPlatformMisc.h
--- a/Engine/Source/Runtime/Core/Public/Mac/MacPlatformMisc.h
+++ b/Engine/Source/Runtime/Core/Public/Mac/MacPlatformMisc.h
@@ -43,7 +43,6 @@ struct CORE_API FMacPlatformMisc : public FApplePlatformMisc
 	static void GetOSVersions( FString& out_OSVersionLabel, FString& out_OSSubVersionLabel );
 	static FString GetOSVersion();
 	static NSOperatingSystemVersion GetNSOperatingSystemVersion();
-	static bool CanShowMenusInWindows();
 	static bool HasPlatformFeature(const TCHAR* FeatureName);
 	static bool GetDiskTotalAndFreeSpace(const FString& InPath, uint64& TotalNumberOfBytes, uint64& NumberOfFreeBytes);
 	static bool HasSeparateChannelForDebugOutput();
diff --git a/Engine/Source/Runtime/Slate/Private/Framework/Docking/SDockingTabStack.cpp b/Engine/Source/Runtime/Slate/Private/Framework/Docking/SDockingTabStack.cpp
--- a/Engine/Source/Runtime/Slate/Private/Framework/Docking/SDockingTabStack.cpp
+++ b/Engine/Source/Runtime/Slate/Private/Framework/Docking/SDockingTabStack.cpp
@@ -3,7 +3,6 @@
 #include "Framework/Docking/SDockingTabStack.h"
 #include "Framework/Commands/UIAction.h"
 #include "Framework/Commands/UICommandList.h"
-#include "HAL/PlatformMisc.h"
 #include "Widgets/Text/STextBlock.h"
 #include "Framework/MultiBox/MultiBoxBuilder.h"
 #include "Layout/WidgetPath.h"
@@ -871,7 +870,7 @@ void SDockingTabStack::ReserveSpaceForWindowChrome(EChromeElement Element, bool
 		FMargin ControlsPadding;
 		static const FMargin IconPadding = FMargin(0);
 
-		if (FPlatformMisc::CanShowMenusInWindows() && bIncludePaddingForMenuBar)
+		if (bIncludePaddingForMenuBar)
 		{
 			static const float TopPaddingForTrafficLightsAndMenuBar = 30.0f;
 			// Always add padding on top, because on the Mac there is always either a main menu bar or the "traffic light" buttons (close, minimize, and maximize) above controls.
diff --git a/Engine/Source/Runtime/Slate/Private/Framework/Docking/TabManager.cpp b/Engine/Source/Runtime/Slate/Private/Framework/Docking/TabManager.cpp
--- a/Engine/Source/Runtime/Slate/Private/Framework/Docking/TabManager.cpp
+++ b/Engine/Source/Runtime/Slate/Private/Framework/Docking/TabManager.cpp
@@ -21,7 +21,6 @@
 #include "Misc/NamePermissionList.h"
 #include "Trace/SlateMemoryTags.h"
 #include "HAL/PlatformApplicationMisc.h"
-#include "HAL/PlatformMisc.h"
 #if PLATFORM_MAC
 #include "Framework/MultiBox/Mac/MacMenu.h"
 #endif
@@ -318,7 +317,6 @@ TSharedRef<FTabManager::FLayoutNode> FTabManager::FLayout::NewFromString_Helper(
 	}	
 }
 
-
 TSharedPtr<FTabManager::FLayout> FTabManager::FLayout::NewFromString( const FString& LayoutAsText )
 {
 	TSharedPtr<FJsonObject> JsonObject;
@@ -791,46 +789,20 @@ void FTabManager::UpdateMainMenu(TSharedPtr<SDockTab> ForTab, const bool bForce)
 		ParentWindowOfOwningTab = MainNonCloseableTabPinned->GetParentWindow();
 	}
 
-#if PLATFORM_MAC
-	if (!FPlatformMisc::CanShowMenusInWindows())
+	if (bAllowPerWindowMenu)
 	{
-		if (MenuMultiBox.IsValid())
-		{
-			bool bUpdate = bForce;
-			// On OS X opening the tab will set the multi-box and take key focus, but not seemingly send a keyboard focus event into Slate.
-			if (ParentWindowOfOwningTab.IsValid())
-			{
-				bUpdate |= ParentWindowOfOwningTab->GetNativeWindow()->IsForegroundWindow();
-			}
-
-			if (bUpdate)
-			{
-				FSlateMacMenu::UpdateWithMultiBox(MenuMultiBox.ToSharedRef());
-			}
-		}
-		else
+		if (ParentWindowOfOwningTab)
 		{
-			FSlateMacMenu::UpdateWithMultiBox(nullptr);
+			ParentWindowOfOwningTab->GetTitleBar()->UpdateWindowMenu(MenuWidget);
 		}
 	}
 	else
-#endif // PLATFORM_MAC
 	{
-		if (bAllowPerWindowMenu)
+		MenuMultiBox.Reset();
+		MenuWidget.Reset();
+		if (ParentWindowOfOwningTab)
 		{
-			if (ParentWindowOfOwningTab)
-			{
-				ParentWindowOfOwningTab->GetTitleBar()->UpdateWindowMenu(MenuWidget);
-			}
-		}
-		else
-		{
-			MenuMultiBox.Reset();
-			MenuWidget.Reset();
-			if (ParentWindowOfOwningTab)
-			{
-				ParentWindowOfOwningTab->GetTitleBar()->UpdateWindowMenu(nullptr);
-			}
+			ParentWindowOfOwningTab->GetTitleBar()->UpdateWindowMenu(nullptr);
 		}
 	}
 }
@@ -1460,7 +1432,7 @@ TSharedPtr<SDockTab> FTabManager::InvokeTab_Internal(const FTabId& TabId, bool b
 
 	// Tab is not live. Figure out where to spawn it.
 	TSharedPtr<SDockingTabStack> StackToSpawnIn = bForceOpenWindowIfNeeded ? AttemptToOpenTab( TabId, true ) : FindPotentiallyClosedTab( TabId );
-	
+
 	if (StackToSpawnIn.IsValid())
 	{
 		const TSharedPtr<SDockTab> NewTab = SpawnTab(TabId, TSharedPtr<SWindow>());
diff --git a/Engine/Source/Runtime/Slate/Private/Framework/MultiBox/MultiBox.cpp b/Engine/Source/Runtime/Slate/Private/Framework/MultiBox/MultiBox.cpp
--- a/Engine/Source/Runtime/Slate/Private/Framework/MultiBox/MultiBox.cpp
+++ b/Engine/Source/Runtime/Slate/Private/Framework/MultiBox/MultiBox.cpp
@@ -3,7 +3,6 @@
 #include "Framework/MultiBox/MultiBox.h"
 
 #include "SClippingVerticalBox.h"
-#include "HAL/PlatformMisc.h"
 #include "Misc/ConfigCacheIni.h"
 #include "Widgets/SBoxPanel.h"
 #include "Widgets/SOverlay.h"
@@ -449,14 +448,6 @@ TSharedRef< SMultiBoxWidget > FMultiBox::MakeWidget( bool bSearchable, FOnMakeMu
 		// Build up the widget
 		NewMultiBoxWidget->BuildMultiBoxWidget();
 	}
-
-#if PLATFORM_MAC
-	// Hide the main menu on the Mac if we are using the macOS system menu instead.
-	if (Type == EMultiBoxType::MenuBar && !FPlatformMisc::CanShowMenusInWindows())
-	{
-		NewMultiBoxWidget->SetVisibility(EVisibility::Collapsed);
-	}
-#endif
 	
 	return NewMultiBoxWidget;
 }
diff --git a/Engine/Source/Runtime/Slate/Public/Framework/Application/SWindowTitleBar.h b/Engine/Source/Runtime/Slate/Public/Framework/Application/SWindowTitleBar.h
--- a/Engine/Source/Runtime/Slate/Public/Framework/Application/SWindowTitleBar.h
+++ b/Engine/Source/Runtime/Slate/Public/Framework/Application/SWindowTitleBar.h
@@ -3,7 +3,6 @@
 #pragma once
 
 #include "CoreMinimal.h"
-#include "HAL/PlatformMisc.h"
 #include "Misc/Attribute.h"
 #include "Layout/Visibility.h"
 #include "Styling/SlateColor.h"
@@ -327,7 +326,7 @@ protected:
 		}
 #endif //!PLATFORM_MAC
 
-		if (ShowAppIcon && bHasWindowButtons && FPlatformMisc::CanShowMenusInWindows())
+		if (ShowAppIcon && bHasWindowButtons)
 		{
 			OutLeftContent = 
 				SNew(SHorizontalBox)
