Expose FSlateApplication::IsUsingHighPrecisionMouseMovment  with the SLATE_API

#rnx
#pf 652d63b7bd575586ded4164a

diff --git a/Engine/Source/Runtime/Slate/Public/Framework/Application/SlateApplication.h b/Engine/Source/Runtime/Slate/Public/Framework/Application/SlateApplication.h
--- a/Engine/Source/Runtime/Slate/Public/Framework/Application/SlateApplication.h
+++ b/Engine/Source/Runtime/Slate/Public/Framework/Application/SlateApplication.h
@@ -752,7 +752,7 @@ public:
 	 * @return True if the application is currently routing high precision mouse movement events (OS specific)
 	 * If this value is true, the mouse is captured and hidden by the widget that originally made the request.
 	 */
-	bool IsUsingHighPrecisionMouseMovment() const { return PlatformApplication.IsValid() ? PlatformApplication->IsUsingHighPrecisionMouseMode() : false; }
+	SLATE_API bool IsUsingHighPrecisionMouseMovment() const { return PlatformApplication.IsValid() ? PlatformApplication->IsUsingHighPrecisionMouseMode() : false; }
 	
 	/**
 	 * @return True if the last mouse event was from a trackpad.
