Property Editor: Follow up to: Remove 0-360 restriction on rotation dragging. Adds the linear delta sensitivity to pass through from RotatorInputBox to NumericInputBox
The linear delta sensitivity defaults to 1, just as the vector input box does.

diff --git a/Engine/Source/Runtime/Slate/Public/Widgets/Input/SRotatorInputBox.h b/Engine/Source/Runtime/Slate/Public/Widgets/Input/SRotatorInputBox.h
--- a/Engine/Source/Runtime/Slate/Public/Widgets/Input/SRotatorInputBox.h
+++ b/Engine/Source/Runtime/Slate/Public/Widgets/Input/SRotatorInputBox.h
@@ -33,6 +33,7 @@ public:
 		, _Font(FAppStyle::Get().GetFontStyle("NormalFont"))
 		, _AllowSpin(true)
 		, _SpinDelta(1.f)
+		, _LinearDeltaSensitivity(1)
 		, _MinSliderValue(TOptional<NumericType>())
 		, _MaxSliderValue(TOptional<NumericType>())
 		, _DisplayToggle( false )
@@ -69,6 +70,9 @@ public:
 		/** The delta amount to apply, per pixel, when the spinner is dragged. */
 		SLATE_ATTRIBUTE(NumericType, SpinDelta)
 
+		/** If we're an unbounded spinbox, what value do we divide mouse movement by before multiplying by Delta. Requires Delta to be set. */
+		SLATE_ATTRIBUTE(int32, LinearDeltaSensitivity)
+
 		/** The minimum value that can be specified by using the slider */
 		SLATE_ATTRIBUTE( TOptional<NumericType>, MinSliderValue )
 
@@ -161,6 +165,7 @@ public:
 				SNew(SNumericEntryBox<NumericType>)
 				.AllowSpin(InArgs._AllowSpin)
 				.Delta(InArgs._SpinDelta)
+				.LinearDeltaSensitivity(InArgs._LinearDeltaSensitivity)
 				.MinValue(InArgs._MinSliderValue)
 				.MaxValue(InArgs._MaxSliderValue)
 				.MinSliderValue(InArgs._MinSliderValue)
@@ -196,6 +201,7 @@ public:
 				SNew(SNumericEntryBox<NumericType>)
 				.AllowSpin(InArgs._AllowSpin)
 				.Delta(InArgs._SpinDelta)
+				.LinearDeltaSensitivity(InArgs._LinearDeltaSensitivity)
 				.MinValue(InArgs._MinSliderValue)
 				.MaxValue(InArgs._MaxSliderValue)
 				.MinSliderValue(InArgs._MinSliderValue)
@@ -231,6 +237,7 @@ public:
 				SNew(SNumericEntryBox<NumericType>)
 				.AllowSpin(InArgs._AllowSpin)
 				.Delta(InArgs._SpinDelta)
+				.LinearDeltaSensitivity(InArgs._LinearDeltaSensitivity)
 				.MinValue(InArgs._MinSliderValue)
 				.MaxValue(InArgs._MaxSliderValue)
 				.MinSliderValue(InArgs._MinSliderValue)
