PR #10952: Use ProjectDisplayedTitle as AppName on Mac if set

diff --git a/Engine/Source/Runtime/Slate/Private/Framework/MultiBox/Mac/MacMenu.cpp b/Engine/Source/Runtime/Slate/Private/Framework/MultiBox/Mac/MacMenu.cpp
--- a/Engine/Source/Runtime/Slate/Private/Framework/MultiBox/Mac/MacMenu.cpp
+++ b/Engine/Source/Runtime/Slate/Private/Framework/MultiBox/Mac/MacMenu.cpp
@@ -280,7 +280,19 @@ namespace MacMenuHelper
 	
 	NSString* ComputeAppName()
 	{
-		return GIsEditor ? NSLOCTEXT("UnrealEditor", "ApplicationTitle", "Unreal Editor").ToString().GetNSString() : FString(FApp::GetProjectName()).GetNSString();
+        if (GIsEditor)
+        {
+            return NSLOCTEXT("UnrealEditor", "ApplicationTitle", "Unreal Editor").ToString().GetNSString();
+        }
+        
+        FText ProjectTitle;
+        GConfig->GetText(TEXT("/Script/EngineSettings.GeneralProjectSettings"), TEXT("ProjectDisplayedTitle"), ProjectTitle, GGameIni);
+        if (!ProjectTitle.IsEmpty())
+        {
+            return ProjectTitle.ToString().GetNSString();
+        }
+        
+        return FString(FApp::GetProjectName()).GetNSString();
 	}
 	
 	bool GMacPostInitStartupRequested = false;
