PR #10952: fix CrashReportClient -nopch compile error

diff --git a/Engine/Source/Runtime/Slate/Private/Framework/MultiBox/Mac/MacMenu.cpp b/Engine/Source/Runtime/Slate/Private/Framework/MultiBox/Mac/MacMenu.cpp
--- a/Engine/Source/Runtime/Slate/Private/Framework/MultiBox/Mac/MacMenu.cpp
+++ b/Engine/Source/Runtime/Slate/Private/Framework/MultiBox/Mac/MacMenu.cpp
@@ -4,6 +4,7 @@
 #include "Framework/Application/SlateApplication.h"
 #include "Mac/CocoaThread.h"
 #include "Mac/MacApplication.h"
+#include "Misc/ConfigCacheIni.h"
 #include "Misc/ScopeLock.h"
 #include "HAL/PlatformApplicationMisc.h"
 #include "Framework/Commands/Commands.h"
