Fix flickering menus by recreating main menu widgets for each window

* Add FTabManager::SetGeneratePerWindowMainMenuDelegate that allows users to set a delegate the tab manager calls when it needs a main menu widget for a tab's window. This way we can avoid drawing the same main menu widget in multiple windows and the flickering that can cause.
* Generated main menu widgets are cached and reused to avoid regenerating widgets for the same tab manager and window combination. FSlateApplication::OnWindowBeingDestroyed is used to clean this cache up as windows are destroyed.
* The old method FTabManager::SetMenuMultiBox to set a main menu widget for a tab manager now uses FTabManager::SetGeneratePerWindowMainMenuDelegate behind the scenes.
* Update FMainMenu::MakeMainMenu to use the new FTabManager::SetGeneratePerWindowMainMenuDelegate to avoid the main menu flickering described in this ticket.
* Remove unused code.

diff --git a/Engine/Source/Editor/MainFrame/Private/Menus/MainMenu.cpp b/Engine/Source/Editor/MainFrame/Private/Menus/MainMenu.cpp
--- a/Engine/Source/Editor/MainFrame/Private/Menus/MainMenu.cpp
+++ b/Engine/Source/Editor/MainFrame/Private/Menus/MainMenu.cpp
@@ -436,18 +436,15 @@ void FMainMenu::MakeMainMenu(const TSharedPtr<FTabManager>& TabManager, const FN
 
 	ToolMenuContext.AppendCommandList(FMainFrameCommands::ActionList);
 
-	USlateTabManagerContext* ContextObject = NewObject<USlateTabManagerContext>();
-	ContextObject->TabManager = TabManager;
-	ToolMenuContext.AddObject(ContextObject);
+	TabManager->SetGeneratePerWindowMainMenuDelegate(FGenerateMenu::CreateLambda([MenuName,ToolMenuContext](TSharedPtr<FTabManager> TabManager) {
+		USlateTabManagerContext* ContextObject = NewObject<USlateTabManagerContext>();
+		ContextObject->TabManager = TabManager;
 
-	// Create the menu bar!
-	TSharedRef<SWidget> MenuBarWidget = UToolMenus::Get()->GenerateWidget(MenuName, ToolMenuContext);
-	if (MenuBarWidget != SNullWidget::NullWidget)
-	{
-		// Tell tab-manager about the multi-box for platforms with a global menu bar
-		TSharedRef<SMultiBoxWidget> MultiBoxWidget = StaticCastSharedRef<SMultiBoxWidget>(MenuBarWidget);
-		TabManager->SetMenuMultiBox(ConstCastSharedRef<FMultiBox>(MultiBoxWidget->GetMultiBox()), MultiBoxWidget);
-	}
+		FToolMenuContext ContextCopy = ToolMenuContext;
+		ContextCopy.AddObject(ContextObject);
+
+		return UToolMenus::Get()->GenerateWidget(MenuName, ContextCopy);
+	}));
 }
 
 void FMainMenu::RegisterMainMenu()
diff --git a/Engine/Source/Runtime/Slate/Private/Framework/Docking/TabManager.cpp b/Engine/Source/Runtime/Slate/Private/Framework/Docking/TabManager.cpp
--- a/Engine/Source/Runtime/Slate/Private/Framework/Docking/TabManager.cpp
+++ b/Engine/Source/Runtime/Slate/Private/Framework/Docking/TabManager.cpp
@@ -15,6 +15,7 @@
 #include "Framework/Docking/SDockingSplitter.h"
 #include "Framework/Docking/SDockingArea.h"
 #include "Widgets/Docking/SDockTab.h"
+#include "Widgets/SNullWidget.h"
 #include "Framework/Docking/SDockingTabStack.h"
 #include "Framework/Docking/SDockingTabWell.h"
 #include "Framework/Docking/LayoutExtender.h"
@@ -762,19 +763,33 @@ void FTabManager::SetAllowWindowMenuBar(bool bInAllowWindowMenuBar)
 	bAllowPerWindowMenu = bInAllowWindowMenuBar;
 }
 
-void FTabManager::SetMenuMultiBox(const TSharedPtr<FMultiBox> NewMenuMutliBox, const TSharedPtr<SWidget> NewMenuWidget)
+void FTabManager::SetGeneratePerWindowMainMenuDelegate(const FGenerateMenu& InGeneratePerWindowMainMenu)
 {
-	// We only use the platform native global menu bar on Mac
-	MenuMultiBox = NewMenuMutliBox;
-	MenuWidget = NewMenuWidget;
+	GeneratePerWindowMainMenu = InGeneratePerWindowMainMenu;
 
 	UpdateMainMenu(OwnerTabPtr.Pin(), false);
 }
 
-void FTabManager::UpdateMainMenu(TSharedPtr<SDockTab> ForTab, const bool bForce)
+void FTabManager::SetMenuMultiBox(const TSharedPtr<FMultiBox>, const TSharedPtr<SWidget> MenuWidget)
 {
-	bool bIsMajorTab = true;
+	// Set this delegate for backwards compatibility.
+	GeneratePerWindowMainMenu = FGenerateMenu::CreateLambda([MenuWidget](TSharedPtr<FTabManager>)
+	{
+		if (MenuWidget)
+		{
+			return MenuWidget.ToSharedRef();
+		}
+		else
+		{
+			return SNullWidget::NullWidget;
+		}
+	});
 
+	UpdateMainMenu(OwnerTabPtr.Pin(), false);
+}
+
+void FTabManager::UpdateMainMenu(TSharedPtr<SDockTab> ForTab, const bool bForce)
+{
 	TSharedPtr<SWindow> ParentWindowOfOwningTab;
 	if (ForTab && (ForTab->GetTabRole() == ETabRole::MajorTab || ForTab->GetVisualTabRole() == ETabRole::MajorTab))
 	{
@@ -789,22 +804,57 @@ void FTabManager::UpdateMainMenu(TSharedPtr<SDockTab> ForTab, const bool bForce)
 		ParentWindowOfOwningTab = MainNonCloseableTabPinned->GetParentWindow();
 	}
 
-	if (bAllowPerWindowMenu)
+	TSharedPtr<SWidget> MainMenuWidget;
+
+	if (GeneratePerWindowMainMenu.IsBound())
 	{
+		// Try to retrieve a previously generated main menu widget.
 		if (ParentWindowOfOwningTab)
 		{
-			ParentWindowOfOwningTab->GetTitleBar()->UpdateWindowMenu(MenuWidget);
+			SWindow* const WindowPtr = ParentWindowOfOwningTab.Get();
+			
+			if (WindowToGeneratedMainMenuWidgetsMap.Contains(WindowPtr))
+			{
+				TWeakPtr<SWidget> WeakMainMenuWidget = WindowToGeneratedMainMenuWidgetsMap[WindowPtr];
+				MainMenuWidget = WeakMainMenuWidget.Pin();
+			}
+		}
+
+		// Generate a main menu widget if we still don't have a valid one.
+		if (!MainMenuWidget)
+		{
+			TSharedRef<SWidget> MenuBarWidget = GeneratePerWindowMainMenu.Execute(AsShared());
+			MainMenuWidget = MenuBarWidget.ToSharedPtr();
+
+			// Save the newly generated main menu widget for next time.
+			if (ParentWindowOfOwningTab)
+			{
+				SWindow* const WindowPtr = ParentWindowOfOwningTab.Get();
+				WindowToGeneratedMainMenuWidgetsMap.Add(WindowPtr, MainMenuWidget.ToWeakPtr());
+			}
 		}
 	}
-	else
+
+	// Remove all Window-Widget pairs where the main menu widget is no longer valid.
+	for (auto It = WindowToGeneratedMainMenuWidgetsMap.CreateIterator(); It; ++It)
 	{
-		MenuMultiBox.Reset();
-		MenuWidget.Reset();
-		if (ParentWindowOfOwningTab)
+		if (!It.Value().IsValid())
 		{
-			ParentWindowOfOwningTab->GetTitleBar()->UpdateWindowMenu(nullptr);
+			It.RemoveCurrent();
 		}
 	}
+
+	if (ParentWindowOfOwningTab)
+	{
+		if (bAllowPerWindowMenu && MainMenuWidget)
+		{
+			ParentWindowOfOwningTab->GetTitleBar()->UpdateWindowMenu(MainMenuWidget);
+		}
+		else
+		{
+			ParentWindowOfOwningTab->GetTitleBar()->UpdateWindowMenu(nullptr);
+		}
+    }
 }
 
 void FTabManager::SetMainTab(const FTabId& InMainTabID)
@@ -1534,6 +1584,11 @@ FUIAction FTabManager::GetUIActionForTabSpawnerMenuEntry(TSharedPtr<FTabSpawnerE
 		);
 }
 
+void FTabManager::OnWindowBeingDestroyed(const SWindow& WindowBeingDestoyed)
+{
+	WindowToGeneratedMainMenuWidgetsMap.Remove(&WindowBeingDestoyed);
+}
+
 void FTabManager::InvokeTabForMenu( FName TabId )
 {
 	TryInvokeTab(TabId);
@@ -1597,6 +1652,8 @@ FTabManager::FTabManager( const TSharedPtr<SDockTab>& InOwnerTab, const TSharedR
 , TabPermissionList( MakeShareable(new FNamePermissionList()) )
 {
 	LocalWorkspaceMenuRoot = FWorkspaceItem::NewGroup(LOCTEXT("LocalWorkspaceRoot", "Local Workspace Root"));
+
+	FSlateApplication::Get().OnWindowBeingDestroyed().AddRaw(this, &FTabManager::OnWindowBeingDestroyed);
 }
 
 TSharedPtr<SDockingArea> FTabManager::RestoreArea(const TSharedRef<FArea>& AreaToRestore, const TSharedPtr<SWindow>& InParentWindow, const bool bEmbedTitleAreaContent, const EOutputCanBeNullptr OutputCanBeNullptr, bool bForceOpenWindowIfNeeded)
@@ -2138,6 +2195,11 @@ TSharedPtr<SDockTab> FTabManager::FindExistingLiveTab( const FTabId& TabId ) con
 
 FTabManager::~FTabManager()
 {
+	if (FSlateApplication::IsInitialized())
+	{
+		FSlateApplication::Get().OnWindowBeingDestroyed().RemoveAll(this);
+	}
+
 	ClearPendingLayoutSave();
 }
 
diff --git a/Engine/Source/Runtime/Slate/Public/Framework/Docking/TabManager.h b/Engine/Source/Runtime/Slate/Public/Framework/Docking/TabManager.h
--- a/Engine/Source/Runtime/Slate/Public/Framework/Docking/TabManager.h
+++ b/Engine/Source/Runtime/Slate/Public/Framework/Docking/TabManager.h
@@ -21,6 +21,7 @@ class SDockingTabStack;
 class FLayoutExtender;
 struct FTabMatcher;
 struct FSidebarTabLists;
+class FTabManager;
 
 DECLARE_MULTICAST_DELEGATE_TwoParams(
 	FOnActiveTabChanged,
@@ -28,7 +29,12 @@ DECLARE_MULTICAST_DELEGATE_TwoParams(
 	TSharedPtr<SDockTab>,
 	/** Previously active tab */
 	TSharedPtr<SDockTab> );
-	
+
+DECLARE_DELEGATE_RetVal_OneParam(
+	TSharedRef<SWidget>,
+	FGenerateMenu,
+	TSharedPtr<FTabManager>
+);
 
 enum class ESidebarLocation : uint8
 {
@@ -961,12 +967,18 @@ class FTabManager : public TSharedFromThis<FTabManager>
 		/** Whether or not this tab manager supports a custom menu bar for the active major tab that will be shown on top of the major tab area in the window this tab manager resides in. */
 		bool AllowsWindowMenuBar() const { return bAllowPerWindowMenu; }
 
+		/**
+		 * Set delegate for generating a new main menu widget.
+		 * @param GenerateMenu This delegate will be called every time this tab manager needs to populate a new window with a main menu widget. This delegate is expected to return a new unique widget each time it's called.
+		 */
+		SLATE_API void SetGeneratePerWindowMainMenuDelegate(const FGenerateMenu& GenerateMenu);
+
 		/**
 		 * Set the multi-box to use for generating a global menu bar.  The implementation is platform and setting specific
-		 * On Mac the menu bar appears globally at the top of the desktop in all cases regardless of whether or not SetAllowWindowMenuBar is called.  On other desktop platforms the menu appears at the top of the window this tab manager is a part of only if SetAllowWindowMenuBar(true) is called.
-		 * @param NewMenuMutliBox The multi-box to generate the global menu bar from.
+		 * @param NewMenuMutliBox This parameter is unused.
+		 * @param MenuWidget This widget will be used for all windows that tabs of this tab manager appear in. If you see flickering issues because the same widget is drawn in multiple windows, use SetGeneratePerWindowMainMenuDelegate instead.
 		 */
-		SLATE_API void SetMenuMultiBox(const TSharedPtr<FMultiBox> NewMenuMutliBox, const TSharedPtr<SWidget> MenuWidget);
+		SLATE_API void SetMenuMultiBox(const TSharedPtr<FMultiBox> Unused, const TSharedPtr<SWidget> MenuWidget);
 
 		/**
 		 * Update the native, global menu bar if it is being used.
@@ -1041,9 +1053,9 @@ class FTabManager : public TSharedFromThis<FTabManager>
 		SLATE_API const TSharedPtr<const FTabSpawnerEntry> FindTabSpawnerFor(FName TabId) const;
 
 	protected:
-		SLATE_API void InvokeTabForMenu( FName TabId );
+		SLATE_API void OnWindowBeingDestroyed(const SWindow& WindowBeingDestoyed);
 
-	protected:
+		SLATE_API void InvokeTabForMenu( FName TabId );
 		
 		SLATE_API void InsertDocumentTab( FName PlaceholderId, const FSearchPreference& SearchPreference, const TSharedRef<SDockTab>& UnmanagedTab, bool bPlaySpawnAnim );
 		SLATE_API void InsertDocumentTab( FName PlaceholderId, FName NewTabId, const FSearchPreference& SearchPreference, const TSharedRef<SDockTab>& UnmanagedTab, bool bPlaySpawnAnim );
@@ -1122,7 +1134,7 @@ class FTabManager : public TSharedFromThis<FTabManager>
 		 */
 	    SLATE_API TSharedPtr<FArea> GetFAreaFromInitialLayoutWithTabType(const FTabId& InTabIdToMatch ) const;
 
-protected:
+	protected:
 		SLATE_API bool HasValidTabs( const TSharedRef<FTabManager::FLayoutNode>& SomeNode ) const;
 
 		/**
@@ -1203,9 +1215,11 @@ protected:
 		/** A Major tab that contains this TabManager's widgets. */
 		TWeakPtr<SDockTab> OwnerTabPtr;
 
-		/** The current menu multi-box for the tab, used to construct platform native main menus */
-		TSharedPtr<FMultiBox> MenuMultiBox;
-		TSharedPtr<SWidget> MenuWidget;
+		/** Delegate that can generate a widget for this tab manager's main menu. */
+		FGenerateMenu GeneratePerWindowMainMenu;
+
+		/** Map from SWindow to generated main-menu widget. This helps us generate a main menu widget only once per window. */
+		TMap<SWindow*, TWeakPtr<SWidget>> WindowToGeneratedMainMenuWidgetsMap;
 
 		/** Protected private API that must only be accessed by the docking framework internals */
 		TSharedRef<FPrivateApi> PrivateApi;
