Fix automation issue in TabManager caused by assuming FSlateApplication is always available

* Now guarding on FSlateApplication::IsInitialized before adding a delegate to FSlateApplication::OnWindowBeingDestroyed in FTabManager's constructor. Without this guard, crashes can happen in UAT workloads where FSlateApplication has not been initialized.

diff --git a/Engine/Source/Runtime/Slate/Private/Framework/Docking/TabManager.cpp b/Engine/Source/Runtime/Slate/Private/Framework/Docking/TabManager.cpp
--- a/Engine/Source/Runtime/Slate/Private/Framework/Docking/TabManager.cpp
+++ b/Engine/Source/Runtime/Slate/Private/Framework/Docking/TabManager.cpp
@@ -1653,7 +1653,10 @@ FTabManager::FTabManager( const TSharedPtr<SDockTab>& InOwnerTab, const TSharedR
 {
 	LocalWorkspaceMenuRoot = FWorkspaceItem::NewGroup(LOCTEXT("LocalWorkspaceRoot", "Local Workspace Root"));
 
-	FSlateApplication::Get().OnWindowBeingDestroyed().AddRaw(this, &FTabManager::OnWindowBeingDestroyed);
+	if (FSlateApplication::IsInitialized())
+	{
+		FSlateApplication::Get().OnWindowBeingDestroyed().AddRaw(this, &FTabManager::OnWindowBeingDestroyed);
+	}
 }
 
 TSharedPtr<SDockingArea> FTabManager::RestoreArea(const TSharedRef<FArea>& AreaToRestore, const TSharedPtr<SWindow>& InParentWindow, const bool bEmbedTitleAreaContent, const EOutputCanBeNullptr OutputCanBeNullptr, bool bForceOpenWindowIfNeeded)
