Add CVAR to always show search field in menus that support it

* Add the CVAR Slate.AlwaysShowMenuSearchField (default to false). When it is enabled, the search field of menus that support search always show the search widget. The idea is to make the feature easier to find.

diff --git a/Engine/Source/Runtime/Slate/Private/Framework/MultiBox/MultiBox.cpp b/Engine/Source/Runtime/Slate/Private/Framework/MultiBox/MultiBox.cpp
--- a/Engine/Source/Runtime/Slate/Private/Framework/MultiBox/MultiBox.cpp
+++ b/Engine/Source/Runtime/Slate/Private/Framework/MultiBox/MultiBox.cpp
@@ -85,6 +85,10 @@ static FString ToString(const TArray<FText>& SearchTextHierarchyComponents)
 } // namespace MultiBoxUtils
 } // namespace UE
 
+TAutoConsoleVariable<bool> AlwaysShowMenuSearchField(
+	TEXT("Slate.AlwaysShowMenuSearchField"),
+	false,
+	TEXT("Always show the search field in menus (default false). When this is disabled, the user has to type for the search field to appear."));
 
 TAttribute<bool> FMultiBoxSettings::UseSmallToolBarIcons;
 TAttribute<bool> FMultiBoxSettings::DisplayMultiboxHooks;
@@ -879,13 +883,17 @@ void SMultiBoxWidget::CreateSearchTextWidget()
 		return;
 	}
 
-	SearchTextWidget = 
+	const FText SearchHint = AlwaysShowMenuSearchField.GetValueOnAnyThread()
+							   ? LOCTEXT("SearchHintStartTyping", "Start typing to search")
+							   : LOCTEXT("SearchHint", "Search");
+
+	SearchTextWidget =
 		SNew(SSearchBox)
-		   .HintText(LOCTEXT("SearchHint", "Search"))
-		   .SelectAllTextWhenFocused(false)
-		   .OnTextChanged(this, &SMultiBoxWidget::OnFilterTextChanged);
+			.HintText(SearchHint)
+			.SelectAllTextWhenFocused(false)
+			.OnTextChanged(this, &SMultiBoxWidget::OnFilterTextChanged);
 
-	TSharedRef< FWidgetBlock > NewWidgetBlock(new FWidgetBlock(SearchTextWidget.ToSharedRef(), FText::GetEmpty(), false));
+	TSharedRef<FWidgetBlock> NewWidgetBlock(new FWidgetBlock(SearchTextWidget.ToSharedRef(), FText::GetEmpty(), false));
 	NewWidgetBlock->SetSearchable(false);
 
 	MultiBox->AddMultiBlockToFront(NewWidgetBlock);
@@ -902,10 +910,13 @@ void SMultiBoxWidget::OnFilterTextChanged(const FText& InFilterText)
 	{
 		if (SearchTextWidget.IsValid() && SearchBlockWidget.IsValid())
 		{
-			// Make the search box visible and focused
-			SearchBlockWidget->SetVisibility(EVisibility::Visible);
+			// We only have to do this if we're not always showing the search widget.
+			if (!AlwaysShowMenuSearchField.GetValueOnAnyThread())
+			{
+				// Make the search box visible and focused
+				SearchBlockWidget->SetVisibility(EVisibility::Visible);
+			}
 			FSlateApplication::Get().SetUserFocus(0, SearchTextWidget);
-
 		}
 	}
 
@@ -1616,8 +1627,12 @@ void SMultiBoxWidget::BeginSearch(const TCHAR InChar)
 
 		if (SearchTextWidget.IsValid() && SearchBlockWidget.IsValid())
 		{
-			// Make the search box visible and focused
-			SearchBlockWidget->SetVisibility(EVisibility::Visible);
+			// We only have to do this if we're not always showing the search widget.
+			if (!AlwaysShowMenuSearchField.GetValueOnAnyThread())
+			{
+				// Make the search box visible and focused
+				SearchBlockWidget->SetVisibility(EVisibility::Visible);
+			}
 			FSlateApplication::Get().SetUserFocus(0, SearchTextWidget);
 
 			SearchTextWidget->SetText(FText::FromString(NewSearchText));
@@ -1739,9 +1754,13 @@ void SMultiBoxWidget::FilterMultiBoxEntries()
 			It.Key()->SetVisibility(EVisibility::Visible);
 		}
 
-		if (SearchBlockWidget.IsValid())
+		// We only have to do this if we're not always showing the search widget.
+		if (!AlwaysShowMenuSearchField.GetValueOnAnyThread())
 		{
-			SearchBlockWidget->SetVisibility(EVisibility::Collapsed);
+			if (SearchBlockWidget.IsValid())
+			{
+				SearchBlockWidget->SetVisibility(EVisibility::Collapsed);
+			}
 		}
 
 		// Hide the sub-menus widgets that were made visible by searching this multi-box hierarchy.
@@ -1765,6 +1784,15 @@ void SMultiBoxWidget::FilterMultiBoxEntries()
 		FText DisplayText = UE::MultiBoxUtils::GetBlockWidgetDisplayText(It.Value());
 		const TSharedPtr<SWidget>& Widget = It.Key();
 
+		// Skip the search widget itself when scanning for searchable items.
+		if (AlwaysShowMenuSearchField.GetValueOnAnyThread())
+		{
+			if (Widget == SearchBlockWidget)
+			{
+				continue;
+			}
+		}
+
 		// Non-labeled elements should not be visible when searching
 		if (DisplayText.IsEmpty())
 		{
@@ -1809,9 +1837,14 @@ void SMultiBoxWidget::FilterMultiBoxEntries()
 		}
 	}
 
-	if (SearchBlockWidget.IsValid())
+	// If we always show the search widget, we're skipping it in the code above and do not need to show it here to compensate.
+	if (!AlwaysShowMenuSearchField.GetValueOnAnyThread())
 	{
-		SearchBlockWidget->SetVisibility(EVisibility::Visible);
+		// Show the search widget again, it was hidden by the above code.
+		if (SearchBlockWidget.IsValid())
+		{
+			SearchBlockWidget->SetVisibility(EVisibility::Visible);
+		}
 	}
 }
 
diff --git a/Engine/Source/Runtime/Slate/Private/Framework/MultiBox/SWidgetBlock.cpp b/Engine/Source/Runtime/Slate/Private/Framework/MultiBox/SWidgetBlock.cpp
--- a/Engine/Source/Runtime/Slate/Private/Framework/MultiBox/SWidgetBlock.cpp
+++ b/Engine/Source/Runtime/Slate/Private/Framework/MultiBox/SWidgetBlock.cpp
@@ -111,7 +111,7 @@ void SWidgetBlock::BuildMultiBlockWidget(const ISlateStyle* StyleSet, const FNam
 		// If there is no label, allow the custom menu widget to consume the entire space
 		if (!bHasLabel)
 		{
-			// Wee use 1 pixel of padding to ensure the menu border shows up
+			// We use 1 pixel of padding to ensure the menu border shows up
 			Padding = FMargin(1);
 		}
 	}
@@ -146,7 +146,15 @@ void SWidgetBlock::BuildMultiBlockWidget(const ISlateStyle* StyleSet, const FNam
 	if (OwnerMultiBoxWidgetPinned->GetSearchTextWidget() == WidgetBlock->ContentWidget)
 	{
 		OwnerMultiBoxWidgetPinned->SetSearchBlockWidget(this->AsWidget());
-		this->AsWidget()->SetVisibility(EVisibility::Collapsed);
+
+		// When we are always showing the search widget, we should not hide it here.
+		if (IConsoleVariable* AlwaysShowMenuSearchFieldVar = IConsoleManager::Get().FindConsoleVariable(TEXT("Slate.AlwaysShowMenuSearchField")))
+		{
+			if (!AlwaysShowMenuSearchFieldVar->GetBool())
+			{
+				this->AsWidget()->SetVisibility(EVisibility::Collapsed);
+			}
+		}
 	}
 
 	ChildSlot
