Fix editor hang when trying to debug a blueprint while a different modal dialog is already open

diff --git a/Engine/Source/Runtime/Slate/Private/Framework/Application/SlateApplication.cpp b/Engine/Source/Runtime/Slate/Private/Framework/Application/SlateApplication.cpp
--- a/Engine/Source/Runtime/Slate/Private/Framework/Application/SlateApplication.cpp
+++ b/Engine/Source/Runtime/Slate/Private/Framework/Application/SlateApplication.cpp
@@ -3630,6 +3630,12 @@ void FSlateApplication::GetAllVisibleChildWindows(TArray< TSharedRef<SWindow> >&
 
 void FSlateApplication::EnterDebuggingMode()
 {
+	if (GetActiveModalWindow().IsValid())
+	{
+		UE_LOG(LogSlate, Warning, TEXT("EnterDebuggingMode is not supported while a modal window is open."));
+		return;
+	}
+
 	bRequestLeaveDebugMode = false;
 
 	// Note it is ok to hold a reference here as the game viewport should not be destroyed while in debugging mode
