[Keyboard Shortcuts] Add the ability to have project level, read only keybindings supported by the input manager.

diff --git a/Engine/Source/Runtime/Slate/Private/Framework/Commands/InputBindingManager.cpp b/Engine/Source/Runtime/Slate/Private/Framework/Commands/InputBindingManager.cpp
--- a/Engine/Source/Runtime/Slate/Private/Framework/Commands/InputBindingManager.cpp
+++ b/Engine/Source/Runtime/Slate/Private/Framework/Commands/InputBindingManager.cpp
@@ -44,6 +44,7 @@ struct FUserDefinedChordKey
 	FName BindingContext;
 	FName CommandName;
 	EMultipleKeyBindingIndex ChordIndex;
+	bool bIsFromProjectSetting = false;
 };
 
 uint32 GetTypeHash( const FUserDefinedChordKey& Key )
@@ -62,14 +63,29 @@ public:
 	void RemoveAll();
 private:
 
-	void LoadChord(const TSharedPtr<FJsonObject>& ChordInfoObject, const FName& BindingContextName, const EMultipleKeyBindingIndex ChordIndex, const FName& CommandName);
+	static FUserDefinedChordKey LoadKey(const TSharedPtr<FJsonObject>& ChordInfoObject);
+	static FInputChord LoadChord(const TSharedPtr<FJsonObject>& ChordInfoObject);
+	void LoadFromCurrentJsonConfig(const TArray<FString>& ChordJsonArray, bool bFromProjectSettings);
 
 	/* Mapping from a chord key to the user defined chord */
 	typedef TMap<FUserDefinedChordKey, FInputChord> FChordsMap;
 	TSharedPtr<FChordsMap> Chords;
 };
 
-void FUserDefinedChords::LoadChord(const TSharedPtr<FJsonObject>& ChordInfoObject, const FName& BindingContextName, const EMultipleKeyBindingIndex ChordIndex, const FName& CommandName)
+FUserDefinedChordKey FUserDefinedChords::LoadKey(const TSharedPtr<FJsonObject>& ChordInfoObject)
+{
+	const TSharedPtr<FJsonValue> BindingContextObj = ChordInfoObject->Values.FindRef(TEXT("BindingContext"));
+	const TSharedPtr<FJsonValue> CommandNameObj = ChordInfoObject->Values.FindRef(TEXT("CommandName"));
+	const TSharedPtr<FJsonValue> ChordIndexObj = ChordInfoObject->Values.FindRef(TEXT("ChordIndex"));
+
+	const FName BindingContext = *BindingContextObj->AsString();
+	const FName CommandName = *CommandNameObj->AsString();
+	const EMultipleKeyBindingIndex ChordIndex = ChordIndexObj.IsValid() ? static_cast<EMultipleKeyBindingIndex>(static_cast<uint32>(ChordIndexObj->AsNumber())) : EMultipleKeyBindingIndex::Primary;
+
+	return FUserDefinedChordKey(BindingContext, CommandName, ChordIndex);
+}
+
+FInputChord FUserDefinedChords::LoadChord(const TSharedPtr<FJsonObject>& ChordInfoObject)
 {
 	const TSharedPtr<FJsonValue> CtrlObj = ChordInfoObject->Values.FindRef(TEXT("Control"));
 	const TSharedPtr<FJsonValue> AltObj = ChordInfoObject->Values.FindRef( TEXT("Alt") );
@@ -77,9 +93,6 @@ void FUserDefinedChords::LoadChord(const TSharedPtr<FJsonObject>& ChordInfoObjec
 	const TSharedPtr<FJsonValue> CmdObj = ChordInfoObject->Values.FindRef(TEXT("Command"));
 	const TSharedPtr<FJsonValue> KeyObj = ChordInfoObject->Values.FindRef( TEXT("Key") );
 
-	const FUserDefinedChordKey ChordKey(BindingContextName, CommandName, ChordIndex);
-	FInputChord& UserDefinedChord = Chords->FindOrAdd(ChordKey);
-
 #if PLATFORM_MAC
 	// Command is treated like "Control" on Mac and vice-versa, so swap them.
 	EModifierKey::Type Modifiers = EModifierKey::FromBools(
@@ -97,7 +110,26 @@ void FUserDefinedChords::LoadChord(const TSharedPtr<FJsonObject>& ChordInfoObjec
 									);
 #endif	//#if PLATFORM_MAC
 
-	UserDefinedChord = FInputChord(*KeyObj->AsString(), Modifiers);
+	return FInputChord(*KeyObj->AsString(), Modifiers);
+}
+
+void FUserDefinedChords::LoadFromCurrentJsonConfig(const TArray<FString>& ChordJsonArray, bool bFromProjectSettings)
+{
+	// This loads an array of JSON strings representing the FUserDefinedChordKey and FInputChord in a single JSON object
+	for (const FString& ChordJson : ChordJsonArray)
+	{
+		const FString UnescapedContent = FRemoteConfig::ReplaceIniSpecialCharWithChar(ChordJson).ReplaceEscapedCharWithChar();
+
+		TSharedPtr<FJsonObject> ChordInfoObj;
+		auto JsonReader = TJsonReaderFactory<>::Create(UnescapedContent);
+		if (FJsonSerializer::Deserialize(JsonReader, ChordInfoObj))
+		{
+			FUserDefinedChordKey ChordKey = LoadKey(ChordInfoObj);
+			ChordKey.bIsFromProjectSetting = bFromProjectSettings;
+			Chords->Remove(ChordKey);
+			FInputChord& UserDefinedChord = Chords->FindOrAdd(ChordKey, LoadChord(ChordInfoObj));
+		}
+	}
 }
 
 void FUserDefinedChords::LoadChords()
@@ -106,9 +138,13 @@ void FUserDefinedChords::LoadChords()
 	{
 		Chords = MakeShareable( new FChordsMap );
 
+		// First look through project settings
+		TArray<FString> ChordJsonArray;
+		bool bFoundProjectChords = (GConfig->GetArray(TEXT("ProjectDefinedChords"), TEXT("ProjectDefinedChords"), ChordJsonArray, GEditorSettingsIni) > 0);
+		LoadFromCurrentJsonConfig(ChordJsonArray, true);
+
 		// First, try and load the chords from their new location in the ini file
 		// Failing that, try and load them from the older txt file
-		TArray<FString> ChordJsonArray;
 		bool bFoundChords = (GConfig->GetArray(TEXT("UserDefinedChords"), TEXT("UserDefinedChords"), ChordJsonArray, GEditorKeyBindingsIni) > 0);
 		if (!bFoundChords)
 		{
@@ -118,26 +154,7 @@ void FUserDefinedChords::LoadChords()
 
 		if(bFoundChords)
 		{
-			// This loads an array of JSON strings representing the FUserDefinedChordKey and FInputChord in a single JSON object
-			for(const FString& ChordJson : ChordJsonArray)
-			{
-				const FString UnescapedContent = FRemoteConfig::ReplaceIniSpecialCharWithChar(ChordJson).ReplaceEscapedCharWithChar();
-
-				TSharedPtr<FJsonObject> ChordInfoObj;
-				auto JsonReader = TJsonReaderFactory<>::Create( UnescapedContent );
-				if( FJsonSerializer::Deserialize( JsonReader, ChordInfoObj ) )
-				{
-					const TSharedPtr<FJsonValue> BindingContextObj = ChordInfoObj->Values.FindRef( TEXT("BindingContext") );
-					const TSharedPtr<FJsonValue> CommandNameObj = ChordInfoObj->Values.FindRef( TEXT("CommandName") );
-					const TSharedPtr<FJsonValue> ChordIndexObj = ChordInfoObj->Values.FindRef(TEXT("ChordIndex") );
-
-					const FName BindingContext = *BindingContextObj->AsString();
-					const FName CommandName = *CommandNameObj->AsString(); 
-					const EMultipleKeyBindingIndex ChordIndex = ChordIndexObj.IsValid() ? static_cast<EMultipleKeyBindingIndex>(static_cast<uint32>(ChordIndexObj->AsNumber())) : EMultipleKeyBindingIndex::Primary;
-
-					LoadChord(ChordInfoObj, BindingContext, ChordIndex, CommandName);
-				}
-			}
+			LoadFromCurrentJsonConfig(ChordJsonArray, false);
 		}
 		else
 		{
@@ -182,7 +199,9 @@ void FUserDefinedChords::LoadChords()
 						TSharedPtr<FJsonObject> CommandObj = CommandInfo.Value->AsObject();
 						for (uint32 i = 0; i < static_cast<uint8>(EMultipleKeyBindingIndex::NumChords); ++i)
 						{
-							LoadChord(CommandObj, BindingContext, static_cast<EMultipleKeyBindingIndex>(i), CommandName);
+							FUserDefinedChordKey ChordKey(BindingContext, CommandName, static_cast<EMultipleKeyBindingIndex>(i));
+							Chords->Remove(ChordKey); // Remove any existing that may have been added by project settings, to overwrite the flag on the key
+							Chords->FindOrAdd(ChordKey, LoadChord(CommandObj));
 						}
 					}
 				}
@@ -197,8 +216,14 @@ void FUserDefinedChords::SaveChords() const
 	{
 		FString ChordRawJsonContent;
 		TArray<FString> ChordJsonArray;
-		for(const auto& ChordInfo : *Chords)
+		for(const TPair<FUserDefinedChordKey, FInputChord>& ChordInfo : *Chords)
 		{
+			// Don't save keybinds that came out of the project settings. These can however be overridden by user defined keybinds successfully
+			if (ChordInfo.Key.bIsFromProjectSetting)
+			{
+				continue;
+			}
+
 			TSharedPtr<FJsonValueObject> ChordInfoValueObj = MakeShareable( new FJsonValueObject( MakeShareable( new FJsonObject ) ) );
 			TSharedPtr<FJsonObject> ChordInfoObj = ChordInfoValueObj->AsObject();
 
@@ -256,8 +281,9 @@ void FUserDefinedChords::SetUserDefinedChords( const FUICommandInfo& CommandInfo
 		for (uint32 i = 0; i < static_cast<uint8>(EMultipleKeyBindingIndex::NumChords); ++i)
 		{
 			EMultipleKeyBindingIndex ChordIndex = static_cast<EMultipleKeyBindingIndex>(i);
-			// Find or create the command context
+			// Recreate the command context, since it may have come from project chord, and now we want to be sure to flag as a user chord
 			const FUserDefinedChordKey ChordKey(BindingContext, CommandName, ChordIndex);
+			Chords->Remove(ChordKey);
 			FInputChord& UserDefinedChord = Chords->FindOrAdd(ChordKey);
 
 			// Save an empty invalid chord if one was not set
diff --git a/Engine/Source/Runtime/Slate/Public/Framework/Commands/InputBindingManager.h b/Engine/Source/Runtime/Slate/Public/Framework/Commands/InputBindingManager.h
--- a/Engine/Source/Runtime/Slate/Public/Framework/Commands/InputBindingManager.h
+++ b/Engine/Source/Runtime/Slate/Public/Framework/Commands/InputBindingManager.h
@@ -146,6 +146,15 @@ public:
 	*/
 	SLATE_API bool GetUserDefinedChord(const FName InBindingContext, const FName InCommandName, const EMultipleKeyBindingIndex InChordIndex, FInputChord& OutUserDefinedChord);
 
+	/**
+	* Gets the project defined chord (if any) from the provided command name
+	*
+	* @param InBindingContext	The context in which the command is active
+	* @param InCommandName		The name of the command to get the chord from
+	* @param ChordIndex		The index of the key binding (in the multiple key bindings array)
+	*/
+	SLATE_API bool GetProjectDefinedChord(const FName InBindingContext, const FName InCommandName, const EMultipleKeyBindingIndex InChordIndex, FInputChord& OutUserDefinedChord);
+
 	// Type of filter
 	enum class ECommandFilterType
 	{
@@ -272,6 +281,9 @@ private:
 	/** User defined chord overrides for commands */
 	TSharedPtr< class FUserDefinedChords > UserDefinedChords;
 
+	/** Project defined chord overrides for commands */
+	TSharedPtr< class FUserDefinedChords > ProjectDefinedChords;
+
 	/** Delegate called when a user-defined chord is edited */
 	FOnUserDefinedChordChanged OnUserDefinedChordChanged;
 };
