Highlight active search field in menus

Highlight the search field that will receive keyboard input when "Slate.AlwaysShowMenuSearchField" is turned on. This makes it easier for users to understand where their typing will end up when search fields are drawn both in a parent menu and its submenus, as this will highlight the active input field and make the others look greyed out (all using SWidget::SetEnabled).

diff --git a/Engine/Source/Runtime/Slate/Private/Framework/MultiBox/MultiBox.cpp b/Engine/Source/Runtime/Slate/Private/Framework/MultiBox/MultiBox.cpp
--- a/Engine/Source/Runtime/Slate/Private/Framework/MultiBox/MultiBox.cpp
+++ b/Engine/Source/Runtime/Slate/Private/Framework/MultiBox/MultiBox.cpp
@@ -1577,6 +1577,26 @@ FReply SMultiBoxWidget::OnFocusReceived( const FGeometry& MyGeometry, const FFoc
 	return FReply::Unhandled();
 }
 
+void SMultiBoxWidget::OnFocusChanging(const FWeakWidgetPath& PreviousFocusPath, const FWidgetPath& NewWidgetPath, const FFocusEvent& InFocusEvent)
+{
+	if (!AlwaysShowMenuSearchField.GetValueOnAnyThread())
+	{
+		return;
+	}
+
+	if (SearchTextWidget)
+	{
+		// We need to figure out if we're on the current focus path but after the last window, and thus
+		// live in the most-expanded-to submenu. If that's the case, we should enable our search field to
+		// signify to the user that any keyboard input will go there. If not, we disable our search field
+		// (and another menu will enable its search field instead).
+		const TSharedRef<SWindow> DeepestWindow = NewWidgetPath.GetDeepestWindow();
+		const FWidgetPath WidgetPathToThis = NewWidgetPath.GetPathDownTo(AsShared());
+		const bool bIsChildOfDeepestWindow = WidgetPathToThis.IsValid() && WidgetPathToThis.ContainsWidget(&DeepestWindow.Get());
+		SearchTextWidget->SetEnabled(bIsChildOfDeepestWindow);
+	}
+}
+
 FReply SMultiBoxWidget::OnKeyDown( const FGeometry& MyGeometry, const FKeyEvent& KeyEvent )
 {
 	SCompoundWidget::OnKeyDown( MyGeometry, KeyEvent );
diff --git a/Engine/Source/Runtime/Slate/Public/Framework/MultiBox/MultiBox.h b/Engine/Source/Runtime/Slate/Public/Framework/MultiBox/MultiBox.h
--- a/Engine/Source/Runtime/Slate/Public/Framework/MultiBox/MultiBox.h
+++ b/Engine/Source/Runtime/Slate/Public/Framework/MultiBox/MultiBox.h
@@ -724,6 +724,7 @@ public:
 	SLATE_API virtual FReply OnDrop( const FGeometry& MyGeometry, const FDragDropEvent& DragDropEvent ) override;
 	SLATE_API virtual bool SupportsKeyboardFocus() const override;
 	SLATE_API virtual FReply OnFocusReceived( const FGeometry& MyGeometry, const FFocusEvent& InFocusEvent ) override;
+	SLATE_API virtual void OnFocusChanging(const FWeakWidgetPath& PreviousFocusPath, const FWidgetPath& NewWidgetPath, const FFocusEvent& InFocusEvent) override;
 	SLATE_API virtual FReply OnKeyDown( const FGeometry& MyGeometry, const FKeyEvent& KeyEvent ) override;
 	SLATE_API virtual FReply OnKeyChar(const FGeometry& MyGeometry, const FCharacterEvent& InCharacterEvent) override;
 	SLATE_API virtual bool OnVisualizeTooltip(const TSharedPtr<SWidget>& TooltipContent) override;
