Fix padding around menu search fields

The search fields in our menus were laid out too close to the left and right edges of the menus. This adds 8 pixels of padding on both sides.

diff --git a/Engine/Source/Runtime/Slate/Private/Framework/MultiBox/MultiBox.cpp b/Engine/Source/Runtime/Slate/Private/Framework/MultiBox/MultiBox.cpp
--- a/Engine/Source/Runtime/Slate/Private/Framework/MultiBox/MultiBox.cpp
+++ b/Engine/Source/Runtime/Slate/Private/Framework/MultiBox/MultiBox.cpp
@@ -893,7 +893,14 @@ void SMultiBoxWidget::CreateSearchTextWidget()
 			.SelectAllTextWhenFocused(false)
 			.OnTextChanged(this, &SMultiBoxWidget::OnFilterTextChanged);
 
-	TSharedRef<FWidgetBlock> NewWidgetBlock(new FWidgetBlock(SearchTextWidget.ToSharedRef(), FText::GetEmpty(), false));
+	TSharedRef<SBox> SearchBox =
+		SNew(SBox)
+			.Padding(FMargin(8, 0, 8, 0))
+			[
+				SearchTextWidget.ToSharedRef()
+			];
+
+	TSharedRef<FWidgetBlock> NewWidgetBlock(new FWidgetBlock(SearchBox, FText::GetEmpty(), false));
 	NewWidgetBlock->SetSearchable(false);
 
 	MultiBox->AddMultiBlockToFront(NewWidgetBlock);
