Revert changelist 29745157

The reverted change prevented focus to jump to the search field in menus when users typed on the keyboard.

diff --git a/Engine/Source/Runtime/Slate/Private/Framework/MultiBox/MultiBox.cpp b/Engine/Source/Runtime/Slate/Private/Framework/MultiBox/MultiBox.cpp
--- a/Engine/Source/Runtime/Slate/Private/Framework/MultiBox/MultiBox.cpp
+++ b/Engine/Source/Runtime/Slate/Private/Framework/MultiBox/MultiBox.cpp
@@ -893,14 +893,7 @@ void SMultiBoxWidget::CreateSearchTextWidget()
 			.SelectAllTextWhenFocused(false)
 			.OnTextChanged(this, &SMultiBoxWidget::OnFilterTextChanged);
 
-	TSharedRef<SBox> SearchBox =
-		SNew(SBox)
-			.Padding(FMargin(8, 0, 8, 0))
-			[
-				SearchTextWidget.ToSharedRef()
-			];
-
-	TSharedRef<FWidgetBlock> NewWidgetBlock(new FWidgetBlock(SearchBox, FText::GetEmpty(), false));
+	TSharedRef<FWidgetBlock> NewWidgetBlock(new FWidgetBlock(SearchTextWidget.ToSharedRef(), FText::GetEmpty(), false));
 	NewWidgetBlock->SetSearchable(false);
 
 	MultiBox->AddMultiBlockToFront(NewWidgetBlock);
