Fix padding around menu search fields

* The search fields in our menus were laid out too close to the left and right edges of the menus. This adds 8 pixels of padding on both sides.
* Update a check in SWidgetBlock to adjust for the nested structure where the SSearchField of menus now reside inside an SBox.

diff --git a/Engine/Source/Runtime/Slate/Private/Framework/MultiBox/MultiBox.cpp b/Engine/Source/Runtime/Slate/Private/Framework/MultiBox/MultiBox.cpp
--- a/Engine/Source/Runtime/Slate/Private/Framework/MultiBox/MultiBox.cpp
+++ b/Engine/Source/Runtime/Slate/Private/Framework/MultiBox/MultiBox.cpp
@@ -893,7 +893,14 @@ void SMultiBoxWidget::CreateSearchTextWidget()
 			.SelectAllTextWhenFocused(false)
 			.OnTextChanged(this, &SMultiBoxWidget::OnFilterTextChanged);
 
-	TSharedRef<FWidgetBlock> NewWidgetBlock(new FWidgetBlock(SearchTextWidget.ToSharedRef(), FText::GetEmpty(), false));
+	TSharedRef<SBox> SearchBox =
+		SNew(SBox)
+			.Padding(FMargin(8, 0, 8, 0))
+			[
+				SearchTextWidget.ToSharedRef()
+			];
+
+	TSharedRef<FWidgetBlock> NewWidgetBlock(new FWidgetBlock(SearchBox, FText::GetEmpty(), false));
 	NewWidgetBlock->SetSearchable(false);
 
 	MultiBox->AddMultiBlockToFront(NewWidgetBlock);
diff --git a/Engine/Source/Runtime/Slate/Private/Framework/MultiBox/SWidgetBlock.cpp b/Engine/Source/Runtime/Slate/Private/Framework/MultiBox/SWidgetBlock.cpp
--- a/Engine/Source/Runtime/Slate/Private/Framework/MultiBox/SWidgetBlock.cpp
+++ b/Engine/Source/Runtime/Slate/Private/Framework/MultiBox/SWidgetBlock.cpp
@@ -143,7 +143,7 @@ void SWidgetBlock::BuildMultiBlockWidget(const ISlateStyle* StyleSet, const FNam
 	OwnerMultiBoxWidgetPinned->AddElement(this->AsWidget(), SearchLabel, MultiBlock->GetSearchable());
 
 	// This widget holds the search text, set it as the search block widget
-	if (OwnerMultiBoxWidgetPinned->GetSearchTextWidget() == WidgetBlock->ContentWidget)
+	if (OwnerMultiBoxWidgetPinned->GetSearchTextWidget() && OwnerMultiBoxWidgetPinned->GetSearchTextWidget()->GetParentWidget() == WidgetBlock->ContentWidget)
 	{
 		OwnerMultiBoxWidgetPinned->SetSearchBlockWidget(this->AsWidget());
 
