UToolMenus: Add API to support Menu Profiles, a menu profile is a lightweight version of a menu customization - but you can have multiple menu profiles and activate/deactivate them at runtime per menu instance using a new menu context.

Read Only Asset Editors: Automatically hide any custom menu and toolbar items in read only mode, with an opt in permission list for specific asset editors to selectively allow entries that are valid in read only mode

diff --git a/Engine/Source/Developer/ToolMenus/Private/ToolMenu.cpp b/Engine/Source/Developer/ToolMenus/Private/ToolMenu.cpp
--- a/Engine/Source/Developer/ToolMenus/Private/ToolMenu.cpp
+++ b/Engine/Source/Developer/ToolMenus/Private/ToolMenu.cpp
@@ -496,6 +496,38 @@ FCustomizedToolMenuHierarchy UToolMenu::GetMenuCustomizationHierarchy() const
 	return Result;
 }
 
+FToolMenuProfile* UToolMenu::FindMenuProfile(const FName& ProfileName) const
+{
+	return UToolMenus::Get()->FindMenuProfile(MenuName, ProfileName);
+}
+
+FToolMenuProfile* UToolMenu::AddMenuProfile(const FName& ProfileName) const
+{
+	return UToolMenus::Get()->AddMenuProfile(MenuName, ProfileName);
+}
+
+FToolMenuProfileHierarchy UToolMenu::GetMenuProfileHierarchy(const FName& ProfileName) const
+{
+	FToolMenuProfileHierarchy Result;
+	
+	UToolMenus* ToolMenus = UToolMenus::Get();
+	TArray<FName> HierarchyNames = GetMenuHierarchyNames(true);
+	for (const FName& ItName : HierarchyNames)
+	{
+		if (FToolMenuProfile* Found = ToolMenus->FindMenuProfile(ItName, ProfileName))
+		{
+			Result.ProfileHierarchy.Add(Found);
+		}
+
+		if (FToolMenuProfile* FoundRuntime = ToolMenus->FindRuntimeMenuProfile(ItName, ProfileName))
+		{
+			Result.RuntimeProfileHierarchy.Add(FoundRuntime);
+		}
+	}
+
+	return Result;
+}
+
 void UToolMenu::UpdateMenuCustomizationFromMultibox(const TSharedRef<const FMultiBox>& InMultiBox)
 {
 	FCustomizedToolMenu* Customization = AddMenuCustomization();
diff --git a/Engine/Source/Developer/ToolMenus/Private/ToolMenus.cpp b/Engine/Source/Developer/ToolMenus/Private/ToolMenus.cpp
--- a/Engine/Source/Developer/ToolMenus/Private/ToolMenus.cpp
+++ b/Engine/Source/Developer/ToolMenus/Private/ToolMenus.cpp
@@ -933,20 +933,148 @@ FCustomizedToolMenu* UToolMenus::AddRuntimeMenuCustomization(const FName InName)
 	}
 }
 
-void UToolMenus::ApplyCustomization(UToolMenu* GeneratedMenu)
+FToolMenuProfile* UToolMenus::FindMenuProfile(const FName InMenuName, const FName InProfileName)
 {
+	if(FToolMenuProfileMap* FoundMenu = MenuProfiles.Find(InMenuName))
+	{
+		return FoundMenu->MenuProfiles.Find(InProfileName);
+	}
+
+	return nullptr;
+}
+
+FToolMenuProfile* UToolMenus::AddMenuProfile(const FName InMenuName, const FName InProfileName)
+{
+	if (FToolMenuProfile* Found = FindMenuProfile(InMenuName, InProfileName))
+	{
+		return Found;
+	}
+	else
+	{
+		FToolMenuProfileMap& FoundMenu = MenuProfiles.FindOrAdd(InMenuName);
+		
+		FToolMenuProfile& NewCustomization = FoundMenu.MenuProfiles.Add(InProfileName, FToolMenuProfile());
+		NewCustomization.Name = InProfileName;
+		return &NewCustomization;
+	}
+}
+
+
+FToolMenuProfile* UToolMenus::FindRuntimeMenuProfile(const FName InMenuName, const FName InProfileName)
+{
+	if(FToolMenuProfileMap* FoundMenu = RuntimeMenuProfiles.Find(InMenuName))
+	{
+		return FoundMenu->MenuProfiles.Find(InProfileName);
+	}
+
+	return nullptr;
+}
+
+FToolMenuProfile* UToolMenus::AddRuntimeMenuProfile(const FName InMenuName, const FName InProfileName)
+{
+	if (FToolMenuProfile* Found = FindRuntimeMenuProfile(InMenuName, InProfileName))
+	{
+		return Found;
+	}
+	else
+	{
+		FToolMenuProfileMap& FoundMenu = RuntimeMenuProfiles.FindOrAdd(InMenuName);
+		
+		FToolMenuProfile& NewCustomization = FoundMenu.MenuProfiles.Add(InProfileName, FToolMenuProfile());
+		NewCustomization.Name = InProfileName;
+		return &NewCustomization;
+	}
+}
+
+void UToolMenus::ApplyCustomizationAndProfiles(UToolMenu* GeneratedMenu)
+{
+	// Apply all profiles that are active by looking for them in the context
+	UToolMenuProfileContext* ProfileContext = GeneratedMenu->FindContext<UToolMenuProfileContext>();
+	
+	if(ProfileContext)
+	{
+		for(const FName& ActiveProfile : ProfileContext->ActiveProfiles)
+		{
+			FToolMenuProfileHierarchy MenuProfileHieararchy = GeneratedMenu->GetMenuProfileHierarchy(ActiveProfile);
+
+			if (MenuProfileHieararchy.ProfileHierarchy.Num() != 0 || MenuProfileHieararchy.RuntimeProfileHierarchy.Num() != 0)
+			{
+				FToolMenuProfile MenuProfile = MenuProfileHieararchy.GenerateFlattenedMenuProfile();
+				ApplyProfile(GeneratedMenu, MenuProfile);
+			}
+			else
+			{
+				UE_LOG(LogToolMenus, Verbose, TEXT("Menu Profile %s for menu %s not found!"), *ActiveProfile.ToString(), *GeneratedMenu->GetMenuName().ToString());
+
+			}
+		}
+	}
+
+	// Apply the customization for the menu (if any)
 	FCustomizedToolMenuHierarchy CustomizationHierarchy = GeneratedMenu->GetMenuCustomizationHierarchy();
-	if (CustomizationHierarchy.Hierarchy.Num() == 0 && CustomizationHierarchy.RuntimeHierarchy.Num() == 0)
+	if (CustomizationHierarchy.Hierarchy.Num() != 0 || CustomizationHierarchy.RuntimeHierarchy.Num() != 0)
 	{
-		return;
+		FCustomizedToolMenu CustomizedMenu = CustomizationHierarchy.GenerateFlattened();
+		ApplyCustomization(GeneratedMenu, CustomizedMenu);
 	}
-	FCustomizedToolMenu CustomizedMenu = CustomizationHierarchy.GenerateFlattened();
+}
 
-	if (CustomizedMenu.IsSuppressExtenders())
+void UToolMenus::ApplyProfile(UToolMenu* GeneratedMenu, const FToolMenuProfile& MenuProfile)
+{
+	if (MenuProfile.IsSuppressExtenders())
 	{
 		GeneratedMenu->SetExtendersEnabled(false);
 	}
+	
+	TArray<FToolMenuSection> NewSections(GeneratedMenu->Sections);
+	
+	// Hide items based on deny list
+	if (MenuProfile.MenuPermissions.HasFiltering())
+	{
+		for (int32 SectionIndex = 0; SectionIndex < NewSections.Num(); ++SectionIndex)
+		{
+			FToolMenuSection& Section = NewSections[SectionIndex];
+			for (int32 i = 0; i < Section.Blocks.Num(); ++i)
+			{
+				if (!MenuProfile.MenuPermissions.PassesFilter(Section.Blocks[i].Name))
+				{
+					Section.Blocks.RemoveAt(i);
+					--i;
+				}
+			}
+		}
+	}
+
+	// Hide sections and entries
+	if (!GeneratedMenu->IsEditing())
+	{
+		for (int32 SectionIndex = 0; SectionIndex < NewSections.Num(); ++SectionIndex)
+		{
+			FToolMenuSection& Section = NewSections[SectionIndex];
+			if (MenuProfile.IsSectionHidden(Section.Name))
+			{
+				NewSections.RemoveAt(SectionIndex);
+				--SectionIndex;
+				continue;
+			}
+
+			for (int32 i = 0; i < Section.Blocks.Num(); ++i)
+			{
+				if (MenuProfile.IsEntryHidden(Section.Blocks[i].Name))
+				{
+					Section.Blocks.RemoveAt(i);
+					--i;
+				}
+			}
+		}
+	}
+
+	GeneratedMenu->Sections = NewSections;
 
+}
+
+void UToolMenus::ApplyCustomization(UToolMenu* GeneratedMenu, const FCustomizedToolMenu& CustomizedMenu)
+{
 	TArray<FToolMenuSection> NewSections;
 	NewSections.Reserve(GeneratedMenu->Sections.Num());
 
@@ -989,7 +1117,7 @@ void UToolMenus::ApplyCustomization(UToolMenu* GeneratedMenu)
 
 		if (OriginalSection.Name != NAME_None)
 		{
-			if (FCustomizedToolMenuNameArray* EntryOrder = CustomizedMenu.EntryOrder.Find(OriginalSection.Name))
+			if (const FCustomizedToolMenuNameArray* EntryOrder = CustomizedMenu.EntryOrder.Find(OriginalSection.Name))
 			{
 				for (const FName& EntryName : EntryOrder->Names)
 				{
@@ -1030,48 +1158,9 @@ void UToolMenus::ApplyCustomization(UToolMenu* GeneratedMenu)
 		}
 	}
 
-	// Hide items based on deny list
-	if (CustomizedMenu.MenuPermissions.HasFiltering())
-	{
-		for (int32 SectionIndex = 0; SectionIndex < NewSections.Num(); ++SectionIndex)
-		{
-			FToolMenuSection& Section = NewSections[SectionIndex];
-			for (int32 i = 0; i < Section.Blocks.Num(); ++i)
-			{
-				if (!CustomizedMenu.MenuPermissions.PassesFilter(Section.Blocks[i].Name))
-				{
-					Section.Blocks.RemoveAt(i);
-					--i;
-				}
-			}
-		}
-	}
-
-	// Hide sections and entries
-	if (!GeneratedMenu->IsEditing())
-	{
-		for (int32 SectionIndex = 0; SectionIndex < NewSections.Num(); ++SectionIndex)
-		{
-			FToolMenuSection& Section = NewSections[SectionIndex];
-			if (CustomizedMenu.IsSectionHidden(Section.Name))
-			{
-				NewSections.RemoveAt(SectionIndex);
-				--SectionIndex;
-				continue;
-			}
-
-			for (int32 i = 0; i < Section.Blocks.Num(); ++i)
-			{
-				if (CustomizedMenu.IsEntryHidden(Section.Blocks[i].Name))
-				{
-					Section.Blocks.RemoveAt(i);
-					--i;
-				}
-			}
-		}
-	}
-
 	GeneratedMenu->Sections = NewSections;
+
+	ApplyProfile(GeneratedMenu, CustomizedMenu);
 }
 
 void UToolMenus::AssembleMenuHierarchy(UToolMenu* GeneratedMenu, const TArray<UToolMenu*>& Hierarchy)
@@ -1083,7 +1172,7 @@ void UToolMenus::AssembleMenuHierarchy(UToolMenu* GeneratedMenu, const TArray<UT
 		AssembleMenu(GeneratedMenu, FoundParent);
 	}
 
-	ApplyCustomization(GeneratedMenu);
+	ApplyCustomizationAndProfiles(GeneratedMenu);
 }
 
 UToolMenu* UToolMenus::GenerateSubMenu(const UToolMenu* InGeneratedParent, const FName InBlockName)
diff --git a/Engine/Source/Developer/ToolMenus/Public/ToolMenu.h b/Engine/Source/Developer/ToolMenus/Public/ToolMenu.h
--- a/Engine/Source/Developer/ToolMenus/Public/ToolMenu.h
+++ b/Engine/Source/Developer/ToolMenus/Public/ToolMenu.h
@@ -113,6 +113,9 @@ public:
 	virtual FCustomizedToolMenu* FindMenuCustomization() const override;
 	virtual FCustomizedToolMenu* AddMenuCustomization() const override;
 	virtual FCustomizedToolMenuHierarchy GetMenuCustomizationHierarchy() const override;
+	virtual FToolMenuProfile* FindMenuProfile(const FName& ProfileName) const override;
+	virtual FToolMenuProfile* AddMenuProfile(const FName& ProfileName) const override;
+	virtual FToolMenuProfileHierarchy GetMenuProfileHierarchy(const FName& ProfileName) const override;
 	virtual void UpdateMenuCustomizationFromMultibox(const TSharedRef<const FMultiBox>& InMultiBox) override;
 	virtual void OnMenuDestroyed() override;
 	//~ End UToolMenuBase Interface
diff --git a/Engine/Source/Developer/ToolMenus/Public/ToolMenus.h b/Engine/Source/Developer/ToolMenus/Public/ToolMenus.h
--- a/Engine/Source/Developer/ToolMenus/Public/ToolMenus.h
+++ b/Engine/Source/Developer/ToolMenus/Public/ToolMenus.h
@@ -43,6 +43,30 @@ struct FGeneratedToolMenuWidgets
 	TArray<FGeneratedToolMenuWidget> Instances;
 };
 
+/*
+ * A global context that any menu can add/modify to specify which profiles are currently active
+ */
+UCLASS()
+class TOOLMENUS_API UToolMenuProfileContext : public UToolMenuContextBase
+{
+	GENERATED_BODY()
+public:
+
+	TArray<FName> ActiveProfiles;
+};
+
+/*
+ * Struct to store all the profiles for a menu for serialization
+ */
+USTRUCT()
+struct FToolMenuProfileMap
+{
+	GENERATED_BODY()
+
+	UPROPERTY()
+	TMap<FName /*profile name*/, FToolMenuProfile> MenuProfiles;
+};
+
 UCLASS(config=EditorPerProjectUserSettings)
 class TOOLMENUS_API UToolMenus : public UObject
 {
@@ -254,6 +278,18 @@ public:
 	/** Find or add runtime customization settings for a menu */
 	FCustomizedToolMenu* AddRuntimeMenuCustomization(const FName InName);
 
+	/** Find a specific profile for a menu */
+	FToolMenuProfile* FindMenuProfile(const FName InMenuName, const FName InProfileName);
+
+	/** Find or add a specific profile for a menu */
+	FToolMenuProfile* AddMenuProfile(const FName InMenuName, const FName InProfileName);
+
+	/** Find a specific runtime only profile for a menu */
+	FToolMenuProfile* FindRuntimeMenuProfile(const FName InMenuName, const FName InProfileName);
+
+	/** Find or add a specific runtime only profile for a menu */
+	FToolMenuProfile* AddRuntimeMenuProfile(const FName InMenuName, const FName InProfileName);
+
 	/** Unregister runtime customization settings for a specific owner name */
 	void UnregisterRuntimeMenuCustomizationOwner(const FName InOwnerName);
 
@@ -368,7 +404,9 @@ private:
 
 	void AddReferencedContextObjects(const TSharedRef<FMultiBox>& InMultiBox, const UToolMenu* InMenu);
 
-	void ApplyCustomization(UToolMenu* GeneratedMenu);
+	void ApplyCustomizationAndProfiles(UToolMenu* GeneratedMenu);
+	void ApplyProfile(UToolMenu* GeneratedMenu, const FToolMenuProfile& Profile);
+	void ApplyCustomization(UToolMenu* GeneratedMenu, const FCustomizedToolMenu& Profile);
 
 	void UnregisterOwnerInternal(FToolMenuOwner Owner);
 
@@ -382,7 +420,7 @@ private:
 
 	UPROPERTY(EditAnywhere, Category = Misc)
 	TArray<FCustomizedToolMenu> CustomizedMenus;
-
+	
 	/* Allow substituting one menu for another during generate but not during find or extend */
 	UPROPERTY(EditAnywhere, Category = Misc)
 	TMap<FName, FName> MenuSubstitutionsDuringGenerate;
@@ -401,6 +439,13 @@ private:
 	/** Transient customizations made during runtime that will not be saved */
 	TArray<FCustomizedToolMenu> RuntimeCustomizedMenus;
 
+	UPROPERTY(EditAnywhere, Category = Misc)
+	TMap<FName /*MenuName*/, FToolMenuProfileMap> MenuProfiles;
+
+	/** Transient profiles made during runtime that will not be saved */
+	TMap<FName /*MenuName*/, FToolMenuProfileMap> RuntimeMenuProfiles;
+
+
 	FSimpleDelegate SetTimerForNextTickDelegate;
 
 	bool bNextTickTimerIsSet;
diff --git a/Engine/Source/Editor/StaticMeshEditor/Private/StaticMeshEditor.cpp b/Engine/Source/Editor/StaticMeshEditor/Private/StaticMeshEditor.cpp
--- a/Engine/Source/Editor/StaticMeshEditor/Private/StaticMeshEditor.cpp
+++ b/Engine/Source/Editor/StaticMeshEditor/Private/StaticMeshEditor.cpp
@@ -510,12 +510,6 @@ void FStaticMeshEditor::ExtendMenu()
 			{
 				if (TSharedPtr<FStaticMeshEditor> StaticMeshEditor = StaticMeshEditor::GetStaticMeshEditorFromMenuContext(InSection))
 				{
-					// These entries modify the mesh, so we don't want them in read only mode
-					if(StaticMeshEditor->GetOpenMethod() != EAssetOpenMethod::Edit)
-					{
-						return;
-					}
-					
 					FToolMenuEntry& Entry = InSection.AddSubMenu("Collision",
 						LOCTEXT("StaticMeshEditorCollisionMenu", "Collision"),
 						LOCTEXT("StaticMeshEditorCollisionMenu_ToolTip", "Opens a menu with commands for editing this mesh's collision"),
@@ -711,7 +705,8 @@ void FStaticMeshEditor::BindCommands()
 	UICommandList->MapAction(
 		FGenericCommands::Get().Duplicate,
 		FExecuteAction::CreateSP(this, &FStaticMeshEditor::DuplicateSelected),
-		FCanExecuteAction::CreateSP(this, &FStaticMeshEditor::CanDuplicateSelected));
+		FCanExecuteAction::CreateSP(this, &FStaticMeshEditor::CanDuplicateSelected),
+		FIsActionChecked());
 
 	UICommandList->MapAction(
 		FGenericCommands::Get().Copy,
@@ -726,7 +721,8 @@ void FStaticMeshEditor::BindCommands()
 	UICommandList->MapAction(
 		FGenericCommands::Get().Rename,
 		FExecuteAction::CreateSP(this, &FStaticMeshEditor::RequestRenameSelectedSocket),
-		FCanExecuteAction::CreateSP(this, &FStaticMeshEditor::CanRenameSelected));
+		FCanExecuteAction::CreateSP(this, &FStaticMeshEditor::CanRenameSelected),
+		FIsActionChecked());
 
 	UICommandList->MapAction(
 		Commands.CreateDOP10X,
@@ -878,7 +874,11 @@ void FStaticMeshEditor::BindCommands()
 	UICommandList->MapAction(
 		Commands.BakeMaterials,
 		FExecuteAction::CreateSP(this, &FStaticMeshEditor::BakeMaterials),
-		FCanExecuteAction());
+		FCanExecuteAction::CreateLambda([this]()
+		{
+			return GetOpenMethod() != EAssetOpenMethod::View;
+		}),
+		FIsActionChecked());
 }
 
 void FStaticMeshEditor::ExtendToolBar()
@@ -900,12 +900,6 @@ void FStaticMeshEditor::ExtendToolBar()
 			{
 				if (TSharedPtr<FStaticMeshEditor> StaticMeshEditor = StaticMeshEditor::GetStaticMeshEditorFromMenuContext(InSection))
 				{
-					// These entries modify the mesh, so we don't want them in read only mode
-					if(StaticMeshEditor->GetOpenMethod() != EAssetOpenMethod::Edit)
-					{
-						return;
-					}
-					
 					auto ConstructReimportContextMenu = [](UToolMenu* InMenu)
 					{
 						FToolMenuSection& Section = InMenu->AddSection("Reimport");
@@ -940,12 +934,6 @@ void FStaticMeshEditor::ExtendToolBar()
 			{
 				if (TSharedPtr<FStaticMeshEditor> StaticMeshEditor = StaticMeshEditor::GetStaticMeshEditorFromMenuContext(InSection))
 				{
-					// These entries modify the mesh, so we don't want them in read only mode
-					if(StaticMeshEditor->GetOpenMethod() != EAssetOpenMethod::Edit)
-					{
-						return;
-					}
-					
 					FToolMenuEntry& CollisionEntry = InSection.AddEntry(FToolMenuEntry::InitComboButton(
 					"Collision",
 					FUIAction(),
@@ -1716,23 +1704,19 @@ void FStaticMeshEditor::GenerateUVChannelComboList(UToolMenu* InMenu)
 
 	if (TSharedPtr<FStaticMeshEditor> StaticMeshEditor = StaticMeshEditor::GetStaticMeshEditorFromMenuContext(InMenu))
 	{
-		// These actions modify the mesh, so we don't want them in read only mode
-		if(StaticMeshEditor->GetOpenMethod() == EAssetOpenMethod::Edit)
-		{
-			FToolMenuSection& Section = InMenu->AddSection("UVActionOptions");
+		FToolMenuSection& Section = InMenu->AddSection("UVActionOptions");
 
-			FUIAction MenuAction;
-			MenuAction.ExecuteAction.BindSP(this, &FStaticMeshEditor::RemoveCurrentUVChannel);
-			MenuAction.CanExecuteAction.BindSP(this, &FStaticMeshEditor::CanRemoveUVChannel);
-			Section.AddMenuEntry(
-				"Remove_UVChannel",
-				LOCTEXT("Remove_UVChannel", "Remove Selected"),
-				LOCTEXT("Remove_UVChannel_ToolTip", "Remove currently selected UV channel from the static mesh"),
-				FSlateIcon(),
-				MenuAction,
-				EUserInterfaceActionType::Button
-			);
-		}
+		FUIAction MenuAction;
+		MenuAction.ExecuteAction.BindSP(this, &FStaticMeshEditor::RemoveCurrentUVChannel);
+		MenuAction.CanExecuteAction.BindSP(this, &FStaticMeshEditor::CanRemoveUVChannel);
+		Section.AddMenuEntry(
+			"Remove_UVChannel",
+			LOCTEXT("Remove_UVChannel", "Remove Selected"),
+			LOCTEXT("Remove_UVChannel_ToolTip", "Remove currently selected UV channel from the static mesh"),
+			FSlateIcon(),
+			MenuAction,
+			EUserInterfaceActionType::Button
+		);
 		
 	}
 }
@@ -2383,7 +2367,7 @@ void FStaticMeshEditor::DeleteSelected()
 
 bool FStaticMeshEditor::CanDeleteSelected() const
 {
-	return (GetSelectedSocket() != NULL || HasSelectedPrims());
+	return GetOpenMethod() != EAssetOpenMethod::View && (GetSelectedSocket() != NULL || HasSelectedPrims());
 }
 
 void FStaticMeshEditor::DeleteSelectedSockets()
@@ -2472,7 +2456,7 @@ void FStaticMeshEditor::DuplicateSelected()
 
 bool FStaticMeshEditor::CanDuplicateSelected() const
 {
-	return (GetSelectedSocket() != NULL || HasSelectedPrims());
+	return GetOpenMethod() != EAssetOpenMethod::View && (GetSelectedSocket() != NULL || HasSelectedPrims());
 }
 
 void FStaticMeshEditor::CopySelected()
@@ -2492,6 +2476,11 @@ void FStaticMeshEditor::PasteCopied()
 
 bool FStaticMeshEditor::CanPasteCopied() const
 {
+	if(GetOpenMethod() == EAssetOpenMethod::View)
+	{
+		return false;
+	}
+	
 	FString TextToImport;
 	FPlatformApplicationMisc::ClipboardPaste(TextToImport);
 	FBodySetupObjectTextFactory Factory;
@@ -2500,7 +2489,7 @@ bool FStaticMeshEditor::CanPasteCopied() const
 
 bool FStaticMeshEditor::CanRenameSelected() const
 {
-	return (GetSelectedSocket() != NULL);
+	return GetOpenMethod() != EAssetOpenMethod::View && (GetSelectedSocket() != NULL);
 }
 
 void FStaticMeshEditor::ExecuteFindInExplorer()
@@ -2606,6 +2595,25 @@ bool FStaticMeshEditor::OnRequestClose(EAssetEditorCloseReason InCloseReason)
 	return bAllowClose;
 }
 
+void FStaticMeshEditor::SetupReadOnlyMenuProfiles(FReadOnlyAssetEditorCustomization& OutReadOnlyCustomization)
+{
+	FName ReadOnlyOwnerName("StaticMeshEditorReadOnly");
+
+	// The combo button to show UVs is fine to be available in read only mode
+	OutReadOnlyCustomization.ToolbarPermissionList.AddAllowListItem(ReadOnlyOwnerName, "UVToolbar");
+
+	// Hide the command to bake materials in the "Asset" menu in read only mode
+	FNamePermissionList& AssetMenuPermissionList = OutReadOnlyCustomization.MainMenuSubmenuPermissionLists.FindOrAdd("Asset");
+	AssetMenuPermissionList.AddDenyListItem(ReadOnlyOwnerName, FStaticMeshEditorCommands::Get().BakeMaterials->GetCommandName());
+
+	// Hide the command to edit sockets in the "Edit" menu in read only mode
+	FNamePermissionList& EditMenuPermissionList = OutReadOnlyCustomization.MainMenuSubmenuPermissionLists.FindOrAdd("Edit");
+	EditMenuPermissionList.AddDenyListItem(ReadOnlyOwnerName, "DeleteSocket");
+	EditMenuPermissionList.AddDenyListItem(ReadOnlyOwnerName, "DuplicateSocket");
+
+
+}
+
 void FStaticMeshEditor::RegisterOnPostUndo( const FOnPostUndo& Delegate )
 {
 	OnPostUndo.Add( Delegate );
diff --git a/Engine/Source/Editor/StaticMeshEditor/Private/StaticMeshEditor.h b/Engine/Source/Editor/StaticMeshEditor/Private/StaticMeshEditor.h
--- a/Engine/Source/Editor/StaticMeshEditor/Private/StaticMeshEditor.h
+++ b/Engine/Source/Editor/StaticMeshEditor/Private/StaticMeshEditor.h
@@ -474,6 +474,8 @@ private:
 
 	//~ Begin FAssetEditorToolkit Interface.
 	virtual bool OnRequestClose(EAssetEditorCloseReason InCloseReason) override;
+	virtual void SetupReadOnlyMenuProfiles(FReadOnlyAssetEditorCustomization& OutReadOnlyCustomization) override;
+
 	//~ End FAssetEditorToolkit Interface.
 
 	//~ Begin FEditorUndoClient Interface
diff --git a/Engine/Source/Editor/UnrealEd/Private/Toolkits/AssetEditorToolkit.cpp b/Engine/Source/Editor/UnrealEd/Private/Toolkits/AssetEditorToolkit.cpp
--- a/Engine/Source/Editor/UnrealEd/Private/Toolkits/AssetEditorToolkit.cpp
+++ b/Engine/Source/Editor/UnrealEd/Private/Toolkits/AssetEditorToolkit.cpp
@@ -34,6 +34,7 @@
 #include "Misc/Attribute.h"
 #include "Textures/SlateIcon.h"
 #include "WidgetDrawerConfig.h"
+#include "Framework/Commands/GenericCommands.h"
 #include "Interfaces/Interface_AsyncCompilation.h"
 #include "Widgets/Images/SImage.h"
 
@@ -46,6 +47,7 @@ TSharedPtr<FExtensibilityManager> FAssetEditorToolkit::SharedMenuExtensibilityMa
 TSharedPtr<FExtensibilityManager> FAssetEditorToolkit::SharedToolBarExtensibilityManager;
 
 const FName FAssetEditorToolkit::DefaultAssetEditorToolBarName("AssetEditor.DefaultToolBar");
+const FName FAssetEditorToolkit::ReadOnlyMenuProfileName("AssetEditor.ReadOnlyMenuProfile");
 
 FAssetEditorToolkit::FAssetEditorToolkit()
 	: GCEditingObjects(*this)
@@ -293,9 +295,11 @@ void FAssetEditorToolkit::InitAssetEditor( const EToolkitMode::Type Mode, const
 		}
 	}
 
+	InitializeReadOnlyMenuProfiles();
+
 	// Give a chance to customize tab manager and other UI before widgets are created
 	GEditor->GetEditorSubsystem<UAssetEditorSubsystem>()->NotifyEditorOpeningPreWidgets(ObjectPtrDecay(EditingObjects), this);
-
+	
 	// Create menus
 	if (ToolkitMode == EToolkitMode::Standalone)
 	{
@@ -1057,14 +1061,10 @@ void FAssetEditorToolkit::FillDefaultFileMenuCommands(FToolMenuSection& InSectio
 
 	if (UAssetEditorToolkitMenuContext* Context = InSection.FindContext<UAssetEditorToolkitMenuContext>())
 	{
-		// The save buttons are hidden if we are in read only mode
-		if(Context->Toolkit.Pin()->GetOpenMethod() == EAssetOpenMethod::Edit)
+		InSection.AddMenuEntry(FAssetEditorCommonCommands::Get().SaveAsset, TAttribute<FText>(), TAttribute<FText>(), FSlateIcon(FAppStyle::GetAppStyleSetName(), "AssetEditor.SaveAsset")).InsertPosition = InsertPosition;
+		if( IsActuallyAnAsset() )
 		{
-			InSection.AddMenuEntry(FAssetEditorCommonCommands::Get().SaveAsset, TAttribute<FText>(), TAttribute<FText>(), FSlateIcon(FAppStyle::GetAppStyleSetName(), "AssetEditor.SaveAsset")).InsertPosition = InsertPosition;
-			if( IsActuallyAnAsset() )
-			{
-				InSection.AddMenuEntry(FAssetEditorCommonCommands::Get().SaveAssetAs, TAttribute<FText>(), TAttribute<FText>(), FSlateIcon(FAppStyle::GetAppStyleSetName(), "AssetEditor.SaveAssetAs")).InsertPosition = InsertPosition;
-			}
+			InSection.AddMenuEntry(FAssetEditorCommonCommands::Get().SaveAssetAs, TAttribute<FText>(), TAttribute<FText>(), FSlateIcon(FAppStyle::GetAppStyleSetName(), "AssetEditor.SaveAssetAs")).InsertPosition = InsertPosition;
 		}
 	}
 
@@ -1174,8 +1174,82 @@ void FAssetEditorToolkit::RegisterDefaultToolBar()
 	}
 }
 
+void FAssetEditorToolkit::InitializeReadOnlyMenuProfiles()
+{
+	// Toolbar Customizations
+
+	// Only allow the "Find in Content Browser" toolbar item by default - which hides "Save" since we are using an allowlist
+	ReadOnlyCustomization.ToolbarPermissionList.AddAllowListItem(ReadOnlyMenuProfileName, FGlobalEditorCommonCommands::Get().FindInContentBrowser->GetCommandName());
+
+	// Main Menu Customizations
+
+	// The default menus we provide in the main menu
+	TArray<FName> MainMenuSubmenus({"File", "Edit", "Asset", "Window", "Tools", "Help"});
+
+	// Only allow the default menus we provide
+	for(const FName& Submenu : MainMenuSubmenus)
+	{
+		ReadOnlyCustomization.MainMenuPermissionList.AddAllowListItem(ReadOnlyMenuProfileName, Submenu);
+		ReadOnlyCustomization.MainMenuSubmenuPermissionLists.Add(Submenu, FNamePermissionList());
+	}
+
+	// Hide the save commands in the "File" menu in read only mode
+	FNamePermissionList& FileMenuPermissionList = ReadOnlyCustomization.MainMenuSubmenuPermissionLists.FindOrAdd("File");
+
+	FileMenuPermissionList.AddDenyListItem(ReadOnlyMenuProfileName, FAssetEditorCommonCommands::Get().SaveAsset->GetCommandName());
+	FileMenuPermissionList.AddDenyListItem(ReadOnlyMenuProfileName, FAssetEditorCommonCommands::Get().SaveAssetAs->GetCommandName());
+
+	// Hide the re-import command in the "asset" menu in read only mode
+	// There is a reimport command per editing object, so we make sure to hide them all. See FAssetEditorToolkit::FillDefaultAssetMenuCommands
+	FNamePermissionList& AssetMenuPermissionList = ReadOnlyCustomization.MainMenuSubmenuPermissionLists.FindOrAdd("Asset");
+	
+	FName ReimportEntryName = TEXT("Reimport");
+	int32 MenuEntryCount = 0;
+
+	for( auto ObjectIter = EditingObjects.CreateConstIterator(); ObjectIter; ++ObjectIter )
+	{
+		ReimportEntryName.SetNumber(MenuEntryCount++);
+		AssetMenuPermissionList.AddDenyListItem(ReadOnlyMenuProfileName, ReimportEntryName);
+	}
+	
+	// Give specific asset editors a chance to customize the default behavior (e.g show specific menus they want to allow in read only mode)
+	SetupReadOnlyMenuProfiles(ReadOnlyCustomization);
+
+	// Create the menu profile and apply the permission lists
+	FToolMenuProfile* MainMenuProfile = UToolMenus::Get()->AddRuntimeMenuProfile(GetToolMenuName(), ReadOnlyMenuProfileName);
+	FToolMenuProfile* ToolbarProfile = UToolMenus::Get()->AddRuntimeMenuProfile(GetToolMenuToolbarName(), ReadOnlyMenuProfileName);
+	FToolMenuProfile* CommonActionsToolbarProfile = UToolMenus::Get()->AddRuntimeMenuProfile("AssetEditorToolbar.CommonActions", ReadOnlyMenuProfileName);
+
+	MainMenuProfile->MenuPermissions = ReadOnlyCustomization.MainMenuPermissionList;
+	ToolbarProfile->MenuPermissions = ReadOnlyCustomization.ToolbarPermissionList;
+	CommonActionsToolbarProfile->MenuPermissions = ReadOnlyCustomization.ToolbarPermissionList;
+
+	for(const FName& Submenu : MainMenuSubmenus)
+	{
+		const FName SubmenuName = *(GetToolMenuName().ToString() + TEXT(".") + Submenu.ToString());
+		
+		FToolMenuProfile* SubmenuProfile = UToolMenus::Get()->AddRuntimeMenuProfile(SubmenuName, ReadOnlyMenuProfileName);
+		SubmenuProfile->MenuPermissions = ReadOnlyCustomization.MainMenuSubmenuPermissionLists[Submenu];
+	}
+
+}
+
 void FAssetEditorToolkit::InitToolMenuContext(FToolMenuContext& MenuContext)
 {
+	UAssetEditorToolkitMenuContext* ToolkitMenuContext = MenuContext.FindContext<UAssetEditorToolkitMenuContext>();
+	
+	if(!ToolkitMenuContext || !ToolkitMenuContext->Toolkit.IsValid())
+	{
+		return;
+	}
+
+	// If we are in read only mode, set the read only menu profile as active
+	if(ToolkitMenuContext->Toolkit.Pin()->GetOpenMethod() == EAssetOpenMethod::View)
+	{
+		UToolMenuProfileContext* ProfileContext = NewObject<UToolMenuProfileContext>();
+		ProfileContext->ActiveProfiles.Add(ReadOnlyMenuProfileName);
+		MenuContext.AddObject(ProfileContext);
+	}
 }
 
 UToolMenu* FAssetEditorToolkit::GenerateCommonActionsToolbar(FToolMenuContext& MenuContext)
@@ -1200,11 +1274,7 @@ UToolMenu* FAssetEditorToolkit::GenerateCommonActionsToolbar(FToolMenuContext& M
 			{
 				if(TSharedPtr<FAssetEditorToolkit> AssetEditorToolkit = AssetEditorToolkitMenuContext->Toolkit.Pin())
 				{
-					// The save button is hidden if we are in read only mode
-					if(AssetEditorToolkit->GetOpenMethod() == EAssetOpenMethod::Edit)
-					{
-						InSection.AddEntry(FToolMenuEntry::InitToolBarButton(FAssetEditorCommonCommands::Get().SaveAsset));
-					}
+					InSection.AddEntry(FToolMenuEntry::InitToolBarButton(FAssetEditorCommonCommands::Get().SaveAsset));
 				}
 			}
 		}));
diff --git a/Engine/Source/Editor/UnrealEd/Public/Toolkits/AssetEditorToolkit.h b/Engine/Source/Editor/UnrealEd/Public/Toolkits/AssetEditorToolkit.h
--- a/Engine/Source/Editor/UnrealEd/Public/Toolkits/AssetEditorToolkit.h
+++ b/Engine/Source/Editor/UnrealEd/Public/Toolkits/AssetEditorToolkit.h
@@ -87,6 +87,22 @@ public:
 	virtual TSharedPtr<FExtensibilityManager> GetToolBarExtensibilityManager() = 0;
 };
 
+/**
+ * Allows an asset editor to specify which custom menus/toolbar items are visible in read-only mode, with
+ * the default behavior being that all entries specific to an asset editor are hidden
+ */
+struct FReadOnlyAssetEditorCustomization
+{
+	// Permission list for the main menu
+	FNamePermissionList MainMenuPermissionList;
+
+	// Permission list for the toolbar
+	FNamePermissionList ToolbarPermissionList;
+
+	// Permission list for the various submenus on the main menu (e.g File, Help..)
+	TMap<FName, FNamePermissionList> MainMenuSubmenuPermissionLists;
+};
+
 /**
  * Base class for toolkits that are used for asset editing (abstract)
  */
@@ -403,6 +419,12 @@ protected:
 
 	UNREALED_API virtual void CreateEditorModeManager() override;
 
+	/** Specify the permission lists to use in read only mode */
+	UNREALED_API virtual void SetupReadOnlyMenuProfiles(FReadOnlyAssetEditorCustomization& OutReadOnlyCustomization) { }
+
+	/** Get the name of the profile registered with UToolMenus for read only menu customizations */
+	UNREALED_API virtual FName GetReadOnlyMenuProfileName() { return ReadOnlyMenuProfileName; }
+
 private:
 	// Callback for persisting the Asset Editor's layout.
 	void HandleTabManagerPersistLayout( const TSharedRef<FTabManager::FLayout>& LayoutToSave )
@@ -413,6 +435,9 @@ private:
 		}
 	}
 
+	/** Initialize the base level customizations for read only mode and register the profile */
+	void InitializeReadOnlyMenuProfiles();
+
 private:
 	/**
 	 * Report the references of the EditingObjects to the GC.  The level of indirection necessary 
@@ -468,6 +493,9 @@ private:
 	static UNREALED_API TSharedPtr<FExtensibilityManager> SharedMenuExtensibilityManager;
 	static UNREALED_API TSharedPtr<FExtensibilityManager> SharedToolBarExtensibilityManager;
 
+	/** The name of the profile registered with UToolMenus for read only menu customizations */
+	static UNREALED_API const FName ReadOnlyMenuProfileName;
+
 	/** The object we're currently editing */
 	// @todo toolkit minor: Currently we don't need to serialize this object reference because the AssetEditorSubsystem is kept in sync (and will always serialize it.)
 	TArray<TObjectPtr<UObject>> EditingObjects;
@@ -492,4 +520,7 @@ private:
 
 	/** Whether the asset editor was opened in edit mode or read only mode */
 	EAssetOpenMethod OpenMethod;
+
+	/** Determines menu/toolbar customization in read only mode */
+	FReadOnlyAssetEditorCustomization ReadOnlyCustomization;
 };
diff --git a/Engine/Source/Runtime/Slate/Private/Framework/MultiBox/ToolMenuCustomization.cpp b/Engine/Source/Runtime/Slate/Private/Framework/MultiBox/ToolMenuCustomization.cpp
--- a/Engine/Source/Runtime/Slate/Private/Framework/MultiBox/ToolMenuCustomization.cpp
+++ b/Engine/Source/Runtime/Slate/Private/Framework/MultiBox/ToolMenuCustomization.cpp
@@ -2,37 +2,87 @@
 
 #include "Framework/MultiBox/ToolMenuBase.h"
 
-FName FCustomizedToolMenu::GetEntrySectionName(const FName InEntryName) const
+namespace UE::ToolMenuCustomization::Private
 {
-	for (auto& It : EntryOrder)
+	static void HandleToolMenuProfile(FToolMenuProfile& Result, const FToolMenuProfile* Current)
 	{
-		if (It.Value.Names.Contains(InEntryName))
+		if (!Current)
 		{
-			return It.Key;
+			return;
+		}
+
+		if (Current->IsSuppressExtenders() && !Result.IsSuppressExtenders())
+		{
+			Result.SuppressExtenders = Current->SuppressExtenders;
+		}
+
+		for (const auto& EntryIterator : Current->Entries)
+		{
+			if (FCustomizedToolMenuEntry* ExistingEntry = Result.FindEntry(EntryIterator.Key))
+			{
+				if (EntryIterator.Value.Visibility != ECustomizedToolMenuVisibility::None)
+				{
+					ExistingEntry->Visibility = EntryIterator.Value.Visibility;
+				}
+			}
+			else
+			{
+				Result.Entries.Add(EntryIterator.Key, EntryIterator.Value);
+			}
+		}
+
+		for (const auto& SectionIterator : Current->Sections)
+		{
+			if (FCustomizedToolMenuSection* ExistingSection = Result.FindSection(SectionIterator.Key))
+			{
+				if (SectionIterator.Value.Visibility != ECustomizedToolMenuVisibility::None)
+				{
+					ExistingSection->Visibility = SectionIterator.Value.Visibility;
+				}
+			}
+			else
+			{
+				Result.Sections.Add(SectionIterator.Key, SectionIterator.Value);
+			}
 		}
+
+		Result.MenuPermissions.Append(Current->MenuPermissions);
+
 	}
+	
+	static void HandleToolMenuCustomization(FCustomizedToolMenu& Result, const FCustomizedToolMenu* Current)
+	{
+		if (Current->SectionOrder.Num() > 0)
+		{
+			Result.SectionOrder = Current->SectionOrder;
+		}
 
-	return NAME_None;
+		for (const auto& SectionEntryOrderIterator : Current->EntryOrder)
+		{
+			Result.EntryOrder.Add(SectionEntryOrderIterator.Key, SectionEntryOrderIterator.Value);
+		}
+	}
+	
 }
 
-FCustomizedToolMenuEntry* FCustomizedToolMenu::FindEntry(const FName InEntryName)
+FCustomizedToolMenuEntry* FToolMenuProfile::FindEntry(const FName InEntryName)
 {
 	return Entries.Find(InEntryName);
 	//return Entries.FindByPredicate([=](const FCustomizedToolMenuEntry& Entry) { return Entry.Name == InEntryName; });
 }
 
-const FCustomizedToolMenuEntry* FCustomizedToolMenu::FindEntry(const FName InEntryName) const
+const FCustomizedToolMenuEntry* FToolMenuProfile::FindEntry(const FName InEntryName) const
 {
 	return Entries.Find(InEntryName);
 	//return Entries.FindByPredicate([=](const FCustomizedToolMenuEntry& Entry) { return Entry.Name == InEntryName; });
 }
 
-FCustomizedToolMenuEntry* FCustomizedToolMenu::AddEntry(const FName InEntryName)
+FCustomizedToolMenuEntry* FToolMenuProfile::AddEntry(const FName InEntryName)
 {
 	return &Entries.FindOrAdd(InEntryName);
 }
 
-ECustomizedToolMenuVisibility FCustomizedToolMenu::GetEntryVisiblity(const FName InEntryName) const
+ECustomizedToolMenuVisibility FToolMenuProfile::GetEntryVisiblity(const FName InEntryName) const
 {
 	if (const FCustomizedToolMenuEntry* Found = FindEntry(InEntryName))
 	{
@@ -42,29 +92,29 @@ ECustomizedToolMenuVisibility FCustomizedToolMenu::GetEntryVisiblity(const FName
 	return ECustomizedToolMenuVisibility::None;
 }
 
-bool FCustomizedToolMenu::IsEntryHidden(const FName InEntryName) const
+bool FToolMenuProfile::IsEntryHidden(const FName InEntryName) const
 {
 	return GetEntryVisiblity(InEntryName) == ECustomizedToolMenuVisibility::Hidden;
 }
 
-FCustomizedToolMenuSection* FCustomizedToolMenu::FindSection(const FName InSectionName)
+FCustomizedToolMenuSection* FToolMenuProfile::FindSection(const FName InSectionName)
 {
 	return Sections.Find(InSectionName);
 	//return Sections.FindByPredicate([=](const FCustomizedToolMenuSection& Section) { return Section.Name == InSectionName; });
 }
 
-const FCustomizedToolMenuSection* FCustomizedToolMenu::FindSection(const FName InSectionName) const
+const FCustomizedToolMenuSection* FToolMenuProfile::FindSection(const FName InSectionName) const
 {
 	return Sections.Find(InSectionName);
 	//return Sections.FindByPredicate([=](const FCustomizedToolMenuSection& Section) { return Section.Name == InSectionName; });
 }
 
-FCustomizedToolMenuSection* FCustomizedToolMenu::AddSection(const FName InSectionName)
+FCustomizedToolMenuSection* FToolMenuProfile::AddSection(const FName InSectionName)
 {
 	return &Sections.FindOrAdd(InSectionName);
 }
 
-ECustomizedToolMenuVisibility FCustomizedToolMenu::GetSectionVisiblity(const FName InSectionName) const
+ECustomizedToolMenuVisibility FToolMenuProfile::GetSectionVisiblity(const FName InSectionName) const
 {
 	if (const FCustomizedToolMenuSection* Found = FindSection(InSectionName))
 	{
@@ -74,12 +124,12 @@ ECustomizedToolMenuVisibility FCustomizedToolMenu::GetSectionVisiblity(const FNa
 	return ECustomizedToolMenuVisibility::None;
 }
 
-bool FCustomizedToolMenu::IsSectionHidden(const FName InEntryName) const
+bool FToolMenuProfile::IsSectionHidden(const FName InEntryName) const
 {
 	return GetSectionVisiblity(InEntryName) == ECustomizedToolMenuVisibility::Hidden;
 }
 
-void FCustomizedToolMenu::SetSuppressExtenders(const FName InOwnerName, const bool bInSuppress)
+void FToolMenuProfile::SetSuppressExtenders(const FName InOwnerName, const bool bInSuppress)
 {
 	if (bInSuppress)
 	{
@@ -91,11 +141,43 @@ void FCustomizedToolMenu::SetSuppressExtenders(const FName InOwnerName, const bo
 	}
 }
 
-bool FCustomizedToolMenu::IsSuppressExtenders() const
+bool FToolMenuProfile::IsSuppressExtenders() const
 {
 	return !SuppressExtenders.IsEmpty();
 }
 
+FName FCustomizedToolMenu::GetEntrySectionName(const FName InEntryName) const
+{
+	for (auto& It : EntryOrder)
+	{
+		if (It.Value.Names.Contains(InEntryName))
+		{
+			return It.Key;
+		}
+	}
+
+	return NAME_None;
+}
+
+FToolMenuProfile FToolMenuProfileHierarchy::GenerateFlattenedMenuProfile() const
+{
+	// Process parents first then children
+	// Each customization has chance to override what has already been customized before it
+	FToolMenuProfile Destination;
+
+	for (const FToolMenuProfile* Current : ProfileHierarchy)
+	{
+		UE::ToolMenuCustomization::Private::HandleToolMenuProfile(Destination, Current);
+	}
+
+	for (const FToolMenuProfile* Current : RuntimeProfileHierarchy)
+	{
+		UE::ToolMenuCustomization::Private::HandleToolMenuProfile(Destination, Current);
+	}
+
+	return Destination;
+}
+
 FName FCustomizedToolMenuHierarchy::GetEntrySectionName(const FName InEntryName) const
 {
 	for (int32 i = Hierarchy.Num() - 1; i >= 0; --i)
@@ -238,13 +320,16 @@ FCustomizedToolMenu FCustomizedToolMenuHierarchy::GenerateFlattened() const
 
 	for (const FCustomizedToolMenu* Current : Hierarchy)
 	{
-		HandleCustomizedToolMenu(Destination, Current);
+		UE::ToolMenuCustomization::Private::HandleToolMenuProfile(Destination, Current);
+		UE::ToolMenuCustomization::Private::HandleToolMenuCustomization(Destination, Current);
 	}
 
 	for (const FCustomizedToolMenu* Current : RuntimeHierarchy)
 	{
-		HandleCustomizedToolMenu(Destination, Current);
+		UE::ToolMenuCustomization::Private::HandleToolMenuProfile(Destination, Current);
+		UE::ToolMenuCustomization::Private::HandleToolMenuCustomization(Destination, Current);
 	}
 
 	return Destination;
 }
+
diff --git a/Engine/Source/Runtime/Slate/Public/Framework/MultiBox/ToolMenuBase.h b/Engine/Source/Runtime/Slate/Public/Framework/MultiBox/ToolMenuBase.h
--- a/Engine/Source/Runtime/Slate/Public/Framework/MultiBox/ToolMenuBase.h
+++ b/Engine/Source/Runtime/Slate/Public/Framework/MultiBox/ToolMenuBase.h
@@ -54,8 +54,12 @@ struct FCustomizedToolMenuNameArray
 	TArray<FName> Names;
 };
 
+/*
+ * A menu profile is a way for systems to modify instances of a menu by showing/hiding specific items. You can have multiple profiles active on
+ * a single menu at the same time.
+ */ 
 USTRUCT()
-struct FCustomizedToolMenu
+struct FToolMenuProfile
 {
 	GENERATED_BODY()
 
@@ -68,12 +72,6 @@ struct FCustomizedToolMenu
 	UPROPERTY()
 	TMap<FName, FCustomizedToolMenuSection> Sections;
 
-	UPROPERTY()
-	TMap<FName, FCustomizedToolMenuNameArray> EntryOrder;
-
-	UPROPERTY()
-	TArray<FName> SectionOrder;
-
 	UPROPERTY()
 	TArray<FName> SuppressExtenders;
 
@@ -84,7 +82,6 @@ struct FCustomizedToolMenu
 	SLATE_API FCustomizedToolMenuEntry* AddEntry(const FName InEntryName);
 	SLATE_API ECustomizedToolMenuVisibility GetEntryVisiblity(const FName InSectionName) const;
 	SLATE_API bool IsEntryHidden(const FName InEntryName) const;
-	SLATE_API FName GetEntrySectionName(const FName InEntryName) const;
 
 	SLATE_API FCustomizedToolMenuSection* FindSection(const FName InSectionName);
 	SLATE_API const FCustomizedToolMenuSection* FindSection(const FName InSectionName) const;
@@ -94,17 +91,47 @@ struct FCustomizedToolMenu
 
 	SLATE_API void SetSuppressExtenders(const FName InOwnerName, const bool bInSuppress);
 	SLATE_API bool IsSuppressExtenders() const;
+
 };
 
-struct FCustomizedToolMenuHierarchy
+/*
+ * A menu customization is a specialization of menu profiles - that allows for advanced behavior such as modifying the order of sections/entries
+ * A menu can only have one customization active at a time
+ */
+USTRUCT()
+struct FCustomizedToolMenu : public FToolMenuProfile
 {
+	GENERATED_BODY()
+
 	SLATE_API FName GetEntrySectionName(const FName InEntryName) const;
-	SLATE_API bool IsEntryHidden(const FName InEntryName) const;
+	
+	UPROPERTY()
+	TMap<FName, FCustomizedToolMenuNameArray> EntryOrder;
 
-	SLATE_API bool IsSectionHidden(const FName InSectionName) const;
+	UPROPERTY()
+	TArray<FName> SectionOrder;
+};
 
-	SLATE_API bool IsSuppressExtenders() const;
+/*
+ * Structure to describe a menu profile for the whole hierarchy of a menu
+ */
+struct FToolMenuProfileHierarchy
+{
+	SLATE_API FToolMenuProfile GenerateFlattenedMenuProfile() const;
+
+	TArray<const FToolMenuProfile*> ProfileHierarchy;
+	TArray<const FToolMenuProfile*> RuntimeProfileHierarchy;
+}; 
 
+/*
+ * Structure to describe the menu customization for the whole hierarchy of a menu
+ */
+struct FCustomizedToolMenuHierarchy : public FToolMenuProfileHierarchy
+{
+	SLATE_API FName GetEntrySectionName(const FName InEntryName) const;
+	SLATE_API bool IsEntryHidden(const FName InEntryName) const;
+	SLATE_API bool IsSectionHidden(const FName InSectionName) const;
+	SLATE_API bool IsSuppressExtenders() const;
 	SLATE_API FCustomizedToolMenu GenerateFlattened() const;
 
 	TArray<const FCustomizedToolMenu*> Hierarchy;
@@ -127,6 +154,11 @@ public:
 	virtual FCustomizedToolMenu* FindMenuCustomization() const { return nullptr; }
 	virtual FCustomizedToolMenu* AddMenuCustomization() const { return nullptr; }
 	virtual FCustomizedToolMenuHierarchy GetMenuCustomizationHierarchy() const { return FCustomizedToolMenuHierarchy(); }
+
+	virtual FToolMenuProfile* FindMenuProfile(const FName& ProfileName) const { return nullptr; }
+	virtual FToolMenuProfile* AddMenuProfile(const FName& ProfileName) const { return nullptr; }
+	virtual FToolMenuProfileHierarchy GetMenuProfileHierarchy(const FName& ProfileName) const { return FToolMenuProfileHierarchy(); }
+
 	virtual void UpdateMenuCustomizationFromMultibox(const TSharedRef<const FMultiBox>& InMultiBox) {}
 	virtual void OnMenuDestroyed() {}
 };
