Prevent entering debug mode when it would create an infinit loop. When the application is ticking we need to complete the tick before entering the debugging mode. When there is modal open that is stealing the focus, we don't have inputs on the other windows.

diff --git a/Engine/Source/Editor/Blutility/Private/EditorUtilitySubsystem.cpp b/Engine/Source/Editor/Blutility/Private/EditorUtilitySubsystem.cpp
--- a/Engine/Source/Editor/Blutility/Private/EditorUtilitySubsystem.cpp
+++ b/Engine/Source/Editor/Blutility/Private/EditorUtilitySubsystem.cpp
@@ -13,6 +13,7 @@
 #include "EditorUtilityWidgetBlueprint.h"
 #include "Engine/Blueprint.h"
 #include "Engine/BlueprintCore.h"
+#include "Framework/Application/SlateApplication.h"
 #include "Framework/Docking/TabManager.h"
 #include "GameFramework/Actor.h"
 #include "HAL/IConsoleManager.h"
@@ -348,6 +349,7 @@ bool UEditorUtilitySubsystem::SpawnRegisteredTabByID(FName NewTabID)
 			if (LevelEditorTabManager->HasTabSpawner(NewTabID))
 			{
 				TSharedPtr<SDockTab> NewDockTab = LevelEditorTabManager->TryInvokeTab(NewTabID);
+				NewDockTab->SetEnabled(FSlateApplication::Get().GetNormalExecutionAttribute());
 				IBlutilityModule* BlutilityModule = FModuleManager::GetModulePtr<IBlutilityModule>("Blutility");
 				UEditorUtilityWidgetBlueprint** WidgetToSpawn = RegisteredTabs.Find(NewTabID);
 				if (WidgetToSpawn)
@@ -666,7 +668,8 @@ void UEditorUtilitySubsystem::HandleOnEndPIE(const bool bIsSimulating)
 
 TSharedRef<SDockTab> UEditorUtilitySubsystem::SpawnEditorUITabFromGeneratedClass(const FSpawnTabArgs& SpawnTabArgs, UWidgetBlueprintGeneratedClass* InGeneratedWidgetBlueprint)
 {
-	TSharedRef<SDockTab> SpawnedTab = SNew(SDockTab);
+	TSharedRef<SDockTab> SpawnedTab = SNew(SDockTab)
+		.IsEnabled(FSlateApplication::Get().GetNormalExecutionAttribute());
 	UEditorUtilityWidget* CreatedUMGWidget = nullptr;
 
 	auto CreateUtilityWidgetFromGeneratedClass = [&CreatedUMGWidget](UWidgetBlueprintGeneratedClass* InGeneratedWidgetBlueprint)
diff --git a/Engine/Source/Editor/TurnkeySupport/Private/TurnkeySupportModule.cpp b/Engine/Source/Editor/TurnkeySupport/Private/TurnkeySupportModule.cpp
--- a/Engine/Source/Editor/TurnkeySupport/Private/TurnkeySupportModule.cpp
+++ b/Engine/Source/Editor/TurnkeySupport/Private/TurnkeySupportModule.cpp
@@ -2083,6 +2083,10 @@ void FTurnkeySupportModule::MakeTurnkeyMenu(FToolMenuSection& MenuSection) const
 
 	// hide during PIE
 	FUIAction PlatformMenuShownDelegate;
+	if (FSlateApplication::IsInitialized())
+	{
+		PlatformMenuShownDelegate.CanExecuteAction = FCanExecuteAction::CreateRaw(&FSlateApplication::Get(), &FSlateApplication::IsNormalExecution);
+	}
 	PlatformMenuShownDelegate.IsActionVisibleDelegate = FIsActionButtonVisible::CreateLambda([]() 
 	{
 		return !FTurnkeyEditorSupport::IsPIERunning();
diff --git a/Engine/Source/Runtime/Slate/Private/Framework/Application/SlateApplication.cpp b/Engine/Source/Runtime/Slate/Private/Framework/Application/SlateApplication.cpp
--- a/Engine/Source/Runtime/Slate/Private/Framework/Application/SlateApplication.cpp
+++ b/Engine/Source/Runtime/Slate/Private/Framework/Application/SlateApplication.cpp
@@ -1508,6 +1508,10 @@ void FSlateApplication::Tick(ESlateTickType TickType)
 
 	FScopeLock SlateTickAccess(&SlateTickCriticalSection);
 
+#if WITH_EDITOR
+	FScopedPreventDebuggingMode SlatePreventDebugginModeWhileTicking(NSLOCTEXT("EnterDebuggingMode", "WindowTicking", "The window is ticking."));
+#endif
+
 	SCOPED_NAMED_EVENT_F(TEXT("Slate::Tick (%s)"), FColor::Magenta, LexToString(TickType));
 	CSV_SCOPED_TIMING_STAT(Slate, Tick);
 
@@ -3630,12 +3634,49 @@ void FSlateApplication::GetAllVisibleChildWindows(TArray< TSharedRef<SWindow> >&
 
 void FSlateApplication::EnterDebuggingMode()
 {
+	if (!IsInGameThread())
+	{
+		ensureMsgf(false, TEXT("Can only enter Debugging Mode while on the game thread."));
+		return;
+	}
+
+	auto AddNotification = [Self=this](const FText& SubText)
+	{
+		if (TSharedPtr<SNotificationItem> MessagePinned = Self->DebuggingModeNotificationMessage.Pin())
+		{
+			static float DefaultDuration = FNotificationInfo{FText::GetEmpty()}.ExpireDuration;
+			MessagePinned->SetSubText(SubText);
+		}
+		else
+		{
+			FNotificationInfo Info(NSLOCTEXT("EnterDebuggingMode", "FailTitle", "Debugging Mode Fail"));
+			Info.SubText = SubText;
+			Self->DebuggingModeNotificationMessage = FSlateNotificationManager::Get().AddNotification(Info);
+		}
+		UE_LOG(LogSlate, Warning, TEXT("Enter Debugging Mode failed."));
+
+		static bool bDoDebugBreak = true;
+		if (bDoDebugBreak)
+		{
+			UE_DEBUG_BREAK();
+		}
+	};
+
 	if (GetActiveModalWindow().IsValid())
 	{
-		UE_LOG(LogSlate, Warning, TEXT("EnterDebuggingMode is not supported while a modal window is open."));
+		AddNotification(NSLOCTEXT("EnterDebuggingMode", "Fail_ModalWindow", "A modal window is open."));
 		return;
 	}
 
+#if WITH_EDITOR
+	if (PreventDebuggingModeStack.Num() > 0)
+	{
+		AddNotification(PreventDebuggingModeStack.Last().Key);
+		return;
+	}
+	FScopedPreventDebuggingMode Scope(NSLOCTEXT("EnterDebuggingMode", "AlreadyInDebuggingMode", "Already in debug mode."));
+#endif
+
 	bRequestLeaveDebugMode = false;
 
 	// Note it is ok to hold a reference here as the game viewport should not be destroyed while in debugging mode
@@ -3741,6 +3782,24 @@ void FSlateApplication::LeaveDebuggingMode( bool bLeavingForSingleStep )
 	bLeaveDebugForSingleStep = bLeavingForSingleStep;
 }
 
+#if WITH_EDITOR
+FSlateApplication::FScopedPreventDebuggingMode::FScopedPreventDebuggingMode(FText InReason)
+{
+	static int32 IdGenerator = 0;
+	Id = ++IdGenerator;
+	FSlateApplication::Get().PreventDebuggingModeStack.Emplace(MoveTemp(InReason), Id);
+}
+
+FSlateApplication::FScopedPreventDebuggingMode::~FScopedPreventDebuggingMode()
+{
+	int32 IndexToRemove = FSlateApplication::Get().PreventDebuggingModeStack.IndexOfByPredicate([this](const TPair<FText, int32>& Reference){ return Reference.Value == Id;});
+	if (ensure(IndexToRemove != INDEX_NONE))
+	{
+		FSlateApplication::Get().PreventDebuggingModeStack.RemoveAtSwap(IndexToRemove);
+	}
+}
+#endif
+
 bool FSlateApplication::IsWindowInDestroyQueue(TSharedRef<SWindow> Window) const
 {
 	return WindowDestroyQueue.Contains(Window);
diff --git a/Engine/Source/Runtime/Slate/Public/Framework/Application/SlateApplication.h b/Engine/Source/Runtime/Slate/Public/Framework/Application/SlateApplication.h
--- a/Engine/Source/Runtime/Slate/Public/Framework/Application/SlateApplication.h
+++ b/Engine/Source/Runtime/Slate/Public/Framework/Application/SlateApplication.h
@@ -41,6 +41,7 @@ class ISlateSoundDevice;
 class ITextInputMethodSystem;
 class IVirtualKeyboardEntry;
 class IWidgetReflector;
+class SNotificationItem;
 class SViewport;
 class FSlateUser;
 class FSlateVirtualUserHandle;
@@ -839,6 +840,16 @@ public:
 	 * @param bLeavingDebugForSingleStep	Whether or not we are leaving debug mode due to single stepping
 	 */
 	SLATE_API void LeaveDebuggingMode( bool bLeavingDebugForSingleStep = false );
+
+#if WITH_EDITOR
+	struct FScopedPreventDebuggingMode
+	{
+		FScopedPreventDebuggingMode(FText Reason);
+		~FScopedPreventDebuggingMode();
+	private:
+		int32 Id;
+	};
+#endif
 	
 	/**
 	 * Calculates the popup window position from the passed in window position and size. 
@@ -1847,8 +1858,14 @@ private:
 #if WITH_EDITOR
 	/** List of all registered game viewports since the last time UnregisterGameViewport was called. */
 	TSet<TWeakPtr<SViewport>> AllGameViewports;
+
+	/** List of reason why we can't enter in debugging mode. */
+	TArray<TPair<FText, int32>> PreventDebuggingModeStack;
 #endif
 
+	/** The message when the EnterDebugginMode failed. */
+	TWeakPtr<SNotificationItem> DebuggingModeNotificationMessage;
+
 	TSharedPtr<ISlateSoundDevice> SlateSoundDevice;
 
 	/** The current cached absolute real time, right before we tick widgets */
