Fix missing padding for title bar buttons.

#10240

diff --git a/Engine/Source/Runtime/Slate/Private/Framework/Docking/SDockingTabStack.cpp b/Engine/Source/Runtime/Slate/Private/Framework/Docking/SDockingTabStack.cpp
--- a/Engine/Source/Runtime/Slate/Private/Framework/Docking/SDockingTabStack.cpp
+++ b/Engine/Source/Runtime/Slate/Private/Framework/Docking/SDockingTabStack.cpp
@@ -885,8 +885,8 @@ void SDockingTabStack::ReserveSpaceForWindowChrome(EChromeElement Element, bool
 	#else
 		static const float TopPaddingForMenuBar = 27.0f;
 		static const float LeftPaddingForIcon = FSlateApplication::Get().GetAppIcon()->GetImageSize().X;
-		const FMargin ControlsPadding = FMargin(8.f, bIncludePaddingForMenuBar ? TopPaddingForMenuBar : 5.f, 0.f, 0.f);
-		// If we are including top padding for the menu bar we do not need to pad the left side since we will be below the left icon
+		// If we are including top padding for the menu bar we do not need to pad the outer sides since we will be below the left icon and the right controls.
+		const FMargin ControlsPadding = FMargin(8.f, bIncludePaddingForMenuBar ? TopPaddingForMenuBar : 5.f, bIncludePaddingForMenuBar ? 0.f : 128.f, 0.f);
 		const FMargin IconPadding = FMargin(bIncludePaddingForMenuBar ? LeftPaddingForIcon + 12.f : 25.f, bOnlyMinorTabs ? 5.f : 0.f, 0.f, 0.f);
 	#endif
 
@@ -908,8 +908,6 @@ void SDockingTabStack::ReserveSpaceForWindowChrome(EChromeElement Element, bool
 	}
 }
 
-
-
 TSharedRef< SDockingTabStack > SDockingTabStack::CreateNewTabStackBySplitting( const SDockingNode::RelativeDirection Direction )
 {
 	TSharedPtr<SDockingSplitter> ParentNode = ParentNodePtr.Pin();
