Slate: fix retainer box and invalidation panel not propagating "Not Hit-Testable" to their children.
ensure widgets that are not hit visible are not added to the root hittest grid

diff --git a/Engine/Source/Runtime/Slate/Private/Widgets/SInvalidationPanel.cpp b/Engine/Source/Runtime/Slate/Private/Widgets/SInvalidationPanel.cpp
--- a/Engine/Source/Runtime/Slate/Private/Widgets/SInvalidationPanel.cpp
+++ b/Engine/Source/Runtime/Slate/Private/Widgets/SInvalidationPanel.cpp
@@ -247,9 +247,14 @@ int32 SInvalidationPanel::OnPaint( const FPaintArgs& Args, const FGeometry& Allo
 
 		const FSlateInvalidationResult Result = MutableThis->PaintInvalidationRoot(Context);
 
-		// add our widgets to the root hit test grid
-		Args.GetHittestGrid().AddGrid(HittestGrid);
+		const bool bInheritedHittestability = Args.GetInheritedHittestability();
+		const bool bOutgoingHittestability = bInheritedHittestability && GetVisibility().AreChildrenHitTestVisible();
 
+		// add our widgets to the root hit test grid
+		if (bOutgoingHittestability)
+		{
+			Args.GetHittestGrid().AddGrid(HittestGrid);
+		}
 		return Result.MaxLayerIdPainted;
 	}
 	else
diff --git a/Engine/Source/Runtime/UMG/Private/Slate/SRetainerWidget.cpp b/Engine/Source/Runtime/UMG/Private/Slate/SRetainerWidget.cpp
--- a/Engine/Source/Runtime/UMG/Private/Slate/SRetainerWidget.cpp
+++ b/Engine/Source/Runtime/UMG/Private/Slate/SRetainerWidget.cpp
@@ -670,8 +670,14 @@ int32 SRetainerWidget::OnPaint(const FPaintArgs& Args, const FGeometry& Allotted
 				);
 			}
 
+			const bool bInheritedHittestability = Args.GetInheritedHittestability();
+			const bool bOutgoingHittestability = bInheritedHittestability && GetVisibility().AreChildrenHitTestVisible();
+
 			// add our widgets to the root hit test grid
-			Args.GetHittestGrid().AddGrid(HittestGrid);
+			if (bOutgoingHittestability)
+			{
+				Args.GetHittestGrid().AddGrid(HittestGrid);
+			}
 
 			return GetCachedMaxLayerId();
 		}
