Improve structure and formatting of SDockingTabStack::ReserveSpaceForWindowChrome

diff --git a/Engine/Source/Runtime/Slate/Private/Framework/Docking/SDockingTabStack.cpp b/Engine/Source/Runtime/Slate/Private/Framework/Docking/SDockingTabStack.cpp
--- a/Engine/Source/Runtime/Slate/Private/Framework/Docking/SDockingTabStack.cpp
+++ b/Engine/Source/Runtime/Slate/Private/Framework/Docking/SDockingTabStack.cpp
@@ -866,45 +866,53 @@ void SDockingTabStack::ClearReservedSpace()
 
 void SDockingTabStack::ReserveSpaceForWindowChrome(EChromeElement Element, bool bIncludePaddingForMenuBar, bool bOnlyMinorTabs)
 {
-	#if PLATFORM_MAC
-		FMargin ControlsPadding;
-		static const FMargin IconPadding = FMargin(0);
+	FMargin ControlsPadding;
+	FMargin IconPadding;
 
-		if (bIncludePaddingForMenuBar)
-		{
-			static const float TopPaddingForTrafficLightsAndMenuBar = 30.0f;
-			// Always add padding on top, because on the Mac there is always either a main menu bar or the "traffic light" buttons (close, minimize, and maximize) above controls.
-			// Always add padding to the left, because on the Mac there's no Unreal icon to the left of controls, only the window edge, so we need some space.
-			ControlsPadding = FMargin(8.0f, TopPaddingForTrafficLightsAndMenuBar, 0, 0);
-		}
-		else
-		{
-			// Without a main menu bar in the title bar, we just need to pad on the left to avoid overlapping with the "traffic light" buttons (close, minimize, and maximize).
-			ControlsPadding = FMargin(67.0f, 0, 0, 0);
-		}
-	#else
-		static const float TopPaddingForMenuBar = 27.0f;
-		static const float LeftPaddingForIcon = FSlateApplication::Get().GetAppIcon()->GetImageSize().X;
-		// If we are including top padding for the menu bar we do not need to pad the outer sides since we will be below the left icon and the right controls.
-		const FMargin ControlsPadding = FMargin(8.f, bIncludePaddingForMenuBar ? TopPaddingForMenuBar : 5.f, bIncludePaddingForMenuBar ? 0.f : 128.f, 0.f);
-		const FMargin IconPadding = FMargin(bIncludePaddingForMenuBar ? LeftPaddingForIcon + 12.f : 25.f, bOnlyMinorTabs ? 5.f : 0.f, 0.f, 0.f);
-	#endif
+#if PLATFORM_MAC
+	if (bIncludePaddingForMenuBar)
+	{
+		static const float TopPaddingForTrafficLightsAndMenuBar = 30.0f;
+		// Always add padding on top, because on the Mac there is always either a main menu bar or the "traffic light" buttons (close, minimize, and maximize) above controls.
+		// Always add padding to the left, because on the Mac there's no Unreal icon to the left of controls, only the window edge, so we need some space.
+		ControlsPadding = FMargin(8.0f, TopPaddingForTrafficLightsAndMenuBar, 0, 0);
+	}
+	else
+	{
+		// Without a main menu bar in the title bar, we just need to pad on the left to avoid overlapping with the "traffic light" buttons (close, minimize, and maximize).
+		ControlsPadding = FMargin(67.0f, 0, 0, 0);
+	}
+#else
+	static const float TopPaddingForMenuBar = 27.0f;
+	static const float LeftPaddingForIcon = FSlateApplication::Get().GetAppIcon()->GetImageSize().X;
+	// If we are including top padding for the menu bar we do not need to pad the outer sides since we will be below the left icon and the right controls.
+	if (bIncludePaddingForMenuBar)
+	{
+		ControlsPadding = FMargin(8.f, TopPaddingForMenuBar, 0.f, 0.f);
+		IconPadding = FMargin(LeftPaddingForIcon + 12.f, bOnlyMinorTabs ? 5.f : 0.f, 0.f, 0.f);
+	}
+	else
+	{
+		ControlsPadding = FMargin(8.f, 5.f, 128.f, 0.f);
+		IconPadding = FMargin(25.f, bOnlyMinorTabs ? 5.f : 0.f, 0.f, 0.f);
+	}
+#endif
 
 	bShowingTitleBarArea = true;
 	const FMargin CurrentPadding = TitleBarSlot->GetPadding();
 	switch (Element)
 	{
-	case EChromeElement::Controls:
-		TitleBarSlot->SetPadding(CurrentPadding + ControlsPadding);
-		break;
+		case EChromeElement::Controls:
+			TitleBarSlot->SetPadding(CurrentPadding + ControlsPadding);
+			break;
 
-	case EChromeElement::Icon:
-		TitleBarSlot->SetPadding(CurrentPadding + IconPadding);
-		break;
+		case EChromeElement::Icon:
+			TitleBarSlot->SetPadding(CurrentPadding + IconPadding);
+			break;
 
-	default:
-		ensure(false);
-		break;
+		default:
+			ensure(false);
+			break;
 	}
 }
 
@@ -922,7 +930,6 @@ TSharedRef< SDockingTabStack > SDockingTabStack::CreateNewTabStackBySplitting( c
 	return NewStack;
 }
 
-
 void SDockingTabStack::SetParentNode( TSharedRef<class SDockingSplitter> InParent )
 {
 	SDockingNode::SetParentNode(InParent);
