Add tag to prevent creating FScopedPreventDebuggingMode without a name.

diff --git a/Engine/Source/Runtime/Slate/Private/Framework/Application/SlateApplication.cpp b/Engine/Source/Runtime/Slate/Private/Framework/Application/SlateApplication.cpp
--- a/Engine/Source/Runtime/Slate/Private/Framework/Application/SlateApplication.cpp
+++ b/Engine/Source/Runtime/Slate/Private/Framework/Application/SlateApplication.cpp
@@ -3655,7 +3655,7 @@ void FSlateApplication::EnterDebuggingMode()
 		}
 		UE_LOG(LogSlate, Warning, TEXT("Enter Debugging Mode failed."));
 
-		static bool bDoDebugBreak = true;
+		static volatile bool bDoDebugBreak = true;
 		if (bDoDebugBreak)
 		{
 			UE_DEBUG_BREAK();
diff --git a/Engine/Source/Runtime/Slate/Public/Framework/Application/SlateApplication.h b/Engine/Source/Runtime/Slate/Public/Framework/Application/SlateApplication.h
--- a/Engine/Source/Runtime/Slate/Public/Framework/Application/SlateApplication.h
+++ b/Engine/Source/Runtime/Slate/Public/Framework/Application/SlateApplication.h
@@ -844,8 +844,8 @@ public:
 #if WITH_EDITOR
 	struct FScopedPreventDebuggingMode
 	{
-		FScopedPreventDebuggingMode(FText Reason);
-		~FScopedPreventDebuggingMode();
+		UE_NODISCARD_CTOR SLATE_API FScopedPreventDebuggingMode(FText Reason);
+		SLATE_API ~FScopedPreventDebuggingMode();
 	private:
 		int32 Id;
 	};
