Fixed slider only moving by one step when bMouseUsesStep is enabled

diff --git a/Engine/Source/Runtime/Slate/Private/Widgets/Input/SSlider.cpp b/Engine/Source/Runtime/Slate/Private/Widgets/Input/SSlider.cpp
--- a/Engine/Source/Runtime/Slate/Private/Widgets/Input/SSlider.cpp
+++ b/Engine/Source/Runtime/Slate/Private/Widgets/Input/SSlider.cpp
@@ -448,13 +448,21 @@ float SSlider::PositionToValue( const FGeometry& MyGeometry, const UE::Slate::FD
 	{
 		float direction = ValueSlateAttribute.Get() - RelativeValue;
 		float CurrentStepSize = StepSize.Get();
+		if (CurrentStepSize <= 0)
+		{
+			// Invalid step size, keep current value
+			return ValueSlateAttribute.Get();
+		}
+		float Steps = FMath::Abs(direction) / CurrentStepSize;
+		Steps = FMath::RoundHalfFromZero(Steps);
+		const float ClampedDist = Steps * CurrentStepSize;
 		if (direction > CurrentStepSize / 2.0f)
 		{
-			return FMath::Clamp(ValueSlateAttribute.Get() - CurrentStepSize, MinValue, MaxValue);
+			return FMath::Clamp(ValueSlateAttribute.Get() - ClampedDist, MinValue, MaxValue);
 		}
 		else if (direction < CurrentStepSize / -2.0f)
 		{
-			return FMath::Clamp(ValueSlateAttribute.Get() + CurrentStepSize, MinValue, MaxValue);
+			return FMath::Clamp(ValueSlateAttribute.Get() + ClampedDist, MinValue, MaxValue);
 		}
 		else
 		{
