Increase non-main window titlebar height to fully contain min/max/close buttons

diff --git a/Engine/Source/Runtime/Slate/Private/Framework/Docking/SDockingTabStack.cpp b/Engine/Source/Runtime/Slate/Private/Framework/Docking/SDockingTabStack.cpp
--- a/Engine/Source/Runtime/Slate/Private/Framework/Docking/SDockingTabStack.cpp
+++ b/Engine/Source/Runtime/Slate/Private/Framework/Docking/SDockingTabStack.cpp
@@ -893,7 +893,7 @@ void SDockingTabStack::ReserveSpaceForWindowChrome(EChromeElement Element, bool
 	}
 	else
 	{
-		ControlsPadding = FMargin(8.f, 5.f, 128.f, 0.f);
+		ControlsPadding = FMargin(8.f, 10.f, 128.f, 0.f);
 		IconPadding = FMargin(25.f, bOnlyMinorTabs ? 5.f : 0.f, 0.f, 0.f);
 	}
 #endif
