Fix for unsafe calls to FSlateMacMenu::LanguageChanged()

diff --git a/Engine/Source/Runtime/Slate/Private/Framework/MultiBox/Mac/MacMenu.cpp b/Engine/Source/Runtime/Slate/Private/Framework/MultiBox/Mac/MacMenu.cpp
--- a/Engine/Source/Runtime/Slate/Private/Framework/MultiBox/Mac/MacMenu.cpp
+++ b/Engine/Source/Runtime/Slate/Private/Framework/MultiBox/Mac/MacMenu.cpp
@@ -443,6 +443,12 @@ void FSlateMacMenu::PostInitStartup()
 
 void FSlateMacMenu::LanguageChanged()
 {
+    if (!MacMenuHelper::GMacPostInitStartupRequested)
+    {
+        // if MacMenu not even done init setup yet, ignore LanguageChanged call
+        return;
+    }
+    
 	NSMenu* MainMenu = [NSApp mainMenu];
 	NSMenuItem* AppMenuItem = [MainMenu itemWithTitle:@"AppMenuItem"];
 	NSMenu* AppMenu = [AppMenuItem submenu];
