UE-196900 - Make Unreal Editor Titlebar Close and Min and Max Hover Colors Conform to Windows

Use Windows-style hover colors for titlebar min/max/close buttons

diff --git a/Engine/Source/Runtime/Slate/Private/Framework/Application/SWindowTitleBar.cpp b/Engine/Source/Runtime/Slate/Private/Framework/Application/SWindowTitleBar.cpp
--- a/Engine/Source/Runtime/Slate/Private/Framework/Application/SWindowTitleBar.cpp
+++ b/Engine/Source/Runtime/Slate/Private/Framework/Application/SWindowTitleBar.cpp
@@ -138,15 +138,14 @@ void SWindowTitleBar::MakeTitleBarContentWidgets( TSharedPtr< SWidget >& OutLeft
 				.ContentPadding(0.f)
 				.OnClicked(this, &SWindowTitleBar::MinimizeButton_OnClicked)
 				.Cursor(EMouseCursor::Default)
-				.ButtonStyle(FCoreStyle::Get(), "NoBorder")
+				.ButtonStyle(FAppStyle::Get(), "Window.MinMaxRestoreButtonHover")
 				.AddMetaData(FDriverMetaData::Id("launcher-minimizeWindowButton"))
 				[
 					SNew(SImage)
 						.Image(this, &SWindowTitleBar::GetMinimizeImage)
 						.ColorAndOpacity(this, &SWindowTitleBar::GetWindowTitleContentColor)
 						.AccessibleText(NSLOCTEXT("WindowTitleBar", "Minimize", "Minimize"))
-				]
-			;
+				];
 
 		MaximizeRestoreButton = SNew(SButton)
 				.IsFocusable(false)
@@ -154,15 +153,14 @@ void SWindowTitleBar::MakeTitleBarContentWidgets( TSharedPtr< SWidget >& OutLeft
 				.ContentPadding(0.0f)
 				.OnClicked(this, &SWindowTitleBar::MaximizeRestoreButton_OnClicked)
 				.Cursor(EMouseCursor::Default)
-				.ButtonStyle(FCoreStyle::Get(), "NoBorder")
+				.ButtonStyle(FAppStyle::Get(), "Window.MinMaxRestoreButtonHover")
 				.AddMetaData(FDriverMetaData::Id("launcher-maximizeRestoreWindowButton"))
 				[
 					SNew(SImage)
 						.Image(this, &SWindowTitleBar::GetMaximizeRestoreImage)
 						.ColorAndOpacity(this, &SWindowTitleBar::GetWindowTitleContentColor)
 						.AccessibleText(NSLOCTEXT("WindowTitleBar", "Maximize", "Maximize"))
-				]
-			;
+				];
 
 		CloseButton = SNew(SButton)
 				.IsFocusable(false)
@@ -170,15 +168,14 @@ void SWindowTitleBar::MakeTitleBarContentWidgets( TSharedPtr< SWidget >& OutLeft
 				.ContentPadding(0.0f)
 				.OnClicked(this, &SWindowTitleBar::CloseButton_OnClicked)
 				.Cursor(EMouseCursor::Default)
-				.ButtonStyle(FCoreStyle::Get(), "NoBorder")
+				.ButtonStyle(FAppStyle::Get(), "Window.CloseButtonHover")
 				.AddMetaData(FDriverMetaData::Id("launcher-closeWindowButton"))
 				[
 					SNew(SImage)
 						.Image(this, &SWindowTitleBar::GetCloseImage)
 						.ColorAndOpacity(this, &SWindowTitleBar::GetWindowTitleContentColor)
 						.AccessibleText(NSLOCTEXT("WindowTitleBar", "Close", "Close"))
-				]
-			;
+				];
 	}
 #endif //!PLATFORM_MAC
 
@@ -580,4 +577,4 @@ FText SWindowTitleBar::HandleWindowTitleText( ) const
 	}
 
 	return OwnerWindow->GetTitle();
-}
\ No newline at end of file
+}
diff --git a/Engine/Source/Runtime/SlateCore/Private/Styling/StarshipCoreStyle.cpp b/Engine/Source/Runtime/SlateCore/Private/Styling/StarshipCoreStyle.cpp
--- a/Engine/Source/Runtime/SlateCore/Private/Styling/StarshipCoreStyle.cpp
+++ b/Engine/Source/Runtime/SlateCore/Private/Styling/StarshipCoreStyle.cpp
@@ -251,10 +251,8 @@ TSharedRef<ISlateStyle> FStarshipCoreStyle::Create()
 		Style->Set("NoBorder.Pressed", new FSlateNoResource());
 
 		Style->Set("NoBorder", NoBorder);
-
 	}
 
-
 	// Demo Recording
 	{
 		Style->Set("DemoRecording.CursorPing", new IMAGE_BRUSH("Common/CursorPing", FVector2f(31.f, 31.f)));
@@ -752,8 +750,6 @@ TSharedRef<ISlateStyle> FStarshipCoreStyle::Create()
 		);
 	}
 
-
-
 	// SWindow defaults...
 	{
 #if !PLATFORM_MAC
@@ -778,8 +774,7 @@ TSharedRef<ISlateStyle> FStarshipCoreStyle::Create()
 			.SetPressed(IMAGE_BRUSH_SVG("Starship/CoreWidgets/Window/close", FVector2f(42.0f, 34.0f), FStyleColors::Foreground));
 #endif
 
-		FWindowStyle Window =
-			FWindowStyle()
+		FWindowStyle Window = FWindowStyle()
 #if !PLATFORM_MAC
 			.SetMinimizeButtonStyle(MinimizeButtonStyle)
 			.SetMaximizeButtonStyle(MaximizeButtonStyle)
@@ -797,9 +792,16 @@ TSharedRef<ISlateStyle> FStarshipCoreStyle::Create()
 			.SetCornerRadius(2)
 			.SetBorderPadding(FMargin(3.0f, 3.0f, 3.0f, 3.0f));
 
-
 		Style->Set("Window", Window);
-		
+
+		const FButtonStyle MinMaxRestoreButtonBackgroundStyle = FButtonStyle(NoBorder)
+			.SetHovered(FSlateColorBrush(FColor(255, 255, 255, 40)));
+		const FButtonStyle CloseButtonBackgroundStyle = FButtonStyle(NoBorder)
+			.SetHovered(FSlateColorBrush(FColor(255, 20, 20, 150)));
+
+		Style->Set("Window.MinMaxRestoreButtonHover", MinMaxRestoreButtonBackgroundStyle);
+		Style->Set("Window.CloseButtonHover", CloseButtonBackgroundStyle);
+
 		Window.SetCornerRadius(8.0f);
 
 		Style->Set("NotificationWindow", Window);
@@ -807,7 +809,6 @@ TSharedRef<ISlateStyle> FStarshipCoreStyle::Create()
 		Style->Set("ChildWindow.Background", new FSlateColorBrush(FStyleColors::Recessed));
 	}
 
-
 	// Standard Dialog Settings
 	{
 		Style->Set("StandardDialog.ContentPadding", FMargin(12.0f, 2.0f));
