Slate: Update the material and style if they are changed between the SlatePrepass and the OnPaint. A user can call SetFont during a tick in another widget.

diff --git a/Engine/Source/Runtime/Slate/Private/Framework/Text/TextLayout.cpp b/Engine/Source/Runtime/Slate/Private/Framework/Text/TextLayout.cpp
--- a/Engine/Source/Runtime/Slate/Private/Framework/Text/TextLayout.cpp
+++ b/Engine/Source/Runtime/Slate/Private/Framework/Text/TextLayout.cpp
@@ -2521,7 +2521,7 @@ void FTextLayout::SetScale( float Value )
 {
 	if (FMath::IsNaN(Value))
 	{
-		Value = 0.0;
+		Value = 0.0f;
 	}
 
 	if (Scale != Value)
diff --git a/Engine/Source/Runtime/Slate/Private/Widgets/Text/STextBlock.cpp b/Engine/Source/Runtime/Slate/Private/Widgets/Text/STextBlock.cpp
--- a/Engine/Source/Runtime/Slate/Private/Widgets/Text/STextBlock.cpp
+++ b/Engine/Source/Runtime/Slate/Private/Widgets/Text/STextBlock.cpp
@@ -30,24 +30,52 @@ void STextBlock::PrivateRegisterAttributes(FSlateAttributeInitializer& Attribute
 		{
 			return static_cast<const STextBlock&>(OwningWidget).bSimpleTextMode ? EInvalidateWidgetReason::Paint : EInvalidateWidgetReason::Layout;
 		}
+
+		static void UpdateTextStyle(SWidget& OwningWidget)
+		{
+			// see ComputeDesiredSize. We call UpdateTextStyle when the style changes.
+			static_cast<STextBlock&>(OwningWidget).bTextLayoutUpdateTextStyle = true;
+		}
+
+		static void UpdateDesiredSize(SWidget& OwningWidget)
+		{
+			// We call ComputeDesiredSize when the text changes.
+			static_cast<STextBlock&>(OwningWidget).bTextLayoutUpdateDesiredSize = true;
+		}
 	};
 
-	SLATE_ADD_MEMBER_ATTRIBUTE_DEFINITION(AttributeInitializer, BoundText, EInvalidateWidgetReason::Layout);
-	SLATE_ADD_MEMBER_ATTRIBUTE_DEFINITION(AttributeInitializer, Font, EInvalidateWidgetReason::Layout);
-	SLATE_ADD_MEMBER_ATTRIBUTE_DEFINITION(AttributeInitializer, StrikeBrush, FInvalidation::GetInvalidationNoneIfSimpleTextMode);
-	SLATE_ADD_MEMBER_ATTRIBUTE_DEFINITION(AttributeInitializer, ColorAndOpacity, FInvalidation::GetInvalidationPaintIfSimpleTextMode);
-	SLATE_ADD_MEMBER_ATTRIBUTE_DEFINITION(AttributeInitializer, ShadowOffset, EInvalidateWidgetReason::Layout);
-	SLATE_ADD_MEMBER_ATTRIBUTE_DEFINITION(AttributeInitializer, ShadowColorAndOpacity, FInvalidation::GetInvalidationPaintIfSimpleTextMode);
-	SLATE_ADD_MEMBER_ATTRIBUTE_DEFINITION(AttributeInitializer, HighlightColor, FInvalidation::GetInvalidationNoneIfSimpleTextMode);
-	SLATE_ADD_MEMBER_ATTRIBUTE_DEFINITION(AttributeInitializer, HighlightShape, FInvalidation::GetInvalidationNoneIfSimpleTextMode);
-	SLATE_ADD_MEMBER_ATTRIBUTE_DEFINITION(AttributeInitializer, HighlightText, FInvalidation::GetInvalidationNoneIfSimpleTextMode);
-	SLATE_ADD_MEMBER_ATTRIBUTE_DEFINITION(AttributeInitializer, WrapTextAt, FInvalidation::GetInvalidationNoneIfSimpleTextMode);
-	SLATE_ADD_MEMBER_ATTRIBUTE_DEFINITION(AttributeInitializer, AutoWrapText, FInvalidation::GetInvalidationNoneIfSimpleTextMode);
-	SLATE_ADD_MEMBER_ATTRIBUTE_DEFINITION(AttributeInitializer, WrappingPolicy, FInvalidation::GetInvalidationNoneIfSimpleTextMode);
-	SLATE_ADD_MEMBER_ATTRIBUTE_DEFINITION(AttributeInitializer, TransformPolicy, FInvalidation::GetInvalidationNoneIfSimpleTextMode);
-	SLATE_ADD_MEMBER_ATTRIBUTE_DEFINITION(AttributeInitializer, Margin, FInvalidation::GetInvalidationNoneIfSimpleTextMode);
-	SLATE_ADD_MEMBER_ATTRIBUTE_DEFINITION(AttributeInitializer, Justification, FInvalidation::GetInvalidationNoneIfSimpleTextMode);
-	SLATE_ADD_MEMBER_ATTRIBUTE_DEFINITION(AttributeInitializer, LineHeightPercentage, FInvalidation::GetInvalidationNoneIfSimpleTextMode);
+	SLATE_ADD_MEMBER_ATTRIBUTE_DEFINITION(AttributeInitializer, BoundText, EInvalidateWidgetReason::Layout)
+		.OnValueChanged(FSlateAttributeDescriptor::FAttributeValueChangedDelegate::CreateStatic(FInvalidation::UpdateDesiredSize));
+	SLATE_ADD_MEMBER_ATTRIBUTE_DEFINITION(AttributeInitializer, Font, EInvalidateWidgetReason::Layout)
+		.OnValueChanged(FSlateAttributeDescriptor::FAttributeValueChangedDelegate::CreateStatic(FInvalidation::UpdateTextStyle));
+	SLATE_ADD_MEMBER_ATTRIBUTE_DEFINITION(AttributeInitializer, StrikeBrush, FInvalidation::GetInvalidationNoneIfSimpleTextMode)
+		.OnValueChanged(FSlateAttributeDescriptor::FAttributeValueChangedDelegate::CreateStatic(FInvalidation::UpdateTextStyle));
+	SLATE_ADD_MEMBER_ATTRIBUTE_DEFINITION(AttributeInitializer, ColorAndOpacity, FInvalidation::GetInvalidationPaintIfSimpleTextMode)
+		.OnValueChanged(FSlateAttributeDescriptor::FAttributeValueChangedDelegate::CreateStatic(FInvalidation::UpdateTextStyle));
+	SLATE_ADD_MEMBER_ATTRIBUTE_DEFINITION(AttributeInitializer, ShadowOffset, EInvalidateWidgetReason::Layout)
+		.OnValueChanged(FSlateAttributeDescriptor::FAttributeValueChangedDelegate::CreateStatic(FInvalidation::UpdateTextStyle));
+	SLATE_ADD_MEMBER_ATTRIBUTE_DEFINITION(AttributeInitializer, ShadowColorAndOpacity, FInvalidation::GetInvalidationPaintIfSimpleTextMode)
+		.OnValueChanged(FSlateAttributeDescriptor::FAttributeValueChangedDelegate::CreateStatic(FInvalidation::UpdateTextStyle));
+	SLATE_ADD_MEMBER_ATTRIBUTE_DEFINITION(AttributeInitializer, HighlightColor, FInvalidation::GetInvalidationNoneIfSimpleTextMode)
+		.OnValueChanged(FSlateAttributeDescriptor::FAttributeValueChangedDelegate::CreateStatic(FInvalidation::UpdateTextStyle));
+	SLATE_ADD_MEMBER_ATTRIBUTE_DEFINITION(AttributeInitializer, HighlightShape, FInvalidation::GetInvalidationNoneIfSimpleTextMode)
+		.OnValueChanged(FSlateAttributeDescriptor::FAttributeValueChangedDelegate::CreateStatic(FInvalidation::UpdateTextStyle));
+	SLATE_ADD_MEMBER_ATTRIBUTE_DEFINITION(AttributeInitializer, HighlightText, FInvalidation::GetInvalidationNoneIfSimpleTextMode)
+		.OnValueChanged(FSlateAttributeDescriptor::FAttributeValueChangedDelegate::CreateStatic(FInvalidation::UpdateDesiredSize));
+	SLATE_ADD_MEMBER_ATTRIBUTE_DEFINITION(AttributeInitializer, WrapTextAt, FInvalidation::GetInvalidationNoneIfSimpleTextMode)
+		.OnValueChanged(FSlateAttributeDescriptor::FAttributeValueChangedDelegate::CreateStatic(FInvalidation::UpdateDesiredSize));
+	SLATE_ADD_MEMBER_ATTRIBUTE_DEFINITION(AttributeInitializer, AutoWrapText, FInvalidation::GetInvalidationNoneIfSimpleTextMode)
+		.OnValueChanged(FSlateAttributeDescriptor::FAttributeValueChangedDelegate::CreateStatic(FInvalidation::UpdateDesiredSize));
+	SLATE_ADD_MEMBER_ATTRIBUTE_DEFINITION(AttributeInitializer, WrappingPolicy, FInvalidation::GetInvalidationNoneIfSimpleTextMode)
+		.OnValueChanged(FSlateAttributeDescriptor::FAttributeValueChangedDelegate::CreateStatic(FInvalidation::UpdateDesiredSize));
+	SLATE_ADD_MEMBER_ATTRIBUTE_DEFINITION(AttributeInitializer, TransformPolicy, FInvalidation::GetInvalidationNoneIfSimpleTextMode)
+		.OnValueChanged(FSlateAttributeDescriptor::FAttributeValueChangedDelegate::CreateStatic(FInvalidation::UpdateDesiredSize));
+	SLATE_ADD_MEMBER_ATTRIBUTE_DEFINITION(AttributeInitializer, Margin, FInvalidation::GetInvalidationNoneIfSimpleTextMode)
+		.OnValueChanged(FSlateAttributeDescriptor::FAttributeValueChangedDelegate::CreateStatic(FInvalidation::UpdateDesiredSize));
+	SLATE_ADD_MEMBER_ATTRIBUTE_DEFINITION(AttributeInitializer, Justification, FInvalidation::GetInvalidationNoneIfSimpleTextMode)
+		.OnValueChanged(FSlateAttributeDescriptor::FAttributeValueChangedDelegate::CreateStatic(FInvalidation::UpdateDesiredSize));
+	SLATE_ADD_MEMBER_ATTRIBUTE_DEFINITION(AttributeInitializer, LineHeightPercentage, FInvalidation::GetInvalidationNoneIfSimpleTextMode)
+		.OnValueChanged(FSlateAttributeDescriptor::FAttributeValueChangedDelegate::CreateStatic(FInvalidation::UpdateDesiredSize));
 	SLATE_ADD_MEMBER_ATTRIBUTE_DEFINITION(AttributeInitializer, MinDesiredWidth, EInvalidateWidgetReason::Layout);
 }
 
@@ -74,6 +102,8 @@ STextBlock::STextBlock()
 {
 	SetCanTick(false);
 	bCanSupportFocus = false;
+	bTextLayoutUpdateTextStyle = true;
+	bTextLayoutUpdateDesiredSize = true;
 
 #if WITH_ACCESSIBILITY
 	AccessibleBehavior = EAccessibleBehavior::Auto;
@@ -273,15 +303,14 @@ int32 STextBlock::OnPaint( const FPaintArgs& Args, const FGeometry& AllottedGeom
 		const FVector2D LastDesiredSize = TextLayoutCache->GetDesiredSize();
 
 		// OnPaint will also update the text layout cache if required
+		UpdateTextBlockLayout(TextLayoutCache->GetLayoutScale());
 		LayerId = TextLayoutCache->OnPaint(Args, AllottedGeometry, MyCullingRect, OutDrawElements, LayerId, InWidgetStyle, ShouldBeEnabled(bParentEnabled));
 
-		const FVector2D NewDesiredSize = TextLayoutCache->GetDesiredSize();
-
 		// HACK: Due to the nature of wrapping and layout, we may have been arranged in a different box than what we were cached with.  Which
 		// might update wrapping, so make sure we always set the desired size to the current size of the text layout, which may have changed
 		// during paint.
 		const bool bCanWrap = WrapTextAt.Get() > 0 || AutoWrapText.Get();
-
+		const FVector2D NewDesiredSize = TextLayoutCache->GetDesiredSize();
 		if (bCanWrap && !NewDesiredSize.Equals(LastDesiredSize))
 		{
 			const_cast<STextBlock*>(this)->Invalidate(EInvalidateWidgetReason::Layout);
@@ -311,31 +340,9 @@ FVector2D STextBlock::ComputeDesiredSize(float LayoutScaleMultiplier) const
 	}
 	else
 	{
-		FTextBlockStyle::CompareParams StyleParams(TextStyle
-			, GetFontRef()
-			, GetColorAndOpacityRef()
-			, GetShadowOffset()
-			, GetShadowColorAndOpacityRef()
-			, GetHighlightColor()
-			, GetHighlightShape()
-			, GetStrikeBrush());
-
-		TextLayoutCache->ConditionallyUpdateTextStyle(MoveTemp(StyleParams));
-
-		// ComputeDesiredSize will also update the text layout cache if required
-		const FVector2D TextSize = TextLayoutCache->ComputeDesiredSize(
-			FSlateTextBlockLayout::FWidgetDesiredSizeArgs(
-				BoundText.Get(),
-				HighlightText.Get(),
-				WrapTextAt.Get(),
-				AutoWrapText.Get(),
-				WrappingPolicy.Get(),
-				GetTransformPolicyImpl(),
-				Margin.Get(),
-				LineHeightPercentage.Get(),
-				Justification.Get()),
-			LayoutScaleMultiplier);
-
+		bTextLayoutUpdateDesiredSize = true; // force ComputeDesiredSize
+		UpdateTextBlockLayout(LayoutScaleMultiplier);
+		const FVector2D TextSize = TextLayoutCache->GetDesiredSize();
 		return FVector2D(FMath::Max(MinDesiredWidth.Get(), TextSize.X), TextSize.Y);
 	}
 }
@@ -465,7 +472,7 @@ void STextBlock::SetJustification(TAttribute<ETextJustify::Type> InJustification
 FTextBlockStyle STextBlock::GetComputedTextStyle() const
 {
 	FTextBlockStyle ComputedStyle = TextStyle;
-	ComputedStyle.SetFont( GetFont() );
+	ComputedStyle.SetFont( GetFontRef() );
 	if (const FSlateBrush* const ComputedStrikeBrush = GetStrikeBrush())
 	{
 		ComputedStyle.SetStrikeBrush(*ComputedStrikeBrush);
@@ -474,10 +481,41 @@ FTextBlockStyle STextBlock::GetComputedTextStyle() const
 	ComputedStyle.SetShadowOffset( GetShadowOffset() );
 	ComputedStyle.SetShadowColorAndOpacity( GetShadowColorAndOpacity() );
 	ComputedStyle.SetHighlightColor( GetHighlightColor() );
-	ComputedStyle.SetHighlightShape( *GetHighlightShape() );
+	if (const FSlateBrush* const ComputedHighlightShape = GetHighlightShape())
+	{
+		ComputedStyle.SetHighlightShape(*ComputedHighlightShape);
+	}
 	return ComputedStyle;
 }
 
+void STextBlock::UpdateTextBlockLayout(float LayoutScaleMultiplier) const
+{
+	if (bTextLayoutUpdateTextStyle)
+	{
+		TextLayoutCache->UpdateTextStyle(GetComputedTextStyle());
+	}
+
+	if (bTextLayoutUpdateDesiredSize || bTextLayoutUpdateTextStyle)
+	{
+		// ComputeDesiredSize will also update the text layout cache if required
+		FSlateTextBlockLayout::FWidgetDesiredSizeArgs DesiredSizeArgs(
+			BoundText.Get(),
+			HighlightText.Get(),
+			WrapTextAt.Get(),
+			AutoWrapText.Get(),
+			WrappingPolicy.Get(),
+			GetTransformPolicyImpl(),
+			Margin.Get(),
+			LineHeightPercentage.Get(),
+			Justification.Get()
+		);
+		TextLayoutCache->ComputeDesiredSize(DesiredSizeArgs, LayoutScaleMultiplier);
+
+		bTextLayoutUpdateDesiredSize = false;
+		bTextLayoutUpdateTextStyle = false;
+	}
+}
+
 #if WITH_ACCESSIBILITY
 TSharedRef<FSlateAccessibleWidget> STextBlock::CreateAccessibleWidget()
 {
diff --git a/Engine/Source/Runtime/Slate/Private/Widgets/Text/SlateTextBlockLayout.cpp b/Engine/Source/Runtime/Slate/Private/Widgets/Text/SlateTextBlockLayout.cpp
--- a/Engine/Source/Runtime/Slate/Private/Widgets/Text/SlateTextBlockLayout.cpp
+++ b/Engine/Source/Runtime/Slate/Private/Widgets/Text/SlateTextBlockLayout.cpp
@@ -30,8 +30,7 @@ void FSlateTextBlockLayout::ConditionallyUpdateTextStyle(const FTextBlockStyle&
 	// Has the style used for this text block changed?
 	if (!IsStyleUpToDate(InTextStyle))
 	{
-		TextLayout->SetDefaultTextStyle(InTextStyle);
-		Marshaller->MakeDirty(); // will regenerate the text using the new default style
+		UpdateTextStyle(InTextStyle);
 	}
 }
 
@@ -40,23 +39,34 @@ void FSlateTextBlockLayout::ConditionallyUpdateTextStyle(const FTextBlockStyle::
 	// Has the style used for this text block changed?
 	if (!IsStyleUpToDate(InNewStyleParams))
 	{
-		FTextBlockStyle ComputedStyle = InNewStyleParams.StyleBase;
-		ComputedStyle.SetFont(InNewStyleParams.Font);
-		if (InNewStyleParams.StrikeBrush)
-		{
-			ComputedStyle.SetStrikeBrush(*InNewStyleParams.StrikeBrush);
-		}
-		ComputedStyle.SetColorAndOpacity(InNewStyleParams.ColorAndOpacity);
-		ComputedStyle.SetShadowOffset(InNewStyleParams.ShadowOffset);
-		ComputedStyle.SetShadowColorAndOpacity(InNewStyleParams.ShadowColorAndOpacity);
-		ComputedStyle.SetHighlightColor(InNewStyleParams.HighlightColor);
-		if (InNewStyleParams.HighlightShape)
-		{
-			ComputedStyle.SetHighlightShape(*InNewStyleParams.HighlightShape);
-		}
-		TextLayout->SetDefaultTextStyle(MoveTemp(ComputedStyle));
-		Marshaller->MakeDirty(); // will regenerate the text using the new default style
+		UpdateTextStyle(InNewStyleParams);
+	}
+}
+
+void FSlateTextBlockLayout::UpdateTextStyle(const FTextBlockStyle& InTextStyle)
+{
+	TextLayout->SetDefaultTextStyle(InTextStyle);
+	Marshaller->MakeDirty(); // will regenerate the text in ComputeDesiredSize()
+}
+
+void FSlateTextBlockLayout::UpdateTextStyle(const FTextBlockStyle::CompareParams& InNewStyleParams)
+{
+	FTextBlockStyle ComputedStyle = InNewStyleParams.StyleBase;
+	ComputedStyle.SetFont(InNewStyleParams.Font);
+	if (InNewStyleParams.StrikeBrush)
+	{
+		ComputedStyle.SetStrikeBrush(*InNewStyleParams.StrikeBrush);
+	}
+	ComputedStyle.SetColorAndOpacity(InNewStyleParams.ColorAndOpacity);
+	ComputedStyle.SetShadowOffset(InNewStyleParams.ShadowOffset);
+	ComputedStyle.SetShadowColorAndOpacity(InNewStyleParams.ShadowColorAndOpacity);
+	ComputedStyle.SetHighlightColor(InNewStyleParams.HighlightColor);
+	if (InNewStyleParams.HighlightShape)
+	{
+		ComputedStyle.SetHighlightShape(*InNewStyleParams.HighlightShape);
 	}
+
+	UpdateTextStyle(ComputedStyle);
 }
 
 
diff --git a/Engine/Source/Runtime/Slate/Public/Widgets/Text/STextBlock.h b/Engine/Source/Runtime/Slate/Public/Widgets/Text/STextBlock.h
--- a/Engine/Source/Runtime/Slate/Public/Widgets/Text/STextBlock.h
+++ b/Engine/Source/Runtime/Slate/Public/Widgets/Text/STextBlock.h
@@ -297,10 +297,13 @@ public:
 
 private:
 	/** Get the computed text style to use with the text marshaller */
-	SLATE_API FTextBlockStyle GetComputedTextStyle() const;
+	FTextBlockStyle GetComputedTextStyle() const;
+
+	/** Update the TextBlock layout. */
+	void UpdateTextBlockLayout(float LayoutScaleMultiplier) const;
 
 	/** Call to invalidate this text block */
-	SLATE_API void InvalidateText(EInvalidateWidgetReason InvalidateReason);
+	void InvalidateText(EInvalidateWidgetReason InvalidateReason);
 
 private:
 	/** The text displayed in this text block */
@@ -368,6 +371,7 @@ private:
 	{
 		struct 
 		{
+			//~ for attribute
 			uint16 bIsAttributeBoundTextBound : 1;
 			uint16 bIsAttributeFontSet : 1;
 			uint16 bIsAttributeStrikeBrushSet : 1;
@@ -378,6 +382,9 @@ private:
 			uint16 bIsAttributeHighlightShapeSet : 1;
 			uint16 bIsAttributeWrapTextAtSet : 1;
 			uint16 bIsAttributeTransformPolicySet : 1;
+			//~ for TextBlockLayout
+			mutable uint16 bTextLayoutUpdateTextStyle : 1;
+			mutable uint16 bTextLayoutUpdateDesiredSize : 1;
 		};
 		uint16 Union_Flags;
 	};
diff --git a/Engine/Source/Runtime/Slate/Public/Widgets/Text/SlateTextBlockLayout.h b/Engine/Source/Runtime/Slate/Public/Widgets/Text/SlateTextBlockLayout.h
--- a/Engine/Source/Runtime/Slate/Public/Widgets/Text/SlateTextBlockLayout.h
+++ b/Engine/Source/Runtime/Slate/Public/Widgets/Text/SlateTextBlockLayout.h
@@ -108,6 +108,16 @@ public:
 	 * Conditionally update the text style if needed
 	 */
 	SLATE_API void ConditionallyUpdateTextStyle(const FTextBlockStyle::CompareParams& InNewStyleParams);
+	
+	/**
+	 * Update the text style.
+	 */
+	SLATE_API void UpdateTextStyle(const FTextBlockStyle& InTextStyle);
+
+	/**
+	 * Update the text style.
+	 */
+	SLATE_API void UpdateTextStyle(const FTextBlockStyle::CompareParams& InNewStyleParams);
 
 	/**
 	 * Get the computed desired size for this layout, updating the internal cache as required
@@ -125,7 +135,8 @@ public:
 	 */
 	UE_DEPRECATED(5.0, "FWidgetArgs is deprecated. Upgrade to FWidgetDesiredSizeArgs instead.")
 	SLATE_API FVector2D ComputeDesiredSize(const FWidgetArgs& InWidgetArgs, const float InScale, const FTextBlockStyle& InTextStyle);
-	SLATE_API PRAGMA_ENABLE_DEPRECATION_WARNINGS
+
+	PRAGMA_ENABLE_DEPRECATION_WARNINGS
 
 	/**
 	 * Gets the last computed desired size.
