Slate: Add missing dll include introduced in CL 30376342

diff --git a/Engine/Source/Runtime/Slate/Public/Widgets/Text/SlateTextBlockLayout.h b/Engine/Source/Runtime/Slate/Public/Widgets/Text/SlateTextBlockLayout.h
--- a/Engine/Source/Runtime/Slate/Public/Widgets/Text/SlateTextBlockLayout.h
+++ b/Engine/Source/Runtime/Slate/Public/Widgets/Text/SlateTextBlockLayout.h
@@ -141,7 +141,7 @@ public:
 	/**
 	 * Gets the last computed desired size.
 	 */
-	FVector2D GetDesiredSize() const;
+	SLATE_API FVector2D GetDesiredSize() const;
 
 
 	/**
