Virtual Keyboard - Fixed an issue preventing the Virtual Keyboard to show on the first tap
#rnx
#tests Tested with preflight build on all platforms, no regression found

diff --git a/Engine/Source/Runtime/Slate/Private/Widgets/Text/SlateEditableTextLayout.cpp b/Engine/Source/Runtime/Slate/Private/Widgets/Text/SlateEditableTextLayout.cpp
--- a/Engine/Source/Runtime/Slate/Private/Widgets/Text/SlateEditableTextLayout.cpp
+++ b/Engine/Source/Runtime/Slate/Private/Widgets/Text/SlateEditableTextLayout.cpp
@@ -1320,9 +1320,7 @@ FReply FSlateEditableTextLayout::HandleMouseButtonUp(const FGeometry& MyGeometry
 		if (InMouseEvent.GetEffectingButton() == EKeys::LeftMouseButton ||
 			InMouseEvent.GetEffectingButton() == EKeys::RightMouseButton)
 		{
-			// Am I getting focus right now?
-			const bool bIsGettingFocus = !OwnerWidget->GetSlateWidget()->HasAnyUserFocus().IsSet();
-			if (!bIsGettingFocus)
+			if (!bWasFocusedByLastMouseDown)
 			{
 				// On platforms using a virtual keyboard open the virtual keyboard again 
 				if (FPlatformApplicationMisc::RequiresVirtualKeyboard())
