PR #10555: Added flow direction support for WrapBox

diff --git a/Engine/Source/Runtime/Slate/Private/Widgets/Layout/SWrapBox.cpp b/Engine/Source/Runtime/Slate/Private/Widgets/Layout/SWrapBox.cpp
--- a/Engine/Source/Runtime/Slate/Private/Widgets/Layout/SWrapBox.cpp
+++ b/Engine/Source/Runtime/Slate/Private/Widgets/Layout/SWrapBox.cpp
@@ -326,6 +326,12 @@ void SWrapBox::FChildArranger::FinalizeLine(int32 IndexOfLastChildInCurrentLine)
 			ArrangementData.SlotSize.Y = MaximumSizeInCurrentLine;
 		}
 		
+		// Flip horizontally if flow direction is set to RightToLeft
+		if (GSlateFlowDirection == EFlowDirection::RightToLeft)
+		{
+			ArrangementData.SlotOffset.X = WrapBox.PreferredSize.Get() - ArrangementData.SlotOffset.X - ArrangementData.SlotSize.X;
+		}
+				
 		OnSlotArranged(Slot, ArrangementData);
 	}
 
