PR #10832: Added ability for NavigationConfig to modify EUINavigationAction rules

diff --git a/Engine/Source/Runtime/Slate/Private/Framework/Application/NavigationConfig.cpp b/Engine/Source/Runtime/Slate/Private/Framework/Application/NavigationConfig.cpp
--- a/Engine/Source/Runtime/Slate/Private/Framework/Application/NavigationConfig.cpp
+++ b/Engine/Source/Runtime/Slate/Private/Framework/Application/NavigationConfig.cpp
@@ -27,6 +27,15 @@ FNavigationConfig::FNavigationConfig()
 
 	KeyEventRules.Emplace(EKeys::Down, EUINavigation::Down);
 	KeyEventRules.Emplace(EKeys::Gamepad_DPad_Down, EUINavigation::Down);
+
+	// By default, enter, space, and gamepad accept are all counted as accept
+	KeyActionRules.Emplace(EKeys::Enter, EUINavigationAction::Accept);
+	KeyActionRules.Emplace(EKeys::SpaceBar, EUINavigationAction::Accept);
+	KeyActionRules.Emplace(EKeys::Virtual_Accept, EUINavigationAction::Accept);
+
+	// By default, escape and gamepad back count as leaving current scope
+	KeyActionRules.Emplace(EKeys::Escape, EUINavigationAction::Back);
+	KeyActionRules.Emplace(EKeys::Virtual_Back, EUINavigationAction::Back);
 }
 
 FNavigationConfig::~FNavigationConfig()
@@ -168,17 +177,10 @@ PRAGMA_ENABLE_DEPRECATION_WARNINGS
 
 EUINavigationAction FNavigationConfig::GetNavigationActionForKey(const FKey& InKey) const
 {
-	if (InKey == EKeys::Enter || InKey == EKeys::SpaceBar || InKey == EKeys::Virtual_Accept)
+	if (const EUINavigationAction* Action = KeyActionRules.Find(InKey))
 	{
-		// By default, enter, space, and gamepad accept are all counted as accept
-		return EUINavigationAction::Accept;
+		return *Action;
 	}
-	else if (InKey == EKeys::Escape || InKey == EKeys::Virtual_Back)
-	{
-		// By default, escape and gamepad back count as leaving current scope
-		return EUINavigationAction::Back;
-	}
-
 	return EUINavigationAction::Invalid;
 }
 
diff --git a/Engine/Source/Runtime/Slate/Public/Framework/Application/NavigationConfig.h b/Engine/Source/Runtime/Slate/Public/Framework/Application/NavigationConfig.h
--- a/Engine/Source/Runtime/Slate/Public/Framework/Application/NavigationConfig.h
+++ b/Engine/Source/Runtime/Slate/Public/Framework/Application/NavigationConfig.h
@@ -112,6 +112,9 @@ public:
 	/** Digital key navigation rules. */
 	TMap<FKey, EUINavigation> KeyEventRules;
 
+	/** Digital key action rules. */
+	TMap<FKey, EUINavigationAction> KeyActionRules;
+
 protected:
 	/**
 	 * Gets the repeat rate of the navigation based on the current pressure being applied.  The idea being
