Fix FToolBarBuilder::AddToolBarWidget() not respecting label visibility

diff --git a/Engine/Source/Runtime/Slate/Private/Framework/MultiBox/MultiBoxBuilder.cpp b/Engine/Source/Runtime/Slate/Private/Framework/MultiBox/MultiBoxBuilder.cpp
--- a/Engine/Source/Runtime/Slate/Private/Framework/MultiBox/MultiBoxBuilder.cpp
+++ b/Engine/Source/Runtime/Slate/Private/Framework/MultiBox/MultiBoxBuilder.cpp
@@ -566,9 +566,6 @@ void FToolBarBuilder::AddToolBarWidget( TSharedRef<SWidget> InWidget, const TAtt
 {
 	ApplySectionBeginning();
 
-	// If tutorial name specified, wrap in tutorial wrapper
-	const FName WrapperName = GenerateTutorialIdentfierName(InTutorialHighlightName, NAME_None, nullptr, MultiBox->GetBlocks().Num());
-
 	const FToolBarStyle& ToolBarStyle = GetStyleSet()->GetWidgetStyle<FToolBarStyle>(GetStyleName());
 
 	TSharedRef<SWidget> ChildWidget = InWidget;
@@ -589,10 +586,15 @@ void FToolBarBuilder::AddToolBarWidget( TSharedRef<SWidget> InWidget, const TAtt
 		.HAlign( HAlign_Center )
 		[
 			SNew( STextBlock )
-			.Visibility_Lambda( [ChildWidget] () -> EVisibility {
+			.Visibility_Lambda( [LabelVisibilityCopy = LabelVisibility, ChildWidget] () -> EVisibility {
 				if (FMultiBoxSettings::UseSmallToolBarIcons.Get())
 					return EVisibility::Collapsed;
 
+				if (LabelVisibilityCopy.IsSet())
+				{
+					return LabelVisibilityCopy.GetValue();
+				}
+
 				return ChildWidget->GetVisibility();
 			})
 			.IsEnabled_Lambda( [ChildWidget] () -> bool { return ChildWidget->IsEnabled(); } )
@@ -609,9 +611,6 @@ void FToolBarBuilder::AddWidget( TSharedRef<SWidget> InWidget, FName InTutorialH
 {
 	ApplySectionBeginning();
 
-	// If tutorial name specified, wrap in tutorial wrapper
-	const FName WrapperName = GenerateTutorialIdentfierName(InTutorialHighlightName, NAME_None, nullptr, MultiBox->GetBlocks().Num());
-
 	TSharedRef<SWidget> ChildWidget = InWidget;
 	InWidget = 
 		SNew( SBox )
