Fix FToolBarBuilder::AddToolBarButton not accepting const FButtonArgs&

diff --git a/Engine/Source/Runtime/Slate/Private/Framework/MultiBox/MultiBoxBuilder.cpp b/Engine/Source/Runtime/Slate/Private/Framework/MultiBox/MultiBoxBuilder.cpp
--- a/Engine/Source/Runtime/Slate/Private/Framework/MultiBox/MultiBoxBuilder.cpp
+++ b/Engine/Source/Runtime/Slate/Private/Framework/MultiBox/MultiBoxBuilder.cpp
@@ -17,10 +17,7 @@
 #include "Framework/MultiBox/SButtonRowBlock.h"
 #include "Framework/MultiBox/SWidgetBlock.h"
 #include "Framework/MultiBox/SGroupMarkerBlock.h"
-#include "Framework/MultiBox/ToolMenuBase.h"
-#include "Widgets/Layout/SScrollBox.h"
 #include "Styling/ToolBarStyle.h"
-#include "Framework/MultiBox/SUniformToolbarButtonBlock.h"
 
 
 FMultiBoxBuilder::FMultiBoxBuilder( const EMultiBoxType InType, FMultiBoxCustomization InCustomization, const bool bInShouldCloseWindowAfterMenuSelection, const TSharedPtr< const FUICommandList >& InCommandList, TSharedPtr<FExtender> InExtender, FName InTutorialHighlightName, FName InMenuName )
@@ -115,7 +112,7 @@ TSharedRef< class FMultiBox > FMultiBoxBuilder::GetMultiBox()
 }
 
 /** Helper function to generate unique widget-identifying names given various bits of information */
-static FName GenerateTutorialIdentfierName(FName InContainerName, FName InElementName, const TSharedPtr< const FUICommandInfo > InCommand, int32 InIndex)
+static FName GenerateTutorialIdentifierName(FName InContainerName, FName InElementName, const TSharedPtr< const FUICommandInfo > InCommand, int32 InIndex)
 {
 	FString BaseName;
 	if(InContainerName != NAME_None)
@@ -155,7 +152,7 @@ void FBaseMenuBuilder::AddMenuEntry( const TSharedPtr< const FUICommandInfo > In
 	// The command must be valid
 	check( InCommand.IsValid() );
 	TSharedRef< FMenuEntryBlock > NewMenuEntryBlock = MakeShared<FMenuEntryBlock>( InExtensionHook, InCommand, CommandListStack.Last(), InLabelOverride, InToolTipOverride, InIconOverride, bCloseSelfOnly );
-	NewMenuEntryBlock->SetTutorialHighlightName(GenerateTutorialIdentfierName(TutorialHighlightName, InTutorialHighlightName, InCommand, MultiBox->GetBlocks().Num()));
+	NewMenuEntryBlock->SetTutorialHighlightName(GenerateTutorialIdentifierName(TutorialHighlightName, InTutorialHighlightName, InCommand, MultiBox->GetBlocks().Num()));
 	NewMenuEntryBlock->SetCheckBoxStyle(CheckBoxStyle);
 	MultiBox->AddMultiBlock( NewMenuEntryBlock );
 
@@ -169,7 +166,7 @@ void FBaseMenuBuilder::AddMenuEntry( const TAttribute<FText>& InLabel, const TAt
 	ApplyHook(InExtensionHook, EExtensionHook::Before);
 	
 	TSharedRef< FMenuEntryBlock > NewMenuEntryBlock = MakeShared<FMenuEntryBlock>( InExtensionHook, InLabel, InToolTip, InIcon, InAction, UserInterfaceActionType, bCloseSelfOnly );
-	NewMenuEntryBlock->SetTutorialHighlightName(GenerateTutorialIdentfierName(TutorialHighlightName, InTutorialHighlightName, nullptr, MultiBox->GetBlocks().Num()));
+	NewMenuEntryBlock->SetTutorialHighlightName(GenerateTutorialIdentifierName(TutorialHighlightName, InTutorialHighlightName, nullptr, MultiBox->GetBlocks().Num()));
 	MultiBox->AddMultiBlock( NewMenuEntryBlock );
 	
 	ApplyHook(InExtensionHook, EExtensionHook::After);
@@ -182,7 +179,7 @@ void FBaseMenuBuilder::AddMenuEntry( const FUIAction& UIAction, const TSharedRef
 	ApplyHook(InExtensionHook, EExtensionHook::Before);
 
 	TSharedRef< FMenuEntryBlock > NewMenuEntryBlock = MakeShared<FMenuEntryBlock>( InExtensionHook, UIAction, Contents, InToolTip, UserInterfaceActionType, bCloseSelfOnly );
-	NewMenuEntryBlock->SetTutorialHighlightName(GenerateTutorialIdentfierName(TutorialHighlightName, InTutorialHighlightName, nullptr, MultiBox->GetBlocks().Num()));
+	NewMenuEntryBlock->SetTutorialHighlightName(GenerateTutorialIdentifierName(TutorialHighlightName, InTutorialHighlightName, nullptr, MultiBox->GetBlocks().Num()));
 	MultiBox->AddMultiBlock( NewMenuEntryBlock );
 
 	ApplyHook(InExtensionHook, EExtensionHook::After);
@@ -195,7 +192,7 @@ void FBaseMenuBuilder::AddMenuEntry(const FMenuEntryParams& InMenuEntryParams)
 	ApplyHook(InMenuEntryParams.ExtensionHook, EExtensionHook::Before);
 	
 	TSharedPtr< FMenuEntryBlock > NewMenuEntryBlock = MakeShared<FMenuEntryBlock>(InMenuEntryParams);
-	NewMenuEntryBlock->SetTutorialHighlightName(GenerateTutorialIdentfierName(TutorialHighlightName, InMenuEntryParams.TutorialHighlightName, nullptr, MultiBox->GetBlocks().Num()));
+	NewMenuEntryBlock->SetTutorialHighlightName(GenerateTutorialIdentifierName(TutorialHighlightName, InMenuEntryParams.TutorialHighlightName, nullptr, MultiBox->GetBlocks().Num()));
 	MultiBox->AddMultiBlock(NewMenuEntryBlock.ToSharedRef());
 
 	ApplyHook(InMenuEntryParams.ExtensionHook, EExtensionHook::After);
@@ -281,7 +278,7 @@ void FMenuBuilder::AddSubMenu( const TAttribute<FText>& InMenuLabel, const TAttr
 
 	const bool bIsSubMenu = true;
 	TSharedRef< FMenuEntryBlock > NewMenuEntryBlock = MakeShared<FMenuEntryBlock>( InExtensionHook, InMenuLabel, InToolTip, InSubMenu, ExtenderStack.Top(), bIsSubMenu, bInOpenSubMenuOnClick, CommandListStack.Last(), bCloseSelfOnly, InIcon, bInShouldCloseWindowAfterMenuSelection );
-	NewMenuEntryBlock->SetTutorialHighlightName(GenerateTutorialIdentfierName(TutorialHighlightName, InTutorialHighlightName, nullptr, MultiBox->GetBlocks().Num()));
+	NewMenuEntryBlock->SetTutorialHighlightName(GenerateTutorialIdentifierName(TutorialHighlightName, InTutorialHighlightName, nullptr, MultiBox->GetBlocks().Num()));
 	NewMenuEntryBlock->SetRecursivelySearchable(bRecursivelySearchable);
 	NewMenuEntryBlock->SetCheckBoxStyle(CheckBoxStyle);
 	
@@ -398,7 +395,7 @@ void FMenuBarBuilder::AddPullDownMenu(const TAttribute<FText>& InMenuLabel, cons
 	// Pulldown menus always close all menus not just themselves
 	const bool bShouldCloseSelfOnly = false;
 	TSharedRef< FMenuEntryBlock > NewMenuEntryBlock(new FMenuEntryBlock(InExtensionHook, InMenuLabel, InToolTip, InPullDownMenu, ExtenderStack.Top(), bIsSubMenu, bOpenSubMenuOnClick, CommandListStack.Last(), bShouldCloseSelfOnly));
-	NewMenuEntryBlock->SetTutorialHighlightName(GenerateTutorialIdentfierName(TutorialHighlightName, InTutorialHighlightName, nullptr, MultiBox->GetBlocks().Num()));
+	NewMenuEntryBlock->SetTutorialHighlightName(GenerateTutorialIdentifierName(TutorialHighlightName, InTutorialHighlightName, nullptr, MultiBox->GetBlocks().Num()));
 
 	MultiBox->AddMultiBlock(NewMenuEntryBlock);
 
@@ -416,7 +413,7 @@ void FMenuBarBuilder::AddPullDownMenu(const TAttribute<FText>& InMenuLabel, cons
 	// Pulldown menus always close all menus not just themselves
 	const bool bShouldCloseSelfOnly = false;
 	TSharedRef< FMenuEntryBlock > NewMenuEntryBlock(new FMenuEntryBlock(InExtensionHook, InMenuLabel, InToolTip, InMenuContentGenerator, ExtenderStack.Top(), bIsSubMenu, bOpenSubMenuOnClick, CommandListStack.Last(), bShouldCloseSelfOnly));
-	NewMenuEntryBlock->SetTutorialHighlightName(GenerateTutorialIdentfierName(TutorialHighlightName, InTutorialHighlightName, nullptr, MultiBox->GetBlocks().Num()));
+	NewMenuEntryBlock->SetTutorialHighlightName(GenerateTutorialIdentifierName(TutorialHighlightName, InTutorialHighlightName, nullptr, MultiBox->GetBlocks().Num()));
 
 	MultiBox->AddMultiBlock(NewMenuEntryBlock);
 
@@ -458,7 +455,7 @@ void FToolBarBuilder::AddToolBarButton(const TSharedPtr< const FUICommandInfo >
 
 	NewToolBarButtonBlock->SetIsFocusable(bIsFocusable);
 	NewToolBarButtonBlock->SetForceSmallIcons(bForceSmallIcons);
-	NewToolBarButtonBlock->SetTutorialHighlightName(GenerateTutorialIdentfierName(TutorialHighlightName, InTutorialHighlightName, InCommand, MultiBox->GetBlocks().Num()));
+	NewToolBarButtonBlock->SetTutorialHighlightName(GenerateTutorialIdentifierName(TutorialHighlightName, InTutorialHighlightName, InCommand, MultiBox->GetBlocks().Num()));
 	NewToolBarButtonBlock->SetStyleNameOverride(CurrentStyleOverride);
 	NewToolBarButtonBlock->SetCustomMenuDelegate(InCustomMenuDelegate);
 
@@ -467,7 +464,7 @@ void FToolBarBuilder::AddToolBarButton(const TSharedPtr< const FUICommandInfo >
 	ApplyHook(InExtensionHook, EExtensionHook::After);
 }
 
-void FToolBarBuilder::AddToolBarButton(FButtonArgs& Args)
+void FToolBarBuilder::AddToolBarButton(const FButtonArgs& Args)
 {
 	ApplySectionBeginning();
 
@@ -487,7 +484,7 @@ void FToolBarBuilder::AddToolBarButton(FButtonArgs& Args)
 
 	NewToolBarButtonBlock->SetIsFocusable(bIsFocusable);
 	NewToolBarButtonBlock->SetForceSmallIcons(bForceSmallIcons);
-	NewToolBarButtonBlock->SetTutorialHighlightName(GenerateTutorialIdentfierName(TutorialHighlightName, Args.TutorialHighlightName, Args.Command, MultiBox->GetBlocks().Num()));
+	NewToolBarButtonBlock->SetTutorialHighlightName(GenerateTutorialIdentifierName(TutorialHighlightName, Args.TutorialHighlightName, Args.Command, MultiBox->GetBlocks().Num()));
 	NewToolBarButtonBlock->SetStyleNameOverride(CurrentStyleOverride);
 	NewToolBarButtonBlock->SetCustomMenuDelegate(Args.CustomMenuDelegate);
 	NewToolBarButtonBlock->SetOnGetMenuContent(Args.OnGetMenuContent);
@@ -518,7 +515,7 @@ void FToolBarBuilder::AddToolBarButton(const FUIAction& InAction, FName InExtens
 
 	NewToolBarButtonBlock->SetIsFocusable(bIsFocusable);
 	NewToolBarButtonBlock->SetForceSmallIcons(bForceSmallIcons);
-	NewToolBarButtonBlock->SetTutorialHighlightName(GenerateTutorialIdentfierName(TutorialHighlightName, InTutorialHighlightName, nullptr, MultiBox->GetBlocks().Num()));
+	NewToolBarButtonBlock->SetTutorialHighlightName(GenerateTutorialIdentifierName(TutorialHighlightName, InTutorialHighlightName, nullptr, MultiBox->GetBlocks().Num()));
 	NewToolBarButtonBlock->SetStyleNameOverride(CurrentStyleOverride);
 
 	MultiBox->AddMultiBlock( NewToolBarButtonBlock );
@@ -538,7 +535,7 @@ void FToolBarBuilder::AddComboButton( const FUIAction& InAction, const FOnGetCon
 	}
 
 	NewToolBarComboButtonBlock->SetForceSmallIcons(bForceSmallIcons);
-	NewToolBarComboButtonBlock->SetTutorialHighlightName(GenerateTutorialIdentfierName(TutorialHighlightName, InTutorialHighlightName, nullptr, MultiBox->GetBlocks().Num()));
+	NewToolBarComboButtonBlock->SetTutorialHighlightName(GenerateTutorialIdentifierName(TutorialHighlightName, InTutorialHighlightName, nullptr, MultiBox->GetBlocks().Num()));
 	NewToolBarComboButtonBlock->SetStyleNameOverride(CurrentStyleOverride);
 
 	MultiBox->AddMultiBlock( NewToolBarComboButtonBlock );
@@ -556,7 +553,7 @@ void FToolBarBuilder::AddToolbarStackButton(const TSharedPtr< const FUICommandIn
 	}
 
 	NewToolBarStackButtonBlock->SetForceSmallIcons(bForceSmallIcons);
-	NewToolBarStackButtonBlock->SetTutorialHighlightName(GenerateTutorialIdentfierName(TutorialHighlightName, InTutorialHighlightName, InCommand, MultiBox->GetBlocks().Num()));
+	NewToolBarStackButtonBlock->SetTutorialHighlightName(GenerateTutorialIdentifierName(TutorialHighlightName, InTutorialHighlightName, InCommand, MultiBox->GetBlocks().Num()));
 	NewToolBarStackButtonBlock->SetStyleNameOverride(CurrentStyleOverride);
 
 	MultiBox->AddMultiBlock(NewToolBarStackButtonBlock);
@@ -750,7 +747,7 @@ FSlimHorizontalUniformToolBarBuilder::FSlimHorizontalUniformToolBarBuilder(TShar
 	const FToolBarStyle& ToolBarStyle = GetStyleSet()->GetWidgetStyle<FToolBarStyle>(GetStyleName());
 }
 
-void FSlimHorizontalUniformToolBarBuilder::AddToolBarButton( FButtonArgs& ButtonArgs )
+void FSlimHorizontalUniformToolBarBuilder::AddToolBarButton(const FButtonArgs& ButtonArgs)
 {
  	ApplySectionBeginning();
 	ApplyHook(ButtonArgs.ExtensionHook, EExtensionHook::Before);
@@ -759,7 +756,7 @@ void FSlimHorizontalUniformToolBarBuilder::AddToolBarButton( FButtonArgs& Button
 	const TSharedPtr< FToolBarButtonBlock > NewHorizontalToolBarButtonBlock( new FToolBarButtonBlock( ButtonArgs) );
     InitializeToolBarButtonBlock(StaticCastSharedPtr<FToolBarButtonBlock>(NewHorizontalToolBarButtonBlock), ButtonArgs);
 }
-void FToolBarBuilder::InitializeToolBarButtonBlock(TSharedPtr<FToolBarButtonBlock> ToolBarButtonBlock, FButtonArgs& ButtonArgs)
+void FToolBarBuilder::InitializeToolBarButtonBlock(TSharedPtr<FToolBarButtonBlock> ToolBarButtonBlock, const FButtonArgs& ButtonArgs)
 {
 	if ( LabelVisibility.IsSet() )
 	{
@@ -769,7 +766,7 @@ void FToolBarBuilder::InitializeToolBarButtonBlock(TSharedPtr<FToolBarButtonBloc
 	ToolBarButtonBlock->SetBorderBrushName(ButtonArgs.BorderBrushName);
 	ToolBarButtonBlock->SetIsFocusable(bIsFocusable);
 	ToolBarButtonBlock->SetForceSmallIcons(bForceSmallIcons);
-	ToolBarButtonBlock->SetTutorialHighlightName(GenerateTutorialIdentfierName(
+	ToolBarButtonBlock->SetTutorialHighlightName(GenerateTutorialIdentifierName(
 	TutorialHighlightName, ButtonArgs.TutorialHighlightName, ButtonArgs.Command, MultiBox->GetBlocks().Num()));
 	ToolBarButtonBlock->SetCustomMenuDelegate(ButtonArgs.CustomMenuDelegate);
 	ToolBarButtonBlock->SetOnGetMenuContent(ButtonArgs.OnGetMenuContent);
diff --git a/Engine/Source/Runtime/Slate/Private/Framework/MultiBox/SToolBarButtonBlock.cpp b/Engine/Source/Runtime/Slate/Private/Framework/MultiBox/SToolBarButtonBlock.cpp
--- a/Engine/Source/Runtime/Slate/Private/Framework/MultiBox/SToolBarButtonBlock.cpp
+++ b/Engine/Source/Runtime/Slate/Private/Framework/MultiBox/SToolBarButtonBlock.cpp
@@ -51,12 +51,12 @@ FToolBarButtonBlock::FToolBarButtonBlock( const TAttribute<FText>& InLabel, cons
 {
 }
 
-void FToolBarButtonBlock::SetCustomMenuDelegate( FNewMenuDelegate& InCustomMenuDelegate )
+void FToolBarButtonBlock::SetCustomMenuDelegate(const FNewMenuDelegate& InCustomMenuDelegate)
 {
 	CustomMenuDelegate = InCustomMenuDelegate;
 }
 
-void FToolBarButtonBlock::SetOnGetMenuContent( FOnGetContent& InOnGetMenuContent)
+void FToolBarButtonBlock::SetOnGetMenuContent(const FOnGetContent& InOnGetMenuContent)
 {
 	OnGetMenuContent = InOnGetMenuContent;
 }
diff --git a/Engine/Source/Runtime/Slate/Public/Framework/MultiBox/MultiBoxBuilder.h b/Engine/Source/Runtime/Slate/Public/Framework/MultiBox/MultiBoxBuilder.h
--- a/Engine/Source/Runtime/Slate/Public/Framework/MultiBox/MultiBoxBuilder.h
+++ b/Engine/Source/Runtime/Slate/Public/Framework/MultiBox/MultiBoxBuilder.h
@@ -308,8 +308,8 @@ public:
 	 *                                   function return type can do that. Think about a function like "A* A::GetParent()". If a root menu expands "A" functions and expands on the function return types, then
 	 *                                   selecting "GetParent()" will expand another "A" menu. Without simulation, the reflection API don't know when the recursion finishes, nor does the recursive search algorithm.
 	 */
-	FMenuBuilder( const bool bInShouldCloseWindowAfterMenuSelection, TSharedPtr< const FUICommandList > InCommandList, TSharedPtr<FExtender> InExtender = TSharedPtr<FExtender>(), const bool bCloseSelfOnly = false, const ISlateStyle* InStyleSet = &FCoreStyle::Get(), bool bInSearchable = true, FName InMenuName = NAME_None, bool bInRecursivelySearchable = true)
-		: FBaseMenuBuilder( EMultiBoxType::Menu, bInShouldCloseWindowAfterMenuSelection, InCommandList, bCloseSelfOnly, InExtender, InStyleSet, NAME_None, InMenuName )
+	FMenuBuilder( const bool bInShouldCloseWindowAfterMenuSelection, TSharedPtr< const FUICommandList > InCommandList, TSharedPtr<FExtender> InExtender = TSharedPtr<FExtender>(), const bool bInCloseSelfOnly = false, const ISlateStyle* InStyleSet = &FCoreStyle::Get(), bool bInSearchable = true, FName InMenuName = NAME_None, bool bInRecursivelySearchable = true)
+		: FBaseMenuBuilder( EMultiBoxType::Menu, bInShouldCloseWindowAfterMenuSelection, InCommandList, bInCloseSelfOnly, InExtender, InStyleSet, NAME_None, InMenuName )
 		, bSectionNeedsToBeApplied(false)
 		, bSearchable(bInSearchable)
 		, bRecursivelySearchable(bInRecursivelySearchable)
@@ -374,11 +374,12 @@ public:
 	 * 
 	 * @param	InWidget			The widget that should be shown in the menu
 	 * @param	InLabel				Optional label text to be added to the left of the content
-	 * @param	bNoIndent			If true, removes the padding from the left of the widget that lines it up with other menu items (default == false)
-	 * @param	bSearchable			If true, widget will be searchable (default == true)
+	 * @param	bInNoIndent			If true, removes the padding from the left of the widget that lines it up with other menu items (default == false)
+	 * @param	bInSearchable			If true, widget will be searchable (default == true)
 	 * @param	InToolTipText	Optional tooltip text to be added to the widget and label
 	 */
-	SLATE_API void AddWidget( TSharedRef<SWidget> InWidget, const FText& Label, bool bNoIndent = false, bool bSearchable = true, const TAttribute<FText>& InToolTipText = FText());
+	SLATE_API void AddWidget( TSharedRef<SWidget> InWidget, const FText& InLabel, bool bInNoIndent = false, bool bInSearchable = true, const 
+	TAttribute<FText>&  InToolTipText = FText());
 
 	/**
 	* Adds the widget the multibox will use for searching
@@ -524,28 +525,30 @@ public:
 	 *
 	 * @param	ButtonArgs The Parameters object which will provide the data to initialize the button
 	 */
-	SLATE_API virtual void AddToolBarButton(FButtonArgs& ButtonArgs);
+	SLATE_API virtual void AddToolBarButton(const FButtonArgs& ButtonArgs);
 
 
 	/**
 	 * Adds a tool bar button
 	 *
 	 * @param	InCommand				The command associated with this tool bar button
+	 * @param	InExtensionHook			The section hook. Can be NAME_None.
 	 * @param	InLabelOverride			Optional label override.  If omitted, then the action's label will be used instead.
 	 * @param	InToolTipOverride		Optional tool tip override.	 If omitted, then the action's label will be used instead.
 	 * @param	InIconOverride			Optional name of the slate brush to use for the tool bar image.  If omitted, then the action's icon will be used instead.
 	 * @param	InTutorialHighlightName	Name to identify this widget and highlight during tutorials
-	 * @param	CustomMenuDelegate  Optional custom menu delegate for cases where the toolbar is compressed into a menu
+	 * @param	InCustomMenuDelegate  Optional custom menu delegate for cases where the toolbar is compressed into a menu
 	 */
 	SLATE_API void AddToolBarButton(const TSharedPtr< const FUICommandInfo > InCommand, FName InExtensionHook = NAME_None, const TAttribute<FText>& InLabelOverride = TAttribute<FText>(), const TAttribute<FText>& InToolTipOverride = TAttribute<FText>(), const TAttribute<FSlateIcon>& InIconOverride = TAttribute<FSlateIcon>(), FName InTutorialHighlightName = NAME_None, FNewMenuDelegate InCustomMenuDelegate = FNewMenuDelegate());
 	
 	/**
 	 * Adds a tool bar button
 	 *
-	 * @param	UIAction	Actions to execute on this menu item.
-	 * @param	InLabel		Label to show in the menu entry
-	 * @param	InToolTip	Tool tip used when hovering over the menu entry
-	 * @param	InIcon		The icon to use		
+	 * @param	InAction				Actions to execute on this menu item.
+	 * @param	InExtensionHook			The section hook. Can be NAME_None.
+	 * @param	InLabelOverride			Optional label override.  If omitted, then the action's label will be used instead.
+	 * @param	InToolTipOverride		Optional tool tip override.	 If omitted, then the action's label will be used instead.
+	 * @param	InIconOverride			Optional icon to use for the tool bar image.  If omitted, then the action's icon will be used instead.
 	 * @param	UserInterfaceActionType	Type of interface action
 	 * @param	InTutorialHighlightName	Name to identify this widget and highlight during tutorials
 	 */
@@ -578,9 +581,9 @@ public:
 	 * @param	InWidget				The widget that should be shown in the toolbar
 	 * @param	InLabel                 Optional Label.  
 	 * @param	InTutorialHighlightName	Name to identify this widget and highlight during tutorials
-	 * @param	bSearchable			If true, widget will be searchable (default == true)
+	 * @param	bInSearchable				If true, widget will be searchable (default == true)
 	 */
-	SLATE_API void AddToolBarWidget(TSharedRef<SWidget> InWidget, const TAttribute<FText>& InLabel = TAttribute<FText>(), FName InTutorialHighlightName = NAME_None, bool bSearchable = true);
+	SLATE_API void AddToolBarWidget(TSharedRef<SWidget> InWidget, const TAttribute<FText>& InLabel = TAttribute<FText>(), FName InTutorialHighlightName = NAME_None, bool bInSearchable = true);
 
 
 	/**
@@ -588,11 +591,11 @@ public:
 	 * 
 	 * @param	InWidget				The widget that should be shown in the toolbar
 	 * @param	InTutorialHighlightName	Name to identify this widget and highlight during tutorials
-	 * @param	bSearchable			If true, widget will be searchable (default == true)
-	 * @param	Alignment			Horizontal alignment for the widget inside the toolbar
-	 * @param	CustomMenuDelegate  Optional custom menu delegate for cases where the toolbar is compressed into a menu
+	 * @param	bInSearchable			If true, widget will be searchable (default == true)
+	 * @param	InAlignment				Horizontal alignment for the widget inside the toolbar
+	 * @param	InCustomMenuDelegate	Optional custom menu delegate for cases where the toolbar is compressed into a menu
 	 */
-	SLATE_API void AddWidget(TSharedRef<SWidget> InWidget, FName InTutorialHighlightName = NAME_None, bool bSearchable = true, EHorizontalAlignment Alignment = HAlign_Fill, FNewMenuDelegate InCustomMenuDelegate = FNewMenuDelegate());
+	SLATE_API void AddWidget(TSharedRef<SWidget> InWidget, FName InTutorialHighlightName = NAME_None, bool bInSearchable = true, EHorizontalAlignment InAlignment = HAlign_Fill, FNewMenuDelegate InCustomMenuDelegate = FNewMenuDelegate());
 	
 	/**
 	 * Adds a toolbar separator
@@ -642,7 +645,7 @@ protected:
 	SLATE_API virtual void ApplyHook(FName InExtensionHook, EExtensionHook::Position HookPosition) override;
 	SLATE_API virtual void ApplySectionBeginning() override;
 
-	SLATE_API void InitializeToolBarButtonBlock(TSharedPtr<FToolBarButtonBlock> ButtonRowBlock, FButtonArgs& ButtonArgs);
+	SLATE_API void InitializeToolBarButtonBlock(TSharedPtr<FToolBarButtonBlock> ButtonRowBlock, const FButtonArgs& ButtonArgs);
 
 private:
 	/** Current extension hook name for sections to determine where sections begin and end */
@@ -737,12 +740,12 @@ public:
 	 * Adds a button to a row
 	 *
 	 * @param	InLabel					The button label to display
-	 * @param	InToolTipOverride		The tooltip for the button
-	 * @param	UIAction				Action to execute when the button is clicked or when state should be checked
+	 * @param	InToolTip				The tooltip for the button
+	 * @param	InUIAction				Action to execute when the button is clicked or when state should be checked
 	 * @param	InIcon					The icon for the button
-	 * @param	UserInterfaceActionType	The style of button to display
+	 * @param	InUserInterfaceActionType	The style of button to display
 	 */
-	SLATE_API void AddButton(const FText& InLabel, const FText& InToolTip, const FUIAction& UIAction, const FSlateIcon& InIcon = FSlateIcon(), const EUserInterfaceActionType UserInterfaceActionType = EUserInterfaceActionType::Button);
+	SLATE_API void AddButton(const FText& InLabel, const FText& InToolTip, const FUIAction& InUIAction, const FSlateIcon& InIcon = FSlateIcon(), const EUserInterfaceActionType InUserInterfaceActionType = EUserInterfaceActionType::Button);
 
 protected:
 	/** FMultiBoxBuilder interface */
@@ -759,6 +762,6 @@ public:
 	 */
 	SLATE_API FSlimHorizontalUniformToolBarBuilder(TSharedPtr<const FUICommandList> InCommandList, FMultiBoxCustomization InCustomization, TSharedPtr<FExtender> InExtender = nullptr, const bool InForceSmallIcons = false);
 
-	SLATE_API void AddToolBarButton( FButtonArgs& ButtonArgs ) override;
+	SLATE_API virtual void AddToolBarButton(const FButtonArgs& ButtonArgs) override;
 
 };
diff --git a/Engine/Source/Runtime/Slate/Public/Framework/MultiBox/SToolBarButtonBlock.h b/Engine/Source/Runtime/Slate/Public/Framework/MultiBox/SToolBarButtonBlock.h
--- a/Engine/Source/Runtime/Slate/Public/Framework/MultiBox/SToolBarButtonBlock.h
+++ b/Engine/Source/Runtime/Slate/Public/Framework/MultiBox/SToolBarButtonBlock.h
@@ -87,9 +87,9 @@ public:
 	SLATE_API virtual bool HasIcon() const override;
 
 	/** Set optional delegate to customize when a menu appears instead of the widget, such as in toolbars */
-	SLATE_API void SetCustomMenuDelegate( FNewMenuDelegate& InOnFillMenuDelegate);
+	SLATE_API void SetCustomMenuDelegate(const FNewMenuDelegate& InOnFillMenuDelegate);
 
-	SLATE_API void SetOnGetMenuContent( FOnGetContent& OnGetMenuContent);
+	SLATE_API void SetOnGetMenuContent(const FOnGetContent& OnGetMenuContent);
 
 protected:
 	
