PR #11321: Fix AnalogCursor mouse move events not working with mouse capture

diff --git a/Engine/Source/Runtime/Slate/Private/Framework/Application/AnalogCursor.cpp b/Engine/Source/Runtime/Slate/Private/Framework/Application/AnalogCursor.cpp
--- a/Engine/Source/Runtime/Slate/Private/Framework/Application/AnalogCursor.cpp
+++ b/Engine/Source/Runtime/Slate/Private/Framework/Application/AnalogCursor.cpp
@@ -270,6 +270,7 @@ void FAnalogCursor::UpdateCursorPosition(FSlateApplication& SlateApp, TSharedRef
 		const bool bIsPrimaryUser = FSlateApplication::CursorUserIndex == SlateUser->GetUserIndex();
 		FPointerEvent MouseEvent(
 			SlateUser->GetUserIndex(),
+			FSlateApplication::CursorPointerIndex,
 			UpdatedPosition,
 			OldPosition,
 			bIsPrimaryUser ? SlateApp.GetPressedMouseButtons() : TSet<FKey>(),
