Fix warning V1101: The default argument of the second virtual function parameter was changed in the derived class from '2147483647' to '1000'. Calling such function via a pointer or a reference to the base class may result in unexpected behavior.

diff --git a/Engine/Source/Runtime/Slate/Private/Framework/MultiBox/MultiBoxBuilder.cpp b/Engine/Source/Runtime/Slate/Private/Framework/MultiBox/MultiBoxBuilder.cpp
--- a/Engine/Source/Runtime/Slate/Private/Framework/MultiBox/MultiBoxBuilder.cpp
+++ b/Engine/Source/Runtime/Slate/Private/Framework/MultiBox/MultiBoxBuilder.cpp
@@ -101,7 +101,7 @@ FMultiBoxCustomization FMultiBoxBuilder::GetCustomization() const
 	return FMultiBoxCustomization( MultiBox->GetCustomizationName() ); 
 }
 
-TSharedRef< class SWidget > FMultiBoxBuilder::MakeWidget( FMultiBox::FOnMakeMultiBoxBuilderOverride* InMakeMultiBoxBuilderOverride, uint32 MaxHeight /* = nullptr */ )
+TSharedRef< class SWidget > FMultiBoxBuilder::MakeWidget( FMultiBox::FOnMakeMultiBoxBuilderOverride* InMakeMultiBoxBuilderOverride)
 {
 	return MultiBox->MakeWidget( false, InMakeMultiBoxBuilderOverride );
 }
@@ -198,7 +198,12 @@ void FBaseMenuBuilder::AddMenuEntry(const FMenuEntryParams& InMenuEntryParams)
 	ApplyHook(InMenuEntryParams.ExtensionHook, EExtensionHook::After);
 }
 
-TSharedRef< class SWidget > FMenuBuilder::MakeWidget( FMultiBox::FOnMakeMultiBoxBuilderOverride* InMakeMultiBoxBuilderOverride /* = nullptr */, uint32 MaxHeight)
+TSharedRef< class SWidget > FMenuBuilder::MakeWidget( FMultiBox::FOnMakeMultiBoxBuilderOverride* InMakeMultiBoxBuilderOverride)
+{
+	return MakeWidget(InMakeMultiBoxBuilderOverride, 1000);
+}
+
+TSharedRef< class SWidget > FMenuBuilder::MakeWidget( FMultiBox::FOnMakeMultiBoxBuilderOverride* InMakeMultiBoxBuilderOverride, uint32 MaxHeight)
 {
 	// Make menu builders searchable (by default)
 	TAttribute<float> MaxHeightAttribute;
diff --git a/Engine/Source/Runtime/Slate/Public/Framework/MultiBox/MultiBoxBuilder.h b/Engine/Source/Runtime/Slate/Public/Framework/MultiBox/MultiBoxBuilder.h
--- a/Engine/Source/Runtime/Slate/Public/Framework/MultiBox/MultiBoxBuilder.h
+++ b/Engine/Source/Runtime/Slate/Public/Framework/MultiBox/MultiBoxBuilder.h
@@ -79,7 +79,7 @@ public:
 	 *
 	 * @return  New widget object
 	 */
-	SLATE_API virtual TSharedRef< class SWidget > MakeWidget( FMultiBox::FOnMakeMultiBoxBuilderOverride* InMakeMultiBoxBuilderOverride = nullptr, uint32 MaxHeight = INT_MAX);
+	SLATE_API virtual TSharedRef< class SWidget > MakeWidget( FMultiBox::FOnMakeMultiBoxBuilderOverride* InMakeMultiBoxBuilderOverride = nullptr);
 	
 
 	/** 
@@ -326,8 +326,10 @@ public:
 	*
 	* @return  New widget object
 	*/
-	SLATE_API virtual TSharedRef< class SWidget > MakeWidget( FMultiBox::FOnMakeMultiBoxBuilderOverride* InMakeMultiBoxBuilderOverride = nullptr, uint32 MaxHeight = 1000) override;
+	SLATE_API virtual TSharedRef< class SWidget > MakeWidget( FMultiBox::FOnMakeMultiBoxBuilderOverride* InMakeMultiBoxBuilderOverride = nullptr) override;
+	SLATE_API virtual TSharedRef< class SWidget > MakeWidget( FMultiBox::FOnMakeMultiBoxBuilderOverride* InMakeMultiBoxBuilderOverride, uint32 MaxHeight);
 
+public:
 	/**
 	 * Adds a separator
 	 */
