Locker - Engine changes - Add an option to chose the Visibility state of the ScrollBar on ListView so we can avoid the resizing of the List entries which caused the "shaking".
#rnx

diff --git a/Engine/Source/Runtime/Slate/Private/Widgets/Layout/SScrollBar.cpp b/Engine/Source/Runtime/Slate/Private/Widgets/Layout/SScrollBar.cpp
--- a/Engine/Source/Runtime/Slate/Private/Widgets/Layout/SScrollBar.cpp
+++ b/Engine/Source/Runtime/Slate/Private/Widgets/Layout/SScrollBar.cpp
@@ -323,7 +323,7 @@ const FSlateBrush* SScrollBar::GetDragThumbImage() const
 EVisibility SScrollBar::ShouldBeVisible() const
 {
 	const EVisibility CurrentVisibility = GetVisibility();
-    EVisibility NewVisibility = EVisibility::Collapsed;
+    EVisibility NewVisibility = ScrollbarDisabledVisibility;
     
 	if ( this->HasMouseCapture() )
 	{
@@ -421,6 +421,11 @@ void SScrollBar::SetScrollBarTrackAlwaysVisible(bool InAlwaysVisible)
 	bAlwaysShowScrollbarTrack = InAlwaysVisible;
 }
 
+SLATE_API void SScrollBar::SetScrollbarDisabledVisibility(EVisibility InVisibility)
+{
+	ScrollbarDisabledVisibility = InVisibility;
+}
+
 bool SScrollBar::AlwaysShowScrollbar() const
 {
 	return bAlwaysShowScrollbar;
diff --git a/Engine/Source/Runtime/Slate/Public/Widgets/Layout/SScrollBar.h b/Engine/Source/Runtime/Slate/Public/Widgets/Layout/SScrollBar.h
--- a/Engine/Source/Runtime/Slate/Public/Widgets/Layout/SScrollBar.h
+++ b/Engine/Source/Runtime/Slate/Public/Widgets/Layout/SScrollBar.h
@@ -145,6 +145,9 @@ public:
 	/** Set ScrollBarTrackAlwaysVisible attribute */
 	SLATE_API void SetScrollBarTrackAlwaysVisible(bool InAlwaysVisible);
 
+	/** Set the visibility of the ScrollBar when it is not needed. The default value is EVisibility::Collapsed. */
+	SLATE_API void SetScrollbarDisabledVisibility(EVisibility InVisibility);
+
 	/** Returns True when the scrollbar should always be shown, else False */
 	SLATE_API bool AlwaysShowScrollbar() const;
 
@@ -187,6 +190,7 @@ protected:
 	bool bAlwaysShowScrollbarTrack;
 	EFocusCause DragFocusCause;
 	bool bHideWhenNotInUse;
+	EVisibility ScrollbarDisabledVisibility = EVisibility::Collapsed;
 	/*
 	 * Holds whether or not to prevent throttling during mouse capture
 	 * When true, the viewport will be updated with every single change to the value during dragging
diff --git a/Engine/Source/Runtime/Slate/Public/Widgets/Views/STileView.h b/Engine/Source/Runtime/Slate/Public/Widgets/Views/STileView.h
--- a/Engine/Source/Runtime/Slate/Public/Widgets/Views/STileView.h
+++ b/Engine/Source/Runtime/Slate/Public/Widgets/Views/STileView.h
@@ -58,6 +58,7 @@ public:
 		, _ScrollbarDragFocusCause(EFocusCause::Mouse)
 		, _AllowOverscroll(EAllowOverscroll::Yes)
 		, _ScrollBarStyle(&FAppStyle::Get().GetWidgetStyle<FScrollBarStyle>("ScrollBar"))
+		, _ScrollbarDisabledVisibility(EVisibility::Collapsed)
 		, _ConsumeMouseWheel(EConsumeMouseWheel::WhenScrollingPossible)
 		, _WheelScrollMultiplier(GetGlobalScrollAmount())
 		, _HandleGamepadEvents(true)
@@ -116,6 +117,8 @@ public:
 
 		SLATE_STYLE_ARGUMENT( FScrollBarStyle, ScrollBarStyle );
 
+		SLATE_ARGUMENT( EVisibility, ScrollbarDisabledVisibility );
+
 		SLATE_ARGUMENT( EConsumeMouseWheel, ConsumeMouseWheel );
 
 		SLATE_ARGUMENT( float, WheelScrollMultiplier );
@@ -209,6 +212,7 @@ public:
 			{
 				this->ScrollBar->SetDragFocusCause(InArgs._ScrollbarDragFocusCause);
 				this->ScrollBar->SetUserVisibility(InArgs._ScrollbarVisibility);
+				this->ScrollBar->SetScrollbarDisabledVisibility(InArgs._ScrollbarDisabledVisibility);
 			}
 			this->AddMetadata(MakeShared<TTableViewMetadata<ItemType>>(this->SharedThis(this)));
 		}
diff --git a/Engine/Source/Runtime/UMG/Public/Components/ListViewBase.h b/Engine/Source/Runtime/UMG/Public/Components/ListViewBase.h
--- a/Engine/Source/Runtime/UMG/Public/Components/ListViewBase.h
+++ b/Engine/Source/Runtime/UMG/Public/Components/ListViewBase.h
@@ -278,6 +278,7 @@ protected:
 		TAttribute<float> EntryWidth;
 		bool bWrapDirectionalNavigation = false;
 		const FScrollBarStyle* ScrollBarStyle = &FUMGCoreStyle::Get().GetWidgetStyle<FScrollBarStyle>("ScrollBar");
+		EVisibility ScrollbarDisabledVisibility = EVisibility::Collapsed;
 	};
 
 	template <template<typename> class TileViewT = STileView, typename UListViewBaseT>
@@ -299,6 +300,7 @@ protected:
 			.ItemAlignment(Args.TileAlignment)
 			.Orientation(Args.Orientation)
 			.ScrollBarStyle(Args.ScrollBarStyle)
+			.ScrollbarDisabledVisibility(Args.ScrollbarDisabledVisibility)
 			.OnGenerateTile_UObject(Implementer, &UListViewBaseT::HandleGenerateRow)
 			.OnTileReleased_UObject(Implementer, &UListViewBaseT::HandleRowReleased)
 			.OnSelectionChanged_UObject(Implementer, &UListViewBaseT::HandleSelectionChanged)
diff --git a/Engine/Source/Runtime/UMG/Public/Components/TileView.h b/Engine/Source/Runtime/UMG/Public/Components/TileView.h
--- a/Engine/Source/Runtime/UMG/Public/Components/TileView.h
+++ b/Engine/Source/Runtime/UMG/Public/Components/TileView.h
@@ -58,6 +58,7 @@ protected:
 		Args.bWrapDirectionalNavigation = bWrapHorizontalNavigation;
 		Args.Orientation = Orientation;
 		Args.ScrollBarStyle = &ScrollBarStyle;
+		Args.ScrollbarDisabledVisibility = UWidget::ConvertSerializedVisibilityToRuntime(ScrollbarDisabledVisibility);
 
 		if (IsAligned() && !bEntrySizeIncludesEntrySpacing)
 		{
@@ -101,6 +102,10 @@ protected:
 	UPROPERTY(EditAnywhere, Category = Navigation)
 	bool bWrapHorizontalNavigation = false;
 
+	/** Set the visibility of the Scrollbar when it's not needed */
+	UPROPERTY(EditAnywhere, Category = ListView, meta=(ValidEnumValues="Collapsed, Hidden, Visible"))
+	ESlateVisibility ScrollbarDisabledVisibility = ESlateVisibility::Collapsed;
+
 	TSharedPtr<STileView<UObject*>> MyTileView;
 
 private:
