Accept a dock to be close when the manager is not present.

diff --git a/Engine/Source/Runtime/Slate/Private/Widgets/Docking/SDockTab.cpp b/Engine/Source/Runtime/Slate/Private/Widgets/Docking/SDockTab.cpp
--- a/Engine/Source/Runtime/Slate/Private/Widgets/Docking/SDockTab.cpp
+++ b/Engine/Source/Runtime/Slate/Private/Widgets/Docking/SDockTab.cpp
@@ -936,7 +936,10 @@ FText SDockTab::GetCloseButtonToolTipText() const
 
 EVisibility SDockTab::HandleIsCloseButtonVisible() const
 {
-	return bCanEverClose && ((IsHovered() || IsForeground()) && MyTabManager.Pin()->IsTabCloseable(SharedThis(this))) ? EVisibility::Visible : EVisibility::Hidden;
+	TSharedPtr<FTabManager> Manager = MyTabManager.Pin();
+	const bool bIsVisibled = IsHovered() || IsForeground();
+	const bool bIsTabCloseable = Manager == nullptr || Manager->IsTabCloseable(SharedThis(this));
+	return bCanEverClose && bIsVisibled && bIsTabCloseable ? EVisibility::Visible : EVisibility::Hidden;
 }
 
 TOptional<FVector2D> SDockTab::GetTabIconSize() const
@@ -946,8 +949,10 @@ TOptional<FVector2D> SDockTab::GetTabIconSize() const
 
 bool SDockTab::CanCloseTab() const
 {
-	const bool bCanCloseTabNow = MyTabManager.Pin()->IsTabCloseable(SharedThis(this)) && (!OnCanCloseTab.IsBound() || OnCanCloseTab.Execute());
-	return bCanCloseTabNow;
+	TSharedPtr<FTabManager> Manager = MyTabManager.Pin();
+	const bool bIsTabCloseable = Manager == nullptr || Manager->IsTabCloseable(SharedThis(this));
+	const bool bCanCloseTab = !OnCanCloseTab.IsBound() || OnCanCloseTab.Execute();
+	return bIsTabCloseable && bCanCloseTab;
 }
 
 bool SDockTab::RequestCloseTab()
