ComboBox: hide the menu content wrapper background for combo boxes to allow for making transparent rows.

diff --git a/Engine/Source/Runtime/Slate/Private/Framework/Application/MenuStack.cpp b/Engine/Source/Runtime/Slate/Private/Framework/Application/MenuStack.cpp
--- a/Engine/Source/Runtime/Slate/Private/Framework/Application/MenuStack.cpp
+++ b/Engine/Source/Runtime/Slate/Private/Framework/Application/MenuStack.cpp
@@ -82,6 +82,7 @@ namespace MenuStackInternal
 				, _OnKeyDown()
 				, _OptionalMinMenuWidth()
 				, _OptionalMinMenuHeight()
+				, _bShowBackground(true)
 			{}
 
 			SLATE_DEFAULT_SLOT(FArguments, MenuContent)
@@ -89,6 +90,7 @@ namespace MenuStackInternal
 			SLATE_EVENT(FOnMenuLostFocus, OnMenuLostFocus)
 			SLATE_ARGUMENT(FOptionalSize, OptionalMinMenuWidth)
 			SLATE_ARGUMENT(FOptionalSize, OptionalMinMenuHeight)
+			SLATE_ARGUMENT(bool, bShowBackground)
 		SLATE_END_ARGS()
 
 		/** Construct this widget */
@@ -99,14 +101,12 @@ namespace MenuStackInternal
 
 			OnKeyDownDelegate = InArgs._OnKeyDown;
 			OnMenuLostFocus = InArgs._OnMenuLostFocus;
-			ChildSlot
-			[
-				SNew(SBox)
-				.MinDesiredWidth(InArgs._OptionalMinMenuWidth)
-				.MaxDesiredHeight(InArgs._OptionalMinMenuHeight)
-				[
-					// Always add a background to the menu.  This includes a small outline around the background to distinguish open menus from each other
-					SNew(SOverlay)
+
+			TSharedPtr<SWidget> ChildContent;
+			if (InArgs._bShowBackground)
+			{
+				// Always add a background to the menu. This includes a small outline around the background to distinguish open menus from each other
+				ChildContent = SNew(SOverlay)
 					+ SOverlay::Slot()
 					[
 						SNew(SImage)
@@ -126,8 +126,31 @@ namespace MenuStackInternal
 						[
 							InArgs._MenuContent.Widget
 						]
+					];
+			}
+			else
+			{
+				ChildContent = SNew(SOverlay)
+					+ SOverlay::Slot()
+					[
+						SNew(SBorder)
+						.Padding(0.f)
+						.BorderImage(FStyleDefaults::GetNoBrush())
+						.ForegroundColor(FCoreStyle::Get().GetSlateColor("DefaultForeground"))
+						[
+							InArgs._MenuContent.Widget
+						]
+					];
+			}
+
+			ChildSlot
+				[
+					SNew(SBox)
+					.MinDesiredWidth(InArgs._OptionalMinMenuWidth)
+					.MaxDesiredHeight(InArgs._OptionalMinMenuHeight)
+					[
+						ChildContent.ToSharedRef()
 					]
-				]
 			];
 		}
 
@@ -281,7 +304,7 @@ TSharedRef<IMenu> FMenuStack::PushHosted(const TSharedPtr<IMenu>& InParentMenu,
 	check(HostWindow.IsValid());
 
 	// Create a FMenuInHostWidget
-	TSharedRef<SWidget> WrappedContent = WrapContent(InContent);
+	TSharedRef<SWidget> WrappedContent = WrapContent(InContent, FOptionalSize(), FOptionalSize(), InMenuHost->bShowMenuBackground);
 	TSharedRef<FMenuInHostWidget> OutMenu = MakeShareable(new FMenuInHostWidget(InMenuHost, WrappedContent, bIsCollapsedByParent));
 	PendingNewMenu = OutMenu;
 
@@ -731,7 +754,7 @@ void FMenuStack::OnMenuContentLostFocus(const FWidgetPath& InFocussedPath)
 	}
 }
 
-TSharedRef<SWidget> FMenuStack::WrapContent(TSharedRef<SWidget> InContent, FOptionalSize OptionalMinWidth, FOptionalSize OptionalMinHeight)
+TSharedRef<SWidget> FMenuStack::WrapContent(TSharedRef<SWidget> InContent, FOptionalSize OptionalMinWidth, FOptionalSize OptionalMinHeight, bool bShouldShowBackground)
 {
 	// Wrap menu content in a box that limits its maximum height
 	// and in a SMenuContentWrapper that handles key presses and focus changes.
@@ -740,6 +763,7 @@ TSharedRef<SWidget> FMenuStack::WrapContent(TSharedRef<SWidget> InContent, FOpti
 		.OnMenuLostFocus_Raw(this, &FMenuStack::OnMenuContentLostFocus)
 		.OptionalMinMenuWidth(OptionalMinWidth)
 		.OptionalMinMenuHeight(OptionalMinHeight)
+		.bShowBackground(bShouldShowBackground)
 		.MenuContent()
 		[
 				InContent
diff --git a/Engine/Source/Runtime/Slate/Public/Framework/Application/IMenu.h b/Engine/Source/Runtime/Slate/Public/Framework/Application/IMenu.h
--- a/Engine/Source/Runtime/Slate/Public/Framework/Application/IMenu.h
+++ b/Engine/Source/Runtime/Slate/Public/Framework/Application/IMenu.h
@@ -32,4 +32,6 @@ public:
 	virtual TSharedPtr<SWindow> GetMenuWindow() const = 0;
 	virtual void OnMenuDismissed() = 0;
 	virtual bool UsingApplicationMenuStack() const = 0;
+
+	bool bShowMenuBackground = true;
 };
diff --git a/Engine/Source/Runtime/Slate/Public/Framework/Application/MenuStack.h b/Engine/Source/Runtime/Slate/Public/Framework/Application/MenuStack.h
--- a/Engine/Source/Runtime/Slate/Public/Framework/Application/MenuStack.h
+++ b/Engine/Source/Runtime/Slate/Public/Framework/Application/MenuStack.h
@@ -246,8 +246,9 @@ private:
 	 * @param	InContent			The unwrapped content.
 	 * @param	OptionalMinWidth	Optional minimum width for the wrapped content.
 	 * @param	OptionalMinHeight	Optional minimum height for the wrapped content.
+	 * @param	bShouldShowBackground	Whether the menu wrapper should put a background behind the content.
 	 */
-	TSharedRef<SWidget> WrapContent(TSharedRef<SWidget> InContent, FOptionalSize OptionalMinWidth = FOptionalSize(), FOptionalSize OptionalMinHeight = FOptionalSize());
+	TSharedRef<SWidget> WrapContent(TSharedRef<SWidget> InContent, FOptionalSize OptionalMinWidth = FOptionalSize(), FOptionalSize OptionalMinHeight = FOptionalSize(), bool bShouldShowBackground = true);
 
 	/** Contains all the options passed to the pre-push stage of the menu creation process */
 	struct FPrePushArgs
diff --git a/Engine/Source/Runtime/Slate/Public/Widgets/Input/SComboBox.h b/Engine/Source/Runtime/Slate/Public/Widgets/Input/SComboBox.h
--- a/Engine/Source/Runtime/Slate/Public/Widgets/Input/SComboBox.h
+++ b/Engine/Source/Runtime/Slate/Public/Widgets/Input/SComboBox.h
@@ -196,6 +196,7 @@ public:
 		ItemStyle = InArgs._ItemStyle;
 		ComboBoxStyle = InArgs._ComboBoxStyle;
 		MenuRowPadding = ComboBoxStyle->MenuRowPadding;
+		bShowMenuBackground = false;
 
 		// Work out which values we should use based on whether we were given an override, or should use the style's version
 		const FComboButtonStyle& OurComboButtonStyle = ComboBoxStyle->ComboButtonStyle;
@@ -268,6 +269,7 @@ public:
 			ComboListView->RequestScrollIntoView(ValidatedItem, 0);
 		}
 
+		ComboListView->SetBackgroundBrush(FStyleDefaults::GetNoBrush());
 	}
 
 		SComboBox()
