Text Overflow policy: added a new ellipsis mode, LineEllipsis, that avoid a possible issue occuring with multiline text and the Ellipsis mode.
With Ellipsis, all lines that are fully or partially visible will be displayed, the drawback of that being that if the "..." are displayed at the end of a line from which the lower part is clipped, the "..." are not visible.
With LineEllipsis, only lines that are fully visible vertically will be displayed. So if a line would be clipped at its lower part, the whole line will be culled, and the "..." will be displayed at the end of the last fully visible line, preventing the case where the "..." might not be seen.

#tests Added new tests in the FontOverflow functionnal test from EngineTest. Visually checked in UnrealEditor that there were no regressions.

#virtualized

diff --git a/Engine/Source/Runtime/Slate/Private/Framework/Text/SlateTextLayout.cpp b/Engine/Source/Runtime/Slate/Private/Framework/Text/SlateTextLayout.cpp
--- a/Engine/Source/Runtime/Slate/Private/Framework/Text/SlateTextLayout.cpp
+++ b/Engine/Source/Runtime/Slate/Private/Framework/Text/SlateTextLayout.cpp
@@ -69,7 +69,7 @@ int32 FSlateTextLayout::OnPaint( const FPaintArgs& Args, const FGeometry& Allott
 
 	int32 HighestLayerId = LayerId;
 
-	auto IsVisible = [&AllottedGeometry, &InverseScale, &MyCullingRect](FVector2D Offset, FVector2D Size)
+	auto IsAtLeastPartiallyVisible = [&AllottedGeometry, &InverseScale, &MyCullingRect](FVector2D Offset, FVector2D Size)
 	{
 		// Is this line visible?  This checks if the culling rect, which represents the AABB around the last clipping rect, intersects the 
 		// line of text, this requires that we get the text line into render space.
@@ -80,25 +80,60 @@ int32 FSlateTextLayout::OnPaint( const FPaintArgs& Args, const FGeometry& Allott
 		return FSlateRect::DoRectanglesIntersect(Rect, MyCullingRect);
 	};
 
+	auto IsFullyVisibleOnVerticalAxis = [&AllottedGeometry, &InverseScale, &MyCullingRect](FVector2D Offset, FVector2D Size)
+	{
+		const FVector2D LocalLineOffset = Offset * InverseScale;
+		FSlateRect Rect(AllottedGeometry.GetRenderBoundingRect(FSlateRect(LocalLineOffset, LocalLineOffset + (Size * InverseScale))));
+
+		//Clamp on the horizontal axis because we care only about vertical one.
+		Rect.Left = MyCullingRect.Left;
+		Rect.Right = MyCullingRect.Right;
+		return FSlateRect::IsRectangleContained(MyCullingRect, Rect);
+	};
+
 
-	int32 LastVisibleLineIndex = -1;
-	for (int32 LineIndex = 0; LineIndex < LineViews.Num(); ++LineIndex)
-	{	//Depending on their length, some lines can be invisible, but following ones could be visible, so we have to parse all of them to get the real last one.
-		if (IsVisible(LineViews[LineIndex].Offset, LineViews[LineIndex].Size))
+	const ETextOverflowPolicy OverflowPolicy = TextOverflowPolicyOverride.Get(DefaultTextStyle.OverflowPolicy);
+	const bool bIsMultiline = LineViews.Num() > 1;
+
+	int32 LineToDisplayNum = LineViews.Num();
+	int32 LastLineIndexToDisplay = 0;
+	if (bIsMultiline)
+	{
+		if (OverflowPolicy == ETextOverflowPolicy::MultilineEllipsis)
+		{
+			for (int32 LineIndex = 0; LineIndex < LineViews.Num(); ++LineIndex)
+			{	//Depending on their length, some lines can be invisible, but following ones could be visible, so we have to parse all of them to get the real last one.
+				if (IsFullyVisibleOnVerticalAxis(LineViews[LineIndex].Offset, LineViews[LineIndex].Size))
+				{
+					LastLineIndexToDisplay = LineIndex;
+				}
+			}
+			LineToDisplayNum = LastLineIndexToDisplay + 1;
+		}
+		else if (OverflowPolicy == ETextOverflowPolicy::Ellipsis)
 		{
-			LastVisibleLineIndex = LineIndex;
+			for (int32 LineIndex = 0; LineIndex < LineViews.Num(); ++LineIndex)
+			{	//Depending on their length, some lines can be invisible, but following ones could be visible, so we have to parse all of them to get the real last one.
+				if (IsAtLeastPartiallyVisible(LineViews[LineIndex].Offset, LineViews[LineIndex].Size))
+				{
+					LastLineIndexToDisplay = LineIndex;
+				}
+			}
+			LineToDisplayNum = LastLineIndexToDisplay + 1;
+		}
+		else
+		{
+			LastLineIndexToDisplay = LineToDisplayNum - 1;
 		}
 	}
 
-	const bool bIsMultiline = LineViews.Num() > 1;
-	for (int32 LineIndex = 0; LineIndex < LineViews.Num(); ++LineIndex)
+	for (int32 LineIndex = 0; LineIndex < LineToDisplayNum; ++LineIndex)
 	{
 		const FTextLayout::FLineView& LineView = LineViews[LineIndex];
+		ETextOverflowPolicy LineOverflowPolicy = OverflowPolicy;
+		ETextOverflowDirection LineOverflowDirection = ETextOverflowDirection::NoOverflow;
 
-		ETextOverflowPolicy OverflowPolicy = TextOverflowPolicyOverride.Get(DefaultTextStyle.OverflowPolicy);
-		ETextOverflowDirection OverflowDirection = ETextOverflowDirection::NoOverflow;
-
-		if (!IsVisible(LineView.Offset, LineView.Size))
+		if (!IsAtLeastPartiallyVisible(LineView.Offset, LineView.Size))
 		{
 			continue;
 		}
@@ -112,21 +147,21 @@ int32 FSlateTextLayout::OnPaint( const FPaintArgs& Args, const FGeometry& Allott
 
 		bool bIsLastVisibleLine = false;
 		bool bIsNextLineClipped = false;
-		if (OverflowPolicy == ETextOverflowPolicy::Ellipsis)
+		if (LineOverflowPolicy == ETextOverflowPolicy::Ellipsis || LineOverflowPolicy == ETextOverflowPolicy::MultilineEllipsis)
 		{
 			if (bIsMultiline)
 			{
-				bIsLastVisibleLine = LastVisibleLineIndex == LineIndex;
+				bIsLastVisibleLine = LineIndex == LastLineIndexToDisplay;
 				bIsNextLineClipped = LineViews.IsValidIndex(LineIndex + 1) ? bIsLastVisibleLine : false;
 
 				//When wrapping/multiline text, we have to use the reading direction of the text, not the justification.
-				OverflowDirection = LineView.TextBaseDirection == TextBiDi::ETextDirection::LeftToRight ? ETextOverflowDirection::LeftToRight : ETextOverflowDirection::RightToLeft;
+				LineOverflowDirection = LineView.TextBaseDirection == TextBiDi::ETextDirection::LeftToRight ? ETextOverflowDirection::LeftToRight : ETextOverflowDirection::RightToLeft;
 			}
 			else
 			{
 				bIsLastVisibleLine = true;
 				const ETextJustify::Type VisualJustification = CalculateLineViewVisualJustification(LineView);
-				OverflowDirection = VisualJustification == ETextJustify::Left ? ETextOverflowDirection::LeftToRight : (VisualJustification == ETextJustify::Right ? ETextOverflowDirection::RightToLeft : ETextOverflowDirection::NoOverflow);
+				LineOverflowDirection = VisualJustification == ETextJustify::Left ? ETextOverflowDirection::LeftToRight : (VisualJustification == ETextJustify::Right ? ETextOverflowDirection::RightToLeft : ETextOverflowDirection::NoOverflow);
 			}
 		}
 
@@ -136,12 +171,12 @@ int32 FSlateTextLayout::OnPaint( const FPaintArgs& Args, const FGeometry& Allott
 		{
 			int32 StartValue = 0;
 			int32 EndValue = 0;
-			if (OverflowDirection == ETextOverflowDirection::LeftToRight)
+			if (LineOverflowDirection == ETextOverflowDirection::LeftToRight)
 			{
 				EndValue = LineView.Blocks.Num();
 				NextBlockOffset = 1;
 			}
-			else if (OverflowDirection == ETextOverflowDirection::RightToLeft)
+			else if (LineOverflowDirection == ETextOverflowDirection::RightToLeft)
 			{
 				StartValue = LineView.Blocks.Num() - 1;
 				EndValue = -1;
@@ -151,7 +186,7 @@ int32 FSlateTextLayout::OnPaint( const FPaintArgs& Args, const FGeometry& Allott
 			for (int32 BlockIndex = StartValue; BlockIndex != EndValue; BlockIndex += NextBlockOffset)
 			{	//Depending on their position, some blocks can be invisible, but a following one could be visible, so we have to parse all of them.
 				const TSharedRef< ILayoutBlock >& Block = LineView.Blocks[BlockIndex];
-				if (IsVisible(Block->GetLocationOffset(), Block->GetSize()))
+				if (IsAtLeastPartiallyVisible(Block->GetLocationOffset(), Block->GetSize()))
 				{
 					LastVisibleBlockIndex = BlockIndex;
 				}
@@ -159,15 +194,15 @@ int32 FSlateTextLayout::OnPaint( const FPaintArgs& Args, const FGeometry& Allott
 		}
 		else
 		{
-			OverflowDirection = ETextOverflowDirection::NoOverflow;
-			OverflowPolicy = ETextOverflowPolicy::Clip;
+			LineOverflowDirection = ETextOverflowDirection::NoOverflow;
+			LineOverflowPolicy = ETextOverflowPolicy::Clip;
 		}
 
 		// Render every block for this line
 		for (int32 BlockIndex = 0; BlockIndex < LineView.Blocks.Num(); ++BlockIndex)
 		{
 			const TSharedRef< ILayoutBlock >& Block = LineView.Blocks[BlockIndex];
-			if (!IsVisible(Block->GetLocationOffset(), Block->GetSize()))
+			if (!IsAtLeastPartiallyVisible(Block->GetLocationOffset(), Block->GetSize()))
 			{
 				continue;
 			}
@@ -188,7 +223,7 @@ int32 FSlateTextLayout::OnPaint( const FPaintArgs& Args, const FGeometry& Allott
 			}
 #endif
 
-			FTextArgs TextArgs(LineView, Block, DefaultTextStyle, OverflowPolicy, OverflowDirection);
+			FTextArgs TextArgs(LineView, Block, DefaultTextStyle, LineOverflowPolicy, LineOverflowDirection);
 			TextArgs.bIsLastVisibleBlock = bIsLastVisibleLine && BlockIndex == LastVisibleBlockIndex;
 			TextArgs.bIsNextBlockClipped = bIsNextLineClipped && (LineView.Blocks.IsValidIndex(BlockIndex + NextBlockOffset) ? TextArgs.bIsLastVisibleBlock : true);
 
diff --git a/Engine/Source/Runtime/Slate/Private/Framework/Text/SlateTextRun.cpp b/Engine/Source/Runtime/Slate/Private/Framework/Text/SlateTextRun.cpp
--- a/Engine/Source/Runtime/Slate/Private/Framework/Text/SlateTextRun.cpp
+++ b/Engine/Source/Runtime/Slate/Private/Framework/Text/SlateTextRun.cpp
@@ -119,7 +119,7 @@ int32 FSlateTextRun::OnPaint(const FPaintArgs& PaintArgs, const FTextArgs& TextA
 	);
 
 	FTextOverflowArgs OverflowArgs;
-	if (TextArgs.OverflowPolicy == ETextOverflowPolicy::Ellipsis && TextArgs.OverflowDirection != ETextOverflowDirection::NoOverflow)
+	if ((TextArgs.OverflowPolicy == ETextOverflowPolicy::Ellipsis || TextArgs.OverflowPolicy == ETextOverflowPolicy::MultilineEllipsis) && TextArgs.OverflowDirection != ETextOverflowDirection::NoOverflow)
 	{
 		OverflowArgs.OverflowTextPtr = BlockTextContext.ShapedTextCache->FindOrAddOverflowEllipsisText(AllottedGeometry.GetAccumulatedLayoutTransform().GetScale(), BlockTextContext, Style.Font);
 		OverflowArgs.OverflowDirection = TextArgs.OverflowDirection;
diff --git a/Engine/Source/Runtime/SlateCore/Public/Styling/SlateTypes.h b/Engine/Source/Runtime/SlateCore/Public/Styling/SlateTypes.h
--- a/Engine/Source/Runtime/SlateCore/Public/Styling/SlateTypes.h
+++ b/Engine/Source/Runtime/SlateCore/Public/Styling/SlateTypes.h
@@ -82,6 +82,9 @@ enum class ETextOverflowPolicy : uint8
 
 	/** Overflowing text will be replaced with an ellipsis */
 	Ellipsis,
+
+	/** Overflowing text will be replaced with an ellipsis. A partially clipped line on the vertical axis will be totally clipped, and ellipsis displayed on previous line */
+	MultilineEllipsis,
 };
 
 
