PR #11213: [Slate] SBox: Force contents to be center-aligned, not left-aligned, when using Min/MaxAspectRatio

diff --git a/Engine/Source/Runtime/Slate/Private/Widgets/Layout/SBox.cpp b/Engine/Source/Runtime/Slate/Private/Widgets/Layout/SBox.cpp
--- a/Engine/Source/Runtime/Slate/Private/Widgets/Layout/SBox.cpp
+++ b/Engine/Source/Runtime/Slate/Private/Widgets/Layout/SBox.cpp
@@ -190,84 +190,96 @@ void SBox::OnArrangeChildren( const FGeometry& AllottedGeometry, FArrangedChildr
 		const FOptionalSize CurrentMinAspectRatio = MinAspectRatio.Get();
 		const FOptionalSize CurrentMaxAspectRatio = MaxAspectRatio.Get();
 		const FMargin SlotPadding(ChildSlot.GetPadding());
-		bool bAlignChildren = true;
 
-		AlignmentArrangeResult XAlignmentResult(0, 0);
-		AlignmentArrangeResult YAlignmentResult(0, 0);
+		AlignmentArrangeResult XAlignmentResult = AlignChild<Orient_Horizontal>(AllottedGeometry.GetLocalSize().X, ChildSlot, SlotPadding);
+		AlignmentArrangeResult YAlignmentResult = AlignChild<Orient_Vertical>(AllottedGeometry.GetLocalSize().Y, ChildSlot, SlotPadding);
 
 		if (CurrentMaxAspectRatio.IsSet() || CurrentMinAspectRatio.IsSet())
 		{
-			float CurrentWidth = FMath::Min(AllottedGeometry.Size.X, ChildSlot.GetWidget()->GetDesiredSize().X);
-			float CurrentHeight = FMath::Min(AllottedGeometry.Size.Y, ChildSlot.GetWidget()->GetDesiredSize().Y);
-
-			float MinAspectRatioWidth = CurrentMinAspectRatio.IsSet() ? CurrentMinAspectRatio.Get() : 0;
-			float MaxAspectRatioWidth = CurrentMaxAspectRatio.IsSet() ? CurrentMaxAspectRatio.Get() : 0;
-			if (CurrentHeight > 0 && CurrentWidth > 0)
+			const FVector2f CurrentSize = FVector2f(XAlignmentResult.Size, YAlignmentResult.Size);
+			const float MinAspectRatioValue = CurrentMinAspectRatio.IsSet() ? CurrentMinAspectRatio.Get() : 0.0f;
+			const float MaxAspectRatioValue = CurrentMaxAspectRatio.IsSet() ? CurrentMaxAspectRatio.Get() : 0.0f;
+			if (CurrentSize.X > 0.0f && CurrentSize.Y > 0.0f)
 			{
-				const float CurrentRatioWidth = (AllottedGeometry.GetLocalSize().X / AllottedGeometry.GetLocalSize().Y);
-
-				bool bFitMaxRatio = (CurrentRatioWidth > MaxAspectRatioWidth && MaxAspectRatioWidth != 0);
-				bool bFitMinRatio = (CurrentRatioWidth < MinAspectRatioWidth && MinAspectRatioWidth != 0);
+				const float CurrentRatio = CurrentSize.X / CurrentSize.Y;
+				const bool bFitMaxRatio = (CurrentRatio > MaxAspectRatioValue && MaxAspectRatioValue != 0.0f);
+				const bool bFitMinRatio = (CurrentRatio < MinAspectRatioValue && MinAspectRatioValue != 0.0f);
 				if (bFitMaxRatio || bFitMinRatio)
 				{
-					XAlignmentResult = AlignChild<Orient_Horizontal>(AllottedGeometry.GetLocalSize().X, ChildSlot, SlotPadding);
-					YAlignmentResult = AlignChild<Orient_Vertical>(AllottedGeometry.GetLocalSize().Y, ChildSlot, SlotPadding);
 
-					float NewWidth;
-					float NewHeight;
+					FVector2f MaxSize = FVector2f(AllottedGeometry.Size.X - SlotPadding.GetTotalSpaceAlong<Orient_Horizontal>(), AllottedGeometry.Size.Y - SlotPadding.GetTotalSpaceAlong<Orient_Vertical>());
+
+					const bool bXFillAlignment = ArrangeUtils::GetChildAlignment<Orient_Horizontal>::AsInt(EFlowDirection::LeftToRight, ChildSlot) == HAlign_Fill;
+					const bool bYFillAlignment = ArrangeUtils::GetChildAlignment<Orient_Vertical>::AsInt(EFlowDirection::LeftToRight, ChildSlot) == HAlign_Fill;
 
+					FVector2f NewSize = MaxSize;
+					auto RatioPredicate = [&](float MinMaxAspectRatioValue)
+					{
+						const FVector2f AspectRatioValue = FVector2f(MinMaxAspectRatioValue, 1.0f / MinMaxAspectRatioValue);
+						const bool bIsRatioXFirst = (bXFillAlignment && bYFillAlignment) || (!bXFillAlignment && !bYFillAlignment) ? MinMaxAspectRatioValue >= 1.0f : bXFillAlignment;
+						if (bIsRatioXFirst)
+						{
+							NewSize.X = MinMaxAspectRatioValue >= 1.0f ? AspectRatioValue.X * CurrentSize.X : CurrentSize.X / AspectRatioValue.X;
+							NewSize.Y = AspectRatioValue.Y * NewSize.X;
+							MaxSize.X = FMath::Min(AspectRatioValue.X * MaxSize.Y, MaxSize.X);
+							MaxSize.Y = AspectRatioValue.Y * MaxSize.X;
+						}
+						else
+						{
+							NewSize.Y = MinMaxAspectRatioValue >= 1.0f ? CurrentSize.Y / AspectRatioValue.Y : AspectRatioValue.Y * CurrentSize.Y;
+							NewSize.X = AspectRatioValue.X * NewSize.Y;
+							MaxSize.Y = FMath::Min(AspectRatioValue.Y * MaxSize.X, MaxSize.Y);
+							MaxSize.X = AspectRatioValue.X * MaxSize.Y;
+						}
+					};
 					if (bFitMaxRatio)
 					{
-						const float MaxAspectRatioHeight = 1.0f / MaxAspectRatioWidth;
-						NewWidth = MaxAspectRatioWidth * XAlignmentResult.Size;
-						NewHeight = MaxAspectRatioHeight * NewWidth;
+						RatioPredicate(MaxAspectRatioValue);
 					}
 					else
 					{
-						const float MinAspectRatioHeight = 1.0f / MinAspectRatioWidth;
-						NewWidth = MinAspectRatioWidth * XAlignmentResult.Size;
-						NewHeight = MinAspectRatioHeight * NewWidth;
+						RatioPredicate(MinAspectRatioValue);
 					}
 
-					const float MaxWidth = AllottedGeometry.Size.X - SlotPadding.GetTotalSpaceAlong<Orient_Horizontal>();
-					const float MaxHeight = AllottedGeometry.Size.Y - SlotPadding.GetTotalSpaceAlong<Orient_Vertical>();
-
-					if ( NewWidth > MaxWidth )
+					// Make sure they are inside the max available size
+					if (NewSize.X > MaxSize.X)
 					{
-						float Scale = MaxWidth / NewWidth;
-						NewWidth *= Scale;
-						NewHeight *= Scale;
+						float Scale = NewSize.X != 0.0f ? MaxSize.X / NewSize.X : 0.0f;
+						NewSize *= Scale;
 					}
 
-					if ( NewHeight > MaxHeight )
+					if (NewSize.Y > MaxSize.Y)
 					{
-						float Scale = MaxHeight / NewHeight;
-						NewWidth *= Scale;
-						NewHeight *= Scale;
+						float Scale = NewSize.Y != 0.0f ? MaxSize.Y / NewSize.Y : 0.0f;
+						NewSize *= Scale;
 					}
 
-					XAlignmentResult.Size = NewWidth;
-					YAlignmentResult.Size = NewHeight;
-
-					bAlignChildren = false;
+					// The size changed, realign them. If it's Fill, then center it if needed.
+					if (!bXFillAlignment)
+					{
+						XAlignmentResult = AlignChild<Orient_Horizontal>(AllottedGeometry.GetLocalSize().X, NewSize.X, ChildSlot, SlotPadding);
+					}
+					else
+					{
+						XAlignmentResult = ArrangeUtils::AlignCenter<Orient_Horizontal>(AllottedGeometry.GetLocalSize().X, NewSize.X, SlotPadding);
+					}
+					if (!bYFillAlignment)
+					{
+						YAlignmentResult = AlignChild<Orient_Vertical>(AllottedGeometry.GetLocalSize().Y, NewSize.Y, ChildSlot, SlotPadding);
+					}
+					else
+					{
+						YAlignmentResult = ArrangeUtils::AlignCenter<Orient_Vertical>(AllottedGeometry.GetLocalSize().Y, NewSize.Y, SlotPadding);
+					}
 				}
 			}
 		}
 
-		if ( bAlignChildren )
-		{
-			XAlignmentResult = AlignChild<Orient_Horizontal>(AllottedGeometry.GetLocalSize().X, ChildSlot, SlotPadding);
-			YAlignmentResult = AlignChild<Orient_Vertical>(AllottedGeometry.GetLocalSize().Y, ChildSlot, SlotPadding);
-		}
-
-		const float AlignedSizeX = XAlignmentResult.Size;
-		const float AlignedSizeY = YAlignmentResult.Size;
-
 		ArrangedChildren.AddWidget(
 			AllottedGeometry.MakeChild(
 				ChildSlot.GetWidget(),
 				FVector2D(XAlignmentResult.Offset, YAlignmentResult.Offset),
-				FVector2D(AlignedSizeX, AlignedSizeY)
+				FVector2D(XAlignmentResult.Size, YAlignmentResult.Size)
 			)
 		);
 	}
diff --git a/Engine/Source/Runtime/SlateCore/Public/Layout/LayoutUtils.h b/Engine/Source/Runtime/SlateCore/Public/Layout/LayoutUtils.h
--- a/Engine/Source/Runtime/SlateCore/Public/Layout/LayoutUtils.h
+++ b/Engine/Source/Runtime/SlateCore/Public/Layout/LayoutUtils.h
@@ -68,6 +68,34 @@ namespace ArrangeUtils
 			return static_cast<int32>(InSlot.GetVerticalAlignment());
 		}
 	};
+
+	/**
+	 * Same as AlignChild but force the alignment to be fill.
+	 * @return  Offset and Size of widget
+	 */
+	template<EOrientation Orientation>
+	static AlignmentArrangeResult AlignFill(float AllottedSize, const FMargin& SlotPadding, const float ContentScale = 1.0f)
+	{
+		const FMargin& Margin = SlotPadding;
+		const float TotalMargin = Margin.GetTotalSpaceAlong<Orientation>();
+		const float MarginPre = (Orientation == Orient_Horizontal) ? Margin.Left : Margin.Top;
+		return AlignmentArrangeResult(MarginPre, FMath::Max((AllottedSize - TotalMargin) * ContentScale, 0.f));
+	}
+
+	/**
+	 * Same as AlignChild but force the alignment to be center.
+	 * @return  Offset and Size of widget
+	 */
+	template<EOrientation Orientation>
+	static AlignmentArrangeResult AlignCenter(float AllottedSize, float ChildDesiredSize, const FMargin& SlotPadding, const float ContentScale = 1.0f, bool bClampToParent = true)
+	{
+		const FMargin& Margin = SlotPadding;
+		const float TotalMargin = Margin.GetTotalSpaceAlong<Orientation>();
+		const float MarginPre = (Orientation == Orient_Horizontal) ? Margin.Left : Margin.Top;
+		const float MarginPost = (Orientation == Orient_Horizontal) ? Margin.Right : Margin.Bottom;
+		const float ChildSize = FMath::Max((bClampToParent ? FMath::Min(ChildDesiredSize, AllottedSize - TotalMargin) : ChildDesiredSize), 0.f);
+		return AlignmentArrangeResult((AllottedSize - ChildSize) / 2.0f + MarginPre - MarginPost, ChildSize);
+	}
 }
 
 
