TabManager: stylistic code changes
- Untemplated HasAnyMatchingTabs.
- Used lambdas instead of functors to simplify the code.
- Renamed GetFAreaFromInitialLayoutWithTabType to GetAreaFromInitialLayoutWithTabType.
- Renamed GetFAreaForFTabId to GetAreaForTabId.

diff --git a/Engine/Source/Runtime/Slate/Private/Framework/Docking/TabManager.cpp b/Engine/Source/Runtime/Slate/Private/Framework/Docking/TabManager.cpp
--- a/Engine/Source/Runtime/Slate/Private/Framework/Docking/TabManager.cpp
+++ b/Engine/Source/Runtime/Slate/Private/Framework/Docking/TabManager.cpp
@@ -1453,7 +1453,7 @@ TSharedPtr<SDockTab> FTabManager::InvokeTab_Internal(const FTabId& TabId, bool b
 	}
 	else
 	{
-		const TSharedRef<FArea> NewAreaForTab = GetFAreaForFTabId(TabId);
+		const TSharedRef<FArea> NewAreaForTab = GetAreaForTabId(TabId);
 
 		NewAreaForTab
 		->Split
@@ -2214,27 +2214,14 @@ FVector2D FTabManager::GetDefaultTabWindowSize(const FTabId& TabId)
 
 bool FTabManager::HasAnyTabWithTabId( const TSharedRef<FLayoutNode>& SomeNode, const FName& InTabTypeToMatch ) const
 {
-	struct TabWithIdMatcher
-	{
-		TabWithIdMatcher(const FTabManager* InTabManager, FName& InTabTypeToMatch) : TabManager(InTabManager), TabTypeToMatch(InTabTypeToMatch)
-		{
-		}
-		
-		const FTabManager* TabManager;
-		const FName& TabTypeToMatch;
-
-		bool operator()(const FTab& Candidate) const
+	return HasAnyMatchingTabs(SomeNode,
+		[this, InTabTypeToMatch](const FTab& Candidate)
 		{
-			return TabManager->IsValidTabForSpawning(Candidate) && Candidate.TabId.TabType == TabTypeToMatch;
-		}
-	};
-	FName TabTypeToMatch = InTabTypeToMatch;
-	const TabWithIdMatcher FindTabWithId(this, TabTypeToMatch);
-
-	return HasAnyMatchingTabs(SomeNode, FindTabWithId);
+			return this->IsValidTabForSpawning(Candidate) && Candidate.TabId.TabType == InTabTypeToMatch;
+		});
 }
 
-TSharedPtr<FTabManager::FArea> FTabManager::GetFAreaFromInitialLayoutWithTabType( const FTabId& InTabIdToMatch ) const
+TSharedPtr<FTabManager::FArea> FTabManager::GetAreaFromInitialLayoutWithTabType( const FTabId& InTabIdToMatch ) const
 {
 	const TSharedPtr<FTabManager::FLayout> InitialLayoutSP = FGlobalTabmanager::Get()->GetInitialLayoutSP();
 	if (InitialLayoutSP.IsValid())
@@ -2250,9 +2237,9 @@ TSharedPtr<FTabManager::FArea> FTabManager::GetFAreaFromInitialLayoutWithTabType
 	return nullptr;
 }
 
-TSharedRef<FTabManager::FArea> FTabManager::GetFAreaForFTabId(const FTabId& TabId)
+TSharedRef<FTabManager::FArea> FTabManager::GetAreaForTabId(const FTabId& TabId)
 {
-	if (const TSharedPtr<FArea> AreaFromInitiallyLoadedLayout = FGlobalTabmanager::Get()->GetFAreaFromInitialLayoutWithTabType(TabId))
+	if (const TSharedPtr<FArea> AreaFromInitiallyLoadedLayout = FGlobalTabmanager::Get()->GetAreaFromInitialLayoutWithTabType(TabId))
 	{
 		/* we must reuse positions from the initial layout for positionally specified floating windows. If we don't
 		* do this then any persisted floating windows load in a big cluster in the middle on top of one another */
@@ -2274,8 +2261,7 @@ TSharedPtr<FTabManager::FLayout> FGlobalTabmanager::GetInitialLayoutSP()
 	return InitialLayoutSP;
 }
 
-template<typename MatchFunctorType>
-bool FTabManager::HasAnyMatchingTabs( const TSharedRef<FTabManager::FLayoutNode>& SomeNode, const MatchFunctorType& Matcher )
+bool FTabManager::HasAnyMatchingTabs( const TSharedRef<FTabManager::FLayoutNode>& SomeNode, const TFunctionRef<bool(const FTab& Candidate)>& Matcher )
 {
 	TSharedPtr<FTabManager::FSplitter> AsSplitter = SomeNode->AsSplitter();
 	TSharedPtr<FTabManager::FStack> AsStack = SomeNode->AsStack();
@@ -2302,37 +2288,21 @@ bool FTabManager::HasAnyMatchingTabs( const TSharedRef<FTabManager::FLayoutNode>
 bool FTabManager::HasValidOpenTabs( const TSharedRef<FTabManager::FLayoutNode>& SomeNode ) const
 {
 	// Search for valid and open tabs
-	struct OpenTabMatcher
-	{
-		const FTabManager* TabManager;
-
-		bool operator()(const FTab& Candidate) const
+	return HasAnyMatchingTabs(SomeNode,
+		[this](const FTab& Candidate)
 		{
-			return TabManager->IsValidTabForSpawning(Candidate) && Candidate.TabState == ETabState::OpenedTab;
-		}
-	};
-	OpenTabMatcher FindOpenTab;
-	FindOpenTab.TabManager = this;
-
-	return HasAnyMatchingTabs(SomeNode, FindOpenTab);
+				return this->IsValidTabForSpawning(Candidate) && Candidate.TabState == ETabState::OpenedTab;
+		});
 }
 
 bool FTabManager::HasValidTabs( const TSharedRef<FTabManager::FLayoutNode>& SomeNode ) const
 {
 	// Search for valid tabs that can be spawned
-	struct ValidTabMatcher
-	{
-		const FTabManager* TabManager;
-
-		bool operator()(const FTab& Candidate) const
+	return HasAnyMatchingTabs(SomeNode,
+		[this](const FTab& Candidate)
 		{
-			return TabManager->IsValidTabForSpawning(Candidate);
-		}
-	};
-	ValidTabMatcher FindValidTab;
-	FindValidTab.TabManager = this;
-
-	return HasAnyMatchingTabs(SomeNode, FindValidTab);
+			return this->IsValidTabForSpawning(Candidate);
+		});
 }
 
 void FTabManager::SetTabsTo(const TSharedRef<FTabManager::FLayoutNode>& SomeNode, const ETabState::Type NewTabState, const ETabState::Type OriginalTabState) const
@@ -2992,7 +2962,7 @@ void FProxyTabmanager::OpenUnmanagedTab(FName PlaceholderId, const FSearchPrefer
 	TSharedPtr<SWindow> ParentWindowPtr = ParentWindow.Pin();
 	if (ensure(ParentWindowPtr.IsValid()))
 	{
-		const TSharedPtr<FArea> Area = FGlobalTabmanager::Get()->GetFAreaFromInitialLayoutWithTabType(UnmanagedTab->GetLayoutIdentifier());
+		const TSharedPtr<FArea> Area = FGlobalTabmanager::Get()->GetAreaFromInitialLayoutWithTabType(UnmanagedTab->GetLayoutIdentifier());
 		const TSharedRef<FArea> NewAreaForTab =  Area.IsValid() ? Area.ToSharedRef() : NewPrimaryArea();
 
 		NewAreaForTab
diff --git a/Engine/Source/Runtime/Slate/Public/Framework/Docking/TabManager.h b/Engine/Source/Runtime/Slate/Public/Framework/Docking/TabManager.h
--- a/Engine/Source/Runtime/Slate/Public/Framework/Docking/TabManager.h
+++ b/Engine/Source/Runtime/Slate/Public/Framework/Docking/TabManager.h
@@ -10,6 +10,7 @@
 #include "Widgets/SWindow.h"
 #include "Framework/Docking/WorkspaceItem.h"
 #include "Framework/Commands/UIAction.h"
+#include "Templates/Function.h"
 
 class FJsonObject;
 class FMenuBuilder;
@@ -128,7 +129,7 @@ struct FTabId
 
 class FSpawnTabArgs
 {
-	public:
+public:
 	FSpawnTabArgs( const TSharedPtr<SWindow>& InOwnerWindow, const FTabId& InTabBeingSpawenedId )
 	: TabIdBeingSpawned(InTabBeingSpawenedId)
 	, OwnerWindow(InOwnerWindow)
@@ -145,7 +146,7 @@ class FSpawnTabArgs
 		return TabIdBeingSpawned;
 	}
 
-	private:
+private:
 	FTabId TabIdBeingSpawned;
 	TSharedPtr<SWindow> OwnerWindow;
 };
@@ -1106,7 +1107,7 @@ class FTabManager : public TSharedFromThis<FTabManager>
 		SLATE_API TSharedPtr<class SDockingTabStack> FindTabInLiveAreas( const FTabMatcher& TabMatcher ) const;
 		static SLATE_API TSharedPtr<class SDockingTabStack> FindTabInLiveArea( const FTabMatcher& TabMatcher, const TSharedRef<SDockingArea>& InArea );
 
-		template<typename MatchFunctorType> static bool HasAnyMatchingTabs( const TSharedRef<FTabManager::FLayoutNode>& SomeNode, const MatchFunctorType& Matcher );
+		static bool HasAnyMatchingTabs( const TSharedRef<FTabManager::FLayoutNode>& SomeNode, const TFunctionRef<bool(const FTab&)>& Matcher );
 
 	public:
 		/**
@@ -1120,7 +1121,7 @@ class FTabManager : public TSharedFromThis<FTabManager>
 		 *
 		 * @param InTabIdToMatch the const &FTabId for which to find the FArea
 		 */
-	    SLATE_API TSharedPtr<FArea> GetFAreaFromInitialLayoutWithTabType(const FTabId& InTabIdToMatch ) const;
+	    SLATE_API TSharedPtr<FArea> GetAreaFromInitialLayoutWithTabType(const FTabId& InTabIdToMatch ) const;
 
 protected:
 		SLATE_API bool HasValidTabs( const TSharedRef<FTabManager::FLayoutNode>& SomeNode ) const;
@@ -1177,7 +1178,7 @@ protected:
 	 *
 	 * @param TabId the FTabId for which to find the FArea 
 	 */
-	TSharedRef<FArea> GetFAreaForFTabId(const FTabId& TabId);
+	TSharedRef<FArea> GetAreaForTabId(const FTabId& TabId);
 
 	protected:
 		FTabSpawner TabSpawner;
