Slate - Consuming a Mouse Up event with a preprocessor doesn't clean up the previous mouse down event, preventing a new mouse down from being registered until a mouse up clears the first one.
The issue currently only affects the swipe detection while using touch, so I added a IsTouchEvent condition to lessen the potential impact on other platforms. We can remove the condition later if we want to.
When a mouse down event is triggered, if none of the preprocessors consumes the event, the system registers the pointer in the SlateUser::PointerPositionsByIndex map. A mouse up event then clears it to complete the interaction, unless it is consumed by a preprocessor. In that case, the system will refuse any mouse down event because it was effectively never completed.
To prevent this, we notify the slate user that the pointer needs to be released.
Tested using both Touch and Fake Touch.

#RNX

diff --git a/Engine/Source/Runtime/Slate/Private/Framework/Application/SlateApplication.cpp b/Engine/Source/Runtime/Slate/Private/Framework/Application/SlateApplication.cpp
--- a/Engine/Source/Runtime/Slate/Private/Framework/Application/SlateApplication.cpp
+++ b/Engine/Source/Runtime/Slate/Private/Framework/Application/SlateApplication.cpp
@@ -5837,6 +5837,15 @@ bool FSlateApplication::ProcessMouseButtonUpEvent( const FPointerEvent& MouseEve
 	// Input preprocessors get the first chance at the input
 	if (InputPreProcessors.HandleMouseButtonUpEvent(*this, MouseEvent))
 	{
+		// If mouse up event is consumed by a preprocessor, associated mouse down event needs to be cleared as well. Otherwise, subsequent mouse down events get ignored until the first one is cleared.
+		// This was only affecting the swipe detection on touch, we can remove this condition if we want the same fix for other platforms, but reducing the scope for now
+		if (MouseEvent.IsTouchEvent())
+		{
+			TSharedRef<FSlateUser> SlateUser = GetOrCreateUser(MouseEvent);
+			FWidgetPath WidgetsUnderPointer = LocateWindowUnderMouse(MouseEvent.GetScreenSpacePosition(), GetInteractiveTopLevelWindows(), false, SlateUser->GetUserIndex());
+
+			SlateUser->NotifyPointerReleased(MouseEvent, WidgetsUnderPointer, nullptr, true);
+		}
 		return true;
 	}
 
