Show menu search fields by default

* Remove CVAR Slate.AlwaysShowMenuSearchField and instead make the behavior default.
* Add an editor setting under "General - Appearance - User Interface" that controls when the menu search field appears. The setting is an integer and refers to the number of items a menu has to have before the search field appears immediately when the menu is opened. For shorter menus, the search field will be hidden until the user starts typing. The user can set the value to 0 to always show the search field or a very high number to never show the search field by default (this is also described in the settings tooltip). The default value is 10 and will show the menu search field for menus of 10 items or more, otherwise hide it until the user types.

diff --git a/Engine/Source/Editor/UnrealEd/Classes/Settings/EditorStyleSettings.h b/Engine/Source/Editor/UnrealEd/Classes/Settings/EditorStyleSettings.h
--- a/Engine/Source/Editor/UnrealEd/Classes/Settings/EditorStyleSettings.h
+++ b/Engine/Source/Editor/UnrealEd/Classes/Settings/EditorStyleSettings.h
@@ -99,6 +99,10 @@ public:
 	UPROPERTY(EditAnywhere, config, Category=UserInterface)
 	uint32 bUseSmallToolBarIcons:1;
 
+	/** Menus longer than this threshold show their search field by default. Use 0 to always show, or a high number to always hide. When a searchable menu is open but the field is hidden, you can still start a search by typing. */
+	UPROPERTY(EditAnywhere, config, Category=UserInterface, meta = (UIMin="0", UIMax="100"))
+	uint32 MenuSearchFieldVisibilityThreshold = 10;
+
 	/** If true the material editor and blueprint editor will show a grid on it's background. */
 	UPROPERTY(EditAnywhere, config, Category = Graphs, meta = (DisplayName = "Use Grids In The Material And Blueprint Editor"))
 	uint32 bUseGrid : 1;
diff --git a/Engine/Source/Editor/UnrealEd/Private/EditorEngine.cpp b/Engine/Source/Editor/UnrealEd/Private/EditorEngine.cpp
--- a/Engine/Source/Editor/UnrealEd/Private/EditorEngine.cpp
+++ b/Engine/Source/Editor/UnrealEd/Private/EditorEngine.cpp
@@ -754,6 +754,11 @@ static bool GetDisplayMultiboxHooks()
 	return GetDefault<UEditorPerProjectUserSettings>()->bDisplayUIExtensionPoints;
 }
 
+static int GetMenuSearchFieldVisibilityThreshold()
+{
+	return GetDefault<UEditorStyleSettings>()->MenuSearchFieldVisibilityThreshold;
+}
+
 void UEditorEngine::InitEditor(IEngineLoop* InEngineLoop)
 {
 	// Allow remote execution of derived data builds from this point
@@ -777,6 +782,7 @@ void UEditorEngine::InitEditor(IEngineLoop* InEngineLoop)
 	// Set slate options
 	FMultiBoxSettings::UseSmallToolBarIcons = TAttribute<bool>::Create(TAttribute<bool>::FGetter::CreateStatic(&GetSmallToolBarIcons));
 	FMultiBoxSettings::DisplayMultiboxHooks = TAttribute<bool>::Create(TAttribute<bool>::FGetter::CreateStatic(&GetDisplayMultiboxHooks));
+	FMultiBoxSettings::MenuSearchFieldVisibilityThreshold = TAttribute<int>::Create(TAttribute<int>::FGetter::CreateStatic(&GetMenuSearchFieldVisibilityThreshold));
 
 	if ( FSlateApplication::IsInitialized() )
 	{
diff --git a/Engine/Source/Runtime/Slate/Private/Framework/MultiBox/MultiBox.cpp b/Engine/Source/Runtime/Slate/Private/Framework/MultiBox/MultiBox.cpp
--- a/Engine/Source/Runtime/Slate/Private/Framework/MultiBox/MultiBox.cpp
+++ b/Engine/Source/Runtime/Slate/Private/Framework/MultiBox/MultiBox.cpp
@@ -26,6 +26,7 @@
 #include "Framework/MultiBox/SClippingHorizontalBox.h"
 #include "Framework/MultiBox/SWidgetBlock.h"
 #include "Framework/MultiBox/SToolBarComboButtonBlock.h"
+#include "HAL/PlatformMath.h"
 
 #include "Framework/Commands/UICommandDragDropOp.h"
 #include "Framework/MultiBox/SUniformToolbarPanel.h"
@@ -85,15 +86,10 @@ static FString ToString(const TArray<FText>& SearchTextHierarchyComponents)
 } // namespace MultiBoxUtils
 } // namespace UE
 
-TAutoConsoleVariable<bool> AlwaysShowMenuSearchField(
-	TEXT("Slate.AlwaysShowMenuSearchField"),
-	false,
-	TEXT("Always show the search field in menus (default false). When this is disabled, the user has to type for the search field to appear."));
-
 TAttribute<bool> FMultiBoxSettings::UseSmallToolBarIcons;
 TAttribute<bool> FMultiBoxSettings::DisplayMultiboxHooks;
 FMultiBoxSettings::FConstructToolTip FMultiBoxSettings::ToolTipConstructor = FConstructToolTip::CreateStatic( &FMultiBoxSettings::ConstructDefaultToolTip );
-
+TAttribute<int> FMultiBoxSettings::MenuSearchFieldVisibilityThreshold;
 
 FMultiBoxSettings::FMultiBoxSettings()
 {
@@ -883,7 +879,7 @@ void SMultiBoxWidget::CreateSearchTextWidget()
 		return;
 	}
 
-	const FText SearchHint = AlwaysShowMenuSearchField.GetValueOnAnyThread()
+	const FText SearchHint = ShouldShowMenuSearchField()
 							   ? LOCTEXT("SearchHintStartTyping", "Start typing to search")
 							   : LOCTEXT("SearchHint", "Search");
 
@@ -918,7 +914,7 @@ void SMultiBoxWidget::OnFilterTextChanged(const FText& InFilterText)
 		if (SearchTextWidget.IsValid() && SearchBlockWidget.IsValid())
 		{
 			// We only have to do this if we're not always showing the search widget.
-			if (!AlwaysShowMenuSearchField.GetValueOnAnyThread())
+			if (!ShouldShowMenuSearchField())
 			{
 				// Make the search box visible and focused
 				SearchBlockWidget->SetVisibility(EVisibility::Visible);
@@ -1586,7 +1582,7 @@ FReply SMultiBoxWidget::OnFocusReceived( const FGeometry& MyGeometry, const FFoc
 
 void SMultiBoxWidget::OnFocusChanging(const FWeakWidgetPath& PreviousFocusPath, const FWidgetPath& NewWidgetPath, const FFocusEvent& InFocusEvent)
 {
-	if (!AlwaysShowMenuSearchField.GetValueOnAnyThread())
+	if (!ShouldShowMenuSearchField())
 	{
 		return;
 	}
@@ -1655,7 +1651,7 @@ void SMultiBoxWidget::BeginSearch(const TCHAR InChar)
 		if (SearchTextWidget.IsValid() && SearchBlockWidget.IsValid())
 		{
 			// We only have to do this if we're not always showing the search widget.
-			if (!AlwaysShowMenuSearchField.GetValueOnAnyThread())
+			if (!ShouldShowMenuSearchField())
 			{
 				// Make the search box visible and focused
 				SearchBlockWidget->SetVisibility(EVisibility::Visible);
@@ -1770,6 +1766,33 @@ void SMultiBoxWidget::FlattenSubMenusRecursive(uint32 MaxRecursionLevels)
 	}
 }
 
+bool SMultiBoxWidget::ShouldShowMenuSearchField()
+{
+	const int SearchFieldTreshold = FMultiBoxSettings::MenuSearchFieldVisibilityThreshold.Get();
+
+	const TArray<TSharedRef<const FMultiBlock>> Blocks = MultiBox->GetBlocks();
+
+	int NumUserCountableEntries = 0;
+	for (const TSharedRef<const FMultiBlock>& Block : Blocks)
+	{
+		// Don't count headers and separators.
+		const EMultiBlockType BlockType = Block->GetType();
+		if (BlockType == EMultiBlockType::Heading
+		 || BlockType == EMultiBlockType::Separator)
+		{
+			continue;
+		}
+
+		++NumUserCountableEntries;
+	}
+
+	// Subtract 1 to account for the search field widget. After that widget is created,
+	// it's added at the start of the menu regardless if it's visible or not.
+	NumUserCountableEntries = FMath::Max<int>(0, NumUserCountableEntries-1);
+
+	return NumUserCountableEntries >= SearchFieldTreshold;
+}
+
 void SMultiBoxWidget::FilterMultiBoxEntries()
 {
 	VisibleFlattenHierarchyTips.Empty();
@@ -1782,7 +1805,7 @@ void SMultiBoxWidget::FilterMultiBoxEntries()
 		}
 
 		// We only have to do this if we're not always showing the search widget.
-		if (!AlwaysShowMenuSearchField.GetValueOnAnyThread())
+		if (!ShouldShowMenuSearchField())
 		{
 			if (SearchBlockWidget.IsValid())
 			{
@@ -1812,7 +1835,7 @@ void SMultiBoxWidget::FilterMultiBoxEntries()
 		const TSharedPtr<SWidget>& Widget = It.Key();
 
 		// Skip the search widget itself when scanning for searchable items.
-		if (AlwaysShowMenuSearchField.GetValueOnAnyThread())
+		if (ShouldShowMenuSearchField())
 		{
 			if (Widget == SearchBlockWidget)
 			{
@@ -1865,7 +1888,7 @@ void SMultiBoxWidget::FilterMultiBoxEntries()
 	}
 
 	// If we always show the search widget, we're skipping it in the code above and do not need to show it here to compensate.
-	if (!AlwaysShowMenuSearchField.GetValueOnAnyThread())
+	if (!ShouldShowMenuSearchField())
 	{
 		// Show the search widget again, it was hidden by the above code.
 		if (SearchBlockWidget.IsValid())
diff --git a/Engine/Source/Runtime/Slate/Private/Framework/MultiBox/SWidgetBlock.cpp b/Engine/Source/Runtime/Slate/Private/Framework/MultiBox/SWidgetBlock.cpp
--- a/Engine/Source/Runtime/Slate/Private/Framework/MultiBox/SWidgetBlock.cpp
+++ b/Engine/Source/Runtime/Slate/Private/Framework/MultiBox/SWidgetBlock.cpp
@@ -79,7 +79,9 @@ void SWidgetBlock::Construct( const FArguments& InArgs )
  */
 void SWidgetBlock::BuildMultiBlockWidget(const ISlateStyle* StyleSet, const FName& StyleName)
 {
-	TSharedPtr< const FMultiBox > MultiBox = OwnerMultiBoxWidget.Pin()->GetMultiBox();
+	TSharedPtr<SMultiBoxWidget> OwnerMultiBoxWidgetPinned = OwnerMultiBoxWidget.Pin();
+
+	TSharedPtr< const FMultiBox > MultiBox = OwnerMultiBoxWidgetPinned->GetMultiBox();
 	TSharedRef< const FWidgetBlock > WidgetBlock = StaticCastSharedRef< const FWidgetBlock >( MultiBlock.ToSharedRef() );
 
 	// Support menus which do not have a defined widget style yet
@@ -101,8 +103,6 @@ void SWidgetBlock::BuildMultiBlockWidget(const ISlateStyle* StyleSet, const FNam
 		LabelStyle = &StyleSet->GetWidgetStyle<FTextBlockStyle>(ISlateStyle::Join(StyleName, ".Label"));
 	}
 
-	TSharedPtr<SMultiBoxWidget> OwnerMultiBoxWidgetPinned = OwnerMultiBoxWidget.Pin();
-
 	if(OwnerMultiBoxWidgetPinned->GetMultiBox()->GetType() == EMultiBoxType::Menu)
 	{
 		// Account for checkmark used in other menu blocks but not used in for widget rows
@@ -148,12 +148,9 @@ void SWidgetBlock::BuildMultiBlockWidget(const ISlateStyle* StyleSet, const FNam
 		OwnerMultiBoxWidgetPinned->SetSearchBlockWidget(this->AsWidget());
 
 		// When we are always showing the search widget, we should not hide it here.
-		if (IConsoleVariable* AlwaysShowMenuSearchFieldVar = IConsoleManager::Get().FindConsoleVariable(TEXT("Slate.AlwaysShowMenuSearchField")))
+		if (!OwnerMultiBoxWidgetPinned->ShouldShowMenuSearchField())
 		{
-			if (!AlwaysShowMenuSearchFieldVar->GetBool())
-			{
-				this->AsWidget()->SetVisibility(EVisibility::Collapsed);
-			}
+			this->AsWidget()->SetVisibility(EVisibility::Collapsed);
 		}
 	}
 
diff --git a/Engine/Source/Runtime/Slate/Public/Framework/MultiBox/MultiBox.h b/Engine/Source/Runtime/Slate/Public/Framework/MultiBox/MultiBox.h
--- a/Engine/Source/Runtime/Slate/Public/Framework/MultiBox/MultiBox.h
+++ b/Engine/Source/Runtime/Slate/Public/Framework/MultiBox/MultiBox.h
@@ -810,6 +810,8 @@ public:
 	 */
 	SLATE_API double GetSummonedMenuTime() const;
 
+	bool ShouldShowMenuSearchField();
+
 private:
 	/** Adds a block Widget to this widget */
 	SLATE_API void AddBlockWidget(const FMultiBlock& Block, TSharedPtr<SHorizontalBox> HorizontalBox, TSharedPtr<SVerticalBox> VerticalBox, EMultiBlockLocation::Type InLocation, bool bSectionContainsIcons, TSharedPtr<const FToolBarComboButtonBlock> OptionsBlock);
diff --git a/Engine/Source/Runtime/Slate/Public/Framework/MultiBox/MultiBoxDefs.h b/Engine/Source/Runtime/Slate/Public/Framework/MultiBox/MultiBoxDefs.h
--- a/Engine/Source/Runtime/Slate/Public/Framework/MultiBox/MultiBoxDefs.h
+++ b/Engine/Source/Runtime/Slate/Public/Framework/MultiBox/MultiBoxDefs.h
@@ -71,6 +71,8 @@ public:
 	static SLATE_API TAttribute<bool> UseSmallToolBarIcons;
 	static SLATE_API TAttribute<bool> DisplayMultiboxHooks;
 	static SLATE_API FConstructToolTip ToolTipConstructor;
+	
+	static SLATE_API TAttribute<int> MenuSearchFieldVisibilityThreshold;
 
 	SLATE_API FMultiBoxSettings();
 
