AnimatedAttributeManager: Fix crash on shutdown

#rnx

diff --git a/Engine/Source/Runtime/Slate/Private/Framework/Animation/AnimatedAttributeManager.cpp b/Engine/Source/Runtime/Slate/Private/Framework/Animation/AnimatedAttributeManager.cpp
--- a/Engine/Source/Runtime/Slate/Private/Framework/Animation/AnimatedAttributeManager.cpp
+++ b/Engine/Source/Runtime/Slate/Private/Framework/Animation/AnimatedAttributeManager.cpp
@@ -68,11 +68,11 @@ void FAnimatedAttributeManager::SetupTick()
 
 void FAnimatedAttributeManager::TeardownTick()
 {
-	if(TickHandle.IsValid())
+	if(TickHandle.IsValid() && FSlateApplication::IsInitialized())
 	{
 		FSlateApplication::Get().OnPreTick().Remove(TickHandle);
-		TickHandle.Reset();
 	}
+	TickHandle.Reset();
 }
 
 void FAnimatedAttributeManager::RegisterAttribute(const TSharedRef<TAnimatedAttributeBase>& InAttribute)
