AnimatedAttributeManager: Fix crash on launch

#rnx

diff --git a/Engine/Source/Runtime/Slate/Private/Framework/Animation/AnimatedAttributeManager.cpp b/Engine/Source/Runtime/Slate/Private/Framework/Animation/AnimatedAttributeManager.cpp
--- a/Engine/Source/Runtime/Slate/Private/Framework/Animation/AnimatedAttributeManager.cpp
+++ b/Engine/Source/Runtime/Slate/Private/Framework/Animation/AnimatedAttributeManager.cpp
@@ -61,9 +61,12 @@ void FAnimatedAttributeManager::Tick(float InDeltaTime)
 
 void FAnimatedAttributeManager::SetupTick()
 {
-	const FSlateApplication::FSlateTickEvent::FDelegate TickDelegate =
-		FSlateApplication::FSlateTickEvent::FDelegate::CreateRaw(this, &FAnimatedAttributeManager::Tick);
-	TickHandle = FSlateApplication::Get().OnPreTick().Add(TickDelegate);
+	if(FSlateApplication::IsInitialized() && !TickHandle.IsValid())
+	{
+		const FSlateApplication::FSlateTickEvent::FDelegate TickDelegate =
+			FSlateApplication::FSlateTickEvent::FDelegate::CreateRaw(this, &FAnimatedAttributeManager::Tick);
+		TickHandle = FSlateApplication::Get().OnPreTick().Add(TickDelegate);
+	}
 }
 
 void FAnimatedAttributeManager::TeardownTick()
@@ -78,6 +81,7 @@ void FAnimatedAttributeManager::TeardownTick()
 void FAnimatedAttributeManager::RegisterAttribute(const TSharedRef<TAnimatedAttributeBase>& InAttribute)
 {
 	Attributes.Add(InAttribute.ToWeakPtr());
+	SetupTick();
 }
 
 void FAnimatedAttributeManager::UnregisterAttribute(const TSharedRef<TAnimatedAttributeBase>& InAttribute)
