Fixed an issue where inputs with BroadcastValueChangesPerKey=true would not apply entered math expressions correctly

diff --git a/Engine/Source/Runtime/Slate/Public/Widgets/Input/SSpinBox.h b/Engine/Source/Runtime/Slate/Public/Widgets/Input/SSpinBox.h
--- a/Engine/Source/Runtime/Slate/Public/Widgets/Input/SSpinBox.h
+++ b/Engine/Source/Runtime/Slate/Public/Widgets/Input/SSpinBox.h
@@ -972,12 +972,13 @@ protected:
 				EditableText->SetText(FText::FromString(ValidData));
 			}
 
-			if (bBroadcastValueChangesPerKey)
+			// we check that the input is numeric, as we don't want to commit the new value on every change when an expression like *= is entered
+			if (bBroadcastValueChangesPerKey && FCString::IsNumeric(*Data))
 			{
-				TOptional<NumericType> NewValue = Interface->FromString(NewText.ToString(), ValueAttribute.Get());
+				TOptional<NumericType> NewValue = Interface->FromString(Data, ValueAttribute.Get());
 				if (NewValue.IsSet())
 				{
-					CommitValue(NewValue.GetValue(), (double)NewValue.GetValue(), CommittedViaCode, ETextCommit::Default);
+					CommitValue(NewValue.GetValue(), static_cast<double>(NewValue.GetValue()), CommittedViaCode, ETextCommit::Default);
 				}
 			}
 		}
