Fix drawing unnecessary horizontal line into root items expander arrows.

diff --git a/Engine/Source/Runtime/Slate/Private/Widgets/Views/SExpanderArrow.cpp b/Engine/Source/Runtime/Slate/Private/Widgets/Views/SExpanderArrow.cpp
--- a/Engine/Source/Runtime/Slate/Private/Widgets/Views/SExpanderArrow.cpp
+++ b/Engine/Source/Runtime/Slate/Private/Widgets/Views/SExpanderArrow.cpp
@@ -128,6 +128,7 @@ int32 SExpanderArrow::OnPaint(const FPaintArgs& Args, const FGeometry& AllottedG
 		}
 
 		// Draw horizontal connector from parent wire to child.
+		if (NumLevels > 1)
 		{
 			float LeafDepth = OwnerRow->DoesItemHaveChildren() ? 10.f : 0.0f;
 			const float HorizontalWireStart = (NumLevels - 1)*Indent;
