Use MakeShared instead of MakeSharable when creating groups for workspace menus.

diff --git a/Engine/Source/Runtime/Slate/Public/Framework/Docking/WorkspaceItem.h b/Engine/Source/Runtime/Slate/Public/Framework/Docking/WorkspaceItem.h
--- a/Engine/Source/Runtime/Slate/Public/Framework/Docking/WorkspaceItem.h
+++ b/Engine/Source/Runtime/Slate/Public/Framework/Docking/WorkspaceItem.h
@@ -34,22 +34,22 @@ protected:
 public:
 	static TSharedRef<FWorkspaceItem> NewGroup( const FText& DisplayName, const FSlateIcon& Icon = FSlateIcon(), const bool bSortChildren = false )
 	{
-		return MakeShareable( new FWorkspaceItem( NAME_None, DisplayName, Icon, bSortChildren ) );
+		return MakeShared<FWorkspaceItem>(NAME_None, DisplayName, Icon, bSortChildren);
 	}
 
 	static TSharedRef<FWorkspaceItem> NewGroup( const FName& Name, const FText& DisplayName, const FSlateIcon& Icon = FSlateIcon(), const bool bSortChildren = false )
 	{
-		return MakeShareable( new FWorkspaceItem( Name, DisplayName, Icon, bSortChildren ) );
+		return MakeShared<FWorkspaceItem>(Name, DisplayName, Icon, bSortChildren);
 	}
 
 	static TSharedRef<FWorkspaceItem> NewGroup( const FText& DisplayName, const FText& TooltipText, const FSlateIcon& Icon = FSlateIcon(), const bool bSortChildren = false )
 	{
-		return MakeShareable( new FWorkspaceItem( NAME_None, DisplayName, TooltipText, Icon, bSortChildren ) );
+		return MakeShared<FWorkspaceItem>(NAME_None, DisplayName, TooltipText, Icon, bSortChildren);
 	}
 
 	static TSharedRef<FWorkspaceItem> NewGroup( const FName& Name, const FText& DisplayName, const FText& TooltipText, const FSlateIcon& Icon = FSlateIcon(), const bool bSortChildren = false )
 	{
-		return MakeShareable( new FWorkspaceItem( Name, DisplayName, TooltipText, Icon, bSortChildren ) );
+		return MakeShared<FWorkspaceItem>(Name, DisplayName, TooltipText, Icon, bSortChildren);
 	}
 
 	TSharedRef<FWorkspaceItem> AddGroup( const FText& InDisplayName, const FSlateIcon& InIcon = FSlateIcon(), const bool InSortChildren = false )
@@ -223,4 +223,8 @@ protected:
 	TArray< TSharedRef<FWorkspaceItem> > ChildItems;
 
 	TWeakPtr<FWorkspaceItem> ParentItem;
+
+	// Allows MakeShared with protected constructor.
+	template <typename ObjectType, ESPMode Mode>
+	friend class SharedPointerInternals::TIntrusiveReferenceController;
 };
