[Slate/SInlineEditableTextBlock]
Add GetText function to SInlineEditableTextBlock.

diff --git a/Engine/Source/Runtime/Slate/Private/Widgets/Text/SInlineEditableTextBlock.cpp b/Engine/Source/Runtime/Slate/Private/Widgets/Text/SInlineEditableTextBlock.cpp
--- a/Engine/Source/Runtime/Slate/Private/Widgets/Text/SInlineEditableTextBlock.cpp
+++ b/Engine/Source/Runtime/Slate/Private/Widgets/Text/SInlineEditableTextBlock.cpp
@@ -178,6 +178,14 @@ void SInlineEditableTextBlock::SetText( const FString& InText )
 	SetEditableText( Text );
 }
 
+FText SInlineEditableTextBlock::GetText() const
+{
+	return
+#if WITH_FANCY_TEXT
+	bIsMultiLine ? MultiLineTextBox->GetText() :
+#endif //WITH_FANCY_TEXT
+				 TextBox->GetText();
+}
 
 void SInlineEditableTextBlock::SetHighlightText( const TAttribute< FText >& InText )
 {
diff --git a/Engine/Source/Runtime/Slate/Public/Widgets/Text/SInlineEditableTextBlock.h b/Engine/Source/Runtime/Slate/Public/Widgets/Text/SInlineEditableTextBlock.h
--- a/Engine/Source/Runtime/Slate/Public/Widgets/Text/SInlineEditableTextBlock.h
+++ b/Engine/Source/Runtime/Slate/Public/Widgets/Text/SInlineEditableTextBlock.h
@@ -158,6 +158,9 @@ class SInlineEditableTextBlock: public SCompoundWidget
 	SLATE_API void SetText( const TAttribute< FText >& InText );
 	SLATE_API void SetText( const FString& InText );
 
+	/** Return the current text */
+	SLATE_API FText GetText() const;
+
 	SLATE_API void SetHighlightText(const TAttribute<FText>& InText);
 
 	/** Sets the wrap text at attribute.  See WrapTextAt attribute */
