Fix regression where right-clicking on tabs opened "Save Content" dialog

When tabs are right-clicked, a menu is shown which includes a "Close" entry. To determine if this "Close" entry should be enabled, the editor runs a OnCanClose delegate. A recent change in SDockTab::CanCloseTab() inadvertently disabled boolean early-out behavior which caused further checks to run as this menu was opened, including a check on the level editor side. To determine if the level editor tab can be closed, we check if all data is saved and if not trigger the "Save Content" dialog. The level editor tab is already flagged as non-closable, but because of the the recent change that removed the early-out logic, the second check also ran and triggered the dialog.

This fix adds back the early-out behavior of SDockTab::CanCloseTab() which in the case of the level editor avoids running the last check that triggers the "Save Content" dialog.

diff --git a/Engine/Source/Runtime/Slate/Private/Widgets/Docking/SDockTab.cpp b/Engine/Source/Runtime/Slate/Private/Widgets/Docking/SDockTab.cpp
--- a/Engine/Source/Runtime/Slate/Private/Widgets/Docking/SDockTab.cpp
+++ b/Engine/Source/Runtime/Slate/Private/Widgets/Docking/SDockTab.cpp
@@ -949,10 +949,21 @@ TOptional<FVector2D> SDockTab::GetTabIconSize() const
 
 bool SDockTab::CanCloseTab() const
 {
-	TSharedPtr<FTabManager> Manager = MyTabManager.Pin();
-	const bool bIsTabCloseable = Manager == nullptr || Manager->IsTabCloseable(SharedThis(this));
-	const bool bCanCloseTab = !OnCanCloseTab.IsBound() || OnCanCloseTab.Execute();
-	return bIsTabCloseable && bCanCloseTab;
+	// Check the below conditions separately and early out to prevent the OnCanCloseTab delegate
+	// from being called if we already know we cannot close the tab.
+
+	if (TSharedPtr<FTabManager> Manager = MyTabManager.Pin();
+		Manager && !Manager->IsTabCloseable(SharedThis(this)))
+	{
+		return false;
+	}
+	
+	if (OnCanCloseTab.IsBound() && !OnCanCloseTab.Execute())
+	{
+		return false;
+	}
+	
+	return true;
 }
 
 bool SDockTab::RequestCloseTab()
