Fix comments made in the review for 31612816.
#rnx

diff --git a/Engine/Source/Runtime/Slate/Private/Framework/Application/SlateApplication.cpp b/Engine/Source/Runtime/Slate/Private/Framework/Application/SlateApplication.cpp
--- a/Engine/Source/Runtime/Slate/Private/Framework/Application/SlateApplication.cpp
+++ b/Engine/Source/Runtime/Slate/Private/Framework/Application/SlateApplication.cpp
@@ -1510,7 +1510,7 @@ void FSlateApplication::Tick(ESlateTickType TickType)
 
 	FScopeLock SlateTickAccess(&SlateTickCriticalSection);
 
-	TGuardValue<bool> TmpGuard(bIsTicking, true);
+	TGuardValue<bool> IsTickingGuard(bIsTicking, true);
 
 #if WITH_EDITOR
 	FScopedPreventDebuggingMode SlatePreventDebugginModeWhileTicking(NSLOCTEXT("EnterDebuggingMode", "WindowTicking", "The window is ticking."));
