Added Select All/None button for column view right click context menu

diff --git a/Engine/Source/Runtime/Slate/Private/Widgets/Views/SHeaderRow.cpp b/Engine/Source/Runtime/Slate/Private/Widgets/Views/SHeaderRow.cpp
--- a/Engine/Source/Runtime/Slate/Private/Widgets/Views/SHeaderRow.cpp
+++ b/Engine/Source/Runtime/Slate/Private/Widgets/Views/SHeaderRow.cpp
@@ -965,8 +965,90 @@ void SHeaderRow::RegenerateWidgets()
 	}
 }
 
+void SHeaderRow::ToggleAllColumns()
+{
+	ECheckBoxState CurrentState = GetSelectAllColumnsCheckState();
+	bool bShouldSelectAll = CurrentState == ECheckBoxState::Unchecked || CurrentState == ECheckBoxState::Undetermined;
+
+	for (FColumn& SomeColumn : Columns)
+	{
+		if (!SomeColumn.ShouldGenerateWidget.IsSet())
+		{
+			SomeColumn.bIsVisible = bShouldSelectAll;
+		}
+	}
+	
+	RefreshColumns();
+	ColumnsChanged.Broadcast(SharedThis(this));
+	OnHiddenColumnsListChanged.ExecuteIfBound();
+}
+
+bool SHeaderRow::CanToggleAllColumns() const
+{
+	return true;
+}
+
+ECheckBoxState SHeaderRow::GetSelectAllColumnsCheckState() const
+{
+	bool bAnyColumnsVisible = false;
+	bool bAnyColumnsInvisible = false;
+	for (const FColumn& SomeColumn : Columns)
+	{
+		if (!SomeColumn.ShouldGenerateWidget.IsSet())
+		{
+			if(SomeColumn.bIsVisible)
+			{
+				bAnyColumnsVisible = true;
+			}
+			else
+			{
+				bAnyColumnsInvisible = true;
+			}			
+		}
+	}
+
+	if(bAnyColumnsVisible && bAnyColumnsInvisible)
+	{
+		return ECheckBoxState::Undetermined;
+	}
+
+	if(bAnyColumnsVisible)
+	{
+		return ECheckBoxState::Checked;
+	}
+
+	return ECheckBoxState::Unchecked;
+}
+
+FText SHeaderRow::GetToggleAllColumnsText() const
+{
+	ECheckBoxState CurrentState = GetSelectAllColumnsCheckState();
+	if(CurrentState == ECheckBoxState::Checked)
+	{
+		return LOCTEXT("DeselectAllColumns", "Select: None");
+	}
+	else if(CurrentState == ECheckBoxState::Unchecked || CurrentState == ECheckBoxState::Undetermined)
+	{
+		return LOCTEXT("SelectAllColumns", "Select: All");
+	}
+
+	return FText::GetEmpty();
+}
+
 void SHeaderRow::OnGenerateSelectColumnsSubMenu(FMenuBuilder& InSubMenuBuilder)
 {
+	InSubMenuBuilder.AddMenuEntry(
+		TAttribute<FText>::CreateSP(this, &SHeaderRow::GetToggleAllColumnsText),
+		FText::GetEmpty(),
+		FSlateIcon(),
+		FUIAction(
+			FExecuteAction::CreateSP(this, &SHeaderRow::ToggleAllColumns),
+			FCanExecuteAction::CreateSP(this, &SHeaderRow::CanToggleAllColumns),
+			FGetActionCheckState::CreateSP(this, &SHeaderRow::GetSelectAllColumnsCheckState)
+		),
+		NAME_None,
+		EUserInterfaceActionType::ToggleButton);
+	
 	for (const FColumn& SomeColumn : Columns)
 	{
 		const bool bCanExecuteAction = !SomeColumn.ShouldGenerateWidget.IsSet();
diff --git a/Engine/Source/Runtime/Slate/Public/Widgets/Views/SHeaderRow.h b/Engine/Source/Runtime/Slate/Public/Widgets/Views/SHeaderRow.h
--- a/Engine/Source/Runtime/Slate/Public/Widgets/Views/SHeaderRow.h
+++ b/Engine/Source/Runtime/Slate/Public/Widgets/Views/SHeaderRow.h
@@ -370,7 +370,12 @@ private:
 	/** Regenerates all widgets in the header */
 	SLATE_API void RegenerateWidgets();
 
-	SLATE_API void OnGenerateSelectColumnsSubMenu(FMenuBuilder& InSubMenuBuilder);
+	void ToggleAllColumns();
+	bool CanToggleAllColumns() const;
+	ECheckBoxState GetSelectAllColumnsCheckState() const;
+	FText GetToggleAllColumnsText() const;
+
+	void OnGenerateSelectColumnsSubMenu(FMenuBuilder& InSubMenuBuilder);
 	SLATE_API void ToggleGeneratedColumn(FName ColumnId);
 	SLATE_API ECheckBoxState GetGeneratedColumnCheckedState(FName ColumnId) const;
 
