Consider navigation config thresholds when determining if an analog event should set the last interaction time.

#sh-02-13-playtime-summary-counts-idle-time-ext

https://epicgames.slack.com/archives/C06JR2LJZB6/p1708625179848919

#tests Compiled Test Client / Ran a packaged build with my exe on top. Navigated all frontend tabs and the lobby pedestals and verified via breakpoint I only set last interaction with the joystick if the stick movement was beyond the threshold. Tested all cardinal directions. Repeated test for in-game as well. Verified these thresholds were also used during normal gameplay without menus.

diff --git a/Engine/Source/Runtime/Slate/Private/Framework/Application/NavigationConfig.cpp b/Engine/Source/Runtime/Slate/Private/Framework/Application/NavigationConfig.cpp
--- a/Engine/Source/Runtime/Slate/Private/Framework/Application/NavigationConfig.cpp
+++ b/Engine/Source/Runtime/Slate/Private/Framework/Application/NavigationConfig.cpp
@@ -205,6 +205,20 @@ FString FNavigationConfig::ToString() const
 	return Builder.ToString();
 }
 
+bool FNavigationConfig::IsAnalogEventBeyondNavigationThreshold(const FAnalogInputEvent& InAnalogEvent) const
+{
+	if (bAnalogNavigation)
+	{
+		const FKey& AnalogKey = InAnalogEvent.GetKey();
+		const float AbsAnalogValue = FMath::Abs(InAnalogEvent.GetAnalogValue());
+	 
+	 	return (IsAnalogHorizontalKey(AnalogKey) && AbsAnalogValue > AnalogNavigationHorizontalThreshold)
+			|| (IsAnalogVerticalKey(AnalogKey) 	 && AbsAnalogValue > AnalogNavigationVerticalThreshold);
+	}
+
+	return false;
+}
+
 FTwinStickNavigationConfig::FTwinStickNavigationConfig()
 {
 	bTabNavigation = false;
diff --git a/Engine/Source/Runtime/Slate/Private/Framework/Application/SlateApplication.cpp b/Engine/Source/Runtime/Slate/Private/Framework/Application/SlateApplication.cpp
--- a/Engine/Source/Runtime/Slate/Private/Framework/Application/SlateApplication.cpp
+++ b/Engine/Source/Runtime/Slate/Private/Framework/Application/SlateApplication.cpp
@@ -4924,17 +4924,15 @@ bool FSlateApplication::ProcessAnalogInputEvent(const FAnalogInputEvent& InAnalo
 			}, ESlateDebuggingInputEvent::AnalogInput);
 	}
 
-	// If no one handled this, it was probably motion in the deadzone.  Don't treat it as activity.
-	if (Reply.IsEventHandled())
+	// Ensure the analog input event exceeds the thresholds set in the navigation config before considering as interaction.
+	const TSharedRef<FNavigationConfig> RelevantNavConfig = GetRelevantNavConfig(InAnalogInputEvent.GetUserIndex());
+	if (RelevantNavConfig->IsAnalogEventBeyondNavigationThreshold(InAnalogInputEvent))
 	{
 		SetLastUserInteractionTime(this->GetCurrentTime());
 		LastUserInteractionTimeForThrottling = LastUserInteractionTime;
-		return true;
-	}
-	else
-	{
-		return false;
 	}
+
+	return Reply.IsEventHandled();
 }
 
 FKey TranslateMouseButtonToKey( const EMouseButtons::Type Button )
diff --git a/Engine/Source/Runtime/Slate/Public/Framework/Application/NavigationConfig.h b/Engine/Source/Runtime/Slate/Public/Framework/Application/NavigationConfig.h
--- a/Engine/Source/Runtime/Slate/Public/Framework/Application/NavigationConfig.h
+++ b/Engine/Source/Runtime/Slate/Public/Framework/Application/NavigationConfig.h
@@ -89,6 +89,9 @@ public:
 	/** Simplification of config as string */
 	SLATE_API virtual FString ToString() const;
 
+	/** Returns whether the analog event is beyond the navigation thresholds set in this config. */
+	bool IsAnalogEventBeyondNavigationThreshold(const FAnalogInputEvent& InAnalogEvent) const;
+
 public:
 	/** Should the Tab key perform next and previous style navigation. */
 	bool bTabNavigation;
