Fix issue related to minimizing a notification's parent window

During a short fade in, SNotificationExtendable asks Slate to enter responsive mode. The bug we're fixing here was caused by SNotificationExtendable relying on an SWidget::Tick override to get ticked and eventually release the throttle handle for responsive mode. The bug occurred when the notification's parent window was minimized before the notification is shown, preventint the parent window from ever being painted and thus never ticked, causing the notification to ask for responsive mode but then never ticking out of it, holding onto Slate responsive mode forever. This prevents certain processes in the editor from progressing.

This change includes a related fix of not tying a notification to a window if the window is already minimized or not visible to prevent the notification from showing up as a black square without content because it never gets painted.

#rnx

diff --git a/Engine/Source/Runtime/Slate/Private/Framework/Notifications/NotificationManager.cpp b/Engine/Source/Runtime/Slate/Private/Framework/Notifications/NotificationManager.cpp
--- a/Engine/Source/Runtime/Slate/Private/Framework/Notifications/NotificationManager.cpp
+++ b/Engine/Source/Runtime/Slate/Private/Framework/Notifications/NotificationManager.cpp
@@ -150,7 +150,9 @@ TSharedRef<SNotificationList> FSlateNotificationManager::CreateStackForArea(cons
 	NotificationWindow->SetContent(NotificationList);
 	NotificationList->ParentWindowPtr = NotificationWindow;
 
-	if (Window.IsValid())
+	// Prevent tying a notification to a window that is minimized because a minimized window isn't painted
+	// and therefore the notification wouldn't be either.
+	if (Window.IsValid() && Window->IsVisible() && !Window->IsWindowMinimized())
 	{
 		FSlateApplication::Get().AddWindowAsNativeChild(NotificationWindow, Window.ToSharedRef());
 	}
diff --git a/Engine/Source/Runtime/Slate/Private/Widgets/Notifications/SNotificationList.cpp b/Engine/Source/Runtime/Slate/Private/Widgets/Notifications/SNotificationList.cpp
--- a/Engine/Source/Runtime/Slate/Private/Widgets/Notifications/SNotificationList.cpp
+++ b/Engine/Source/Runtime/Slate/Private/Widgets/Notifications/SNotificationList.cpp
@@ -4,6 +4,7 @@
 #include "Animation/CurveHandle.h"
 #include "Animation/CurveSequence.h"
 #include "Application/ThrottleManager.h"
+#include "Containers/Ticker.h"
 #include "Framework/Application/SlateApplication.h"
 #include "HAL/PlatformApplicationMisc.h"
 #include "Widgets/SBoxPanel.h"
@@ -29,6 +30,13 @@
 class SNotificationExtendable : public SNotificationItem
 {
 public:
+	SNotificationExtendable()
+	{
+		SetCanTick(false);
+
+		TickDelegateHandle = FTSTicker::GetCoreTicker().AddTicker(FTickerDelegate::CreateRaw(this, &SNotificationExtendable::Update));
+	}
+
 	virtual ~SNotificationExtendable()
 	{
 		// Just in case, make sure we have left responsive mode when getting cleaned up
@@ -36,6 +44,8 @@ public:
 		{
 			FSlateThrottleManager::Get().LeaveResponsiveMode( ThrottleHandle );
 		}
+
+		FTSTicker::GetCoreTicker().RemoveTicker(TickDelegateHandle);
 	}
 	
 	
@@ -143,7 +153,7 @@ public:
 		FadeOutDuration = Duration;
 	}
 
-	void Tick( const FGeometry& AllottedGeometry, const double InCurrentTime, const float InDeltaTime ) override
+	bool Update(float InDeltaTime)
 	{
 		bool bIsFadingOut = FadeAnimation.IsInReverse();
 		if (bAutoExpire)
@@ -172,6 +182,8 @@ public:
 			// Leave responsive mode once the intro finishes playing
 			FSlateThrottleManager::Get().LeaveResponsiveMode( ThrottleHandle );
 		}
+
+		return true;
 	}
 
 protected:
@@ -337,6 +349,7 @@ protected:
 
 	/** Handle to a throttle request made to ensure the intro animation is smooth in low FPS situations */
 	FThrottleRequest ThrottleHandle;
+	FTSTicker::FDelegateHandle TickDelegateHandle;
 
 	float InternalTime = 0.0f;
 	bool bAutoExpire = false;
