Engine
- Add a runtime option to disable the ability for a STableViewBase to scroll on Touch and Mouse input.

[RN] Add a runtime option to disable the ability for a STableViewBase to scroll on Touch and Mouse input.

diff --git a/Engine/Source/Runtime/Slate/Private/Widgets/Views/STableViewBase.cpp b/Engine/Source/Runtime/Slate/Private/Widgets/Views/STableViewBase.cpp
--- a/Engine/Source/Runtime/Slate/Private/Widgets/Views/STableViewBase.cpp
+++ b/Engine/Source/Runtime/Slate/Private/Widgets/Views/STableViewBase.cpp
@@ -489,7 +489,7 @@ FReply STableViewBase::OnMouseButtonUp( const FGeometry& MyGeometry, const FPoin
 
 FReply STableViewBase::OnMouseMove( const FGeometry& MyGeometry, const FPointerEvent& MouseEvent )
 {	
-	if( bEnableRightClickScrolling && MouseEvent.IsMouseButtonDown( EKeys::RightMouseButton ) && !MouseEvent.IsTouchEvent() )
+	if( bEnableRightClickScrolling && MouseEvent.IsMouseButtonDown( EKeys::RightMouseButton ) && !MouseEvent.IsTouchEvent() && bIsPointerScrollingEnabled )
 	{
 		// We only care about deltas along the scroll axis
 		FTableViewDimensions CursorDeltaDimensions(Orientation, MouseEvent.GetCursorDelta());
@@ -573,7 +573,7 @@ void STableViewBase::OnMouseLeave( const FPointerEvent& MouseEvent )
 
 FReply STableViewBase::OnMouseWheel( const FGeometry& MyGeometry, const FPointerEvent& MouseEvent )
 {
-	if( !MouseEvent.IsControlDown() )
+	if( bIsPointerScrollingEnabled && !MouseEvent.IsControlDown() )
 	{
 		// Make sure scroll velocity is cleared so it doesn't fight with the mouse wheel input
 		this->InertialScrollManager.ClearScrollVelocity();
@@ -634,7 +634,7 @@ FReply STableViewBase::OnTouchStarted( const FGeometry& MyGeometry, const FPoint
 
 FReply STableViewBase::OnTouchMoved( const FGeometry& MyGeometry, const FPointerEvent& InTouchEvent )
 {
-	if (bEnableTouchScrolling && bStartedTouchInteraction)
+	if (bIsPointerScrollingEnabled && bEnableTouchScrolling && bStartedTouchInteraction)
 	{
 		// We only care about deltas along the scroll axis
 		FTableViewDimensions CursorDeltaDimensions(Orientation, InTouchEvent.GetCursorDelta());
@@ -931,6 +931,11 @@ void STableViewBase::SetWheelScrollMultiplier(float NewWheelScrollMultiplier)
 	WheelScrollMultiplier = NewWheelScrollMultiplier;
 }
 
+void STableViewBase::SetIsPointerScrollingEnabled(bool bInIsPointerScrollingEnabled)
+{
+	bIsPointerScrollingEnabled = bInIsPointerScrollingEnabled;
+}
+
 void STableViewBase::SetBackgroundBrush(const TAttribute<const FSlateBrush*>& InBackgroundBrush)
 {
 	BackgroundBrush.SetImage(*this, InBackgroundBrush);
diff --git a/Engine/Source/Runtime/Slate/Public/Widgets/Views/STableViewBase.h b/Engine/Source/Runtime/Slate/Public/Widgets/Views/STableViewBase.h
--- a/Engine/Source/Runtime/Slate/Public/Widgets/Views/STableViewBase.h
+++ b/Engine/Source/Runtime/Slate/Public/Widgets/Views/STableViewBase.h
@@ -185,6 +185,9 @@ public:
 	/** Sets the multiplier applied when wheel scrolling. Higher numbers will cover more distance per click of the wheel. */
 	SLATE_API void SetWheelScrollMultiplier(float NewWheelScrollMultiplier);
 
+	/** Enables/disables being able to scroll. This should be use as a temporary mean to disable scrolling. */
+	SLATE_API void SetIsPointerScrollingEnabled(bool bInIsPointerScrollingEnabled);
+
 	/** Sets the Background Brush */
 	SLATE_API void SetBackgroundBrush(const TAttribute<const FSlateBrush*>& InBackgroundBrush);
 
@@ -475,6 +478,9 @@ protected:
 	/** How much to scroll when using mouse wheel */
 	float WheelScrollMultiplier;
 
+	/** Wheter the list is allowed to scroll. */
+	bool bIsPointerScrollingEnabled = true;
+
 	/** The layout and scroll orientation of the list */
 	EOrientation Orientation = Orient_Vertical;
 
diff --git a/Engine/Source/Runtime/UMG/Private/Components/ListViewBase.cpp b/Engine/Source/Runtime/UMG/Private/Components/ListViewBase.cpp
--- a/Engine/Source/Runtime/UMG/Private/Components/ListViewBase.cpp
+++ b/Engine/Source/Runtime/UMG/Private/Components/ListViewBase.cpp
@@ -90,6 +90,15 @@ void UListViewBase::SetScrollbarVisibility(ESlateVisibility InVisibility)
 	}
 }
 
+UMG_API void UListViewBase::SetIsPointerScrollingEnabled(bool bInIsPointerScrollingEnabled)
+{
+	bIsPointerScrollingEnabled = bInIsPointerScrollingEnabled;
+	if (MyTableViewBase)
+	{
+		MyTableViewBase->SetIsPointerScrollingEnabled(bInIsPointerScrollingEnabled);
+	}
+}
+
 const TArray<UUserWidget*>& UListViewBase::GetDisplayedEntryWidgets() const
 { 
 	return EntryWidgetPool.GetActiveWidgets(); 
@@ -139,6 +148,7 @@ TSharedRef<SWidget> UListViewBase::RebuildWidget()
 	MyTableViewBase->SetEnableTouchAnimatedScrolling(bInEnableTouchAnimatedScrolling);
 	MyTableViewBase->SetIsRightClickScrollingEnabled(bEnableRightClickScrolling);
 	MyTableViewBase->SetIsTouchScrollingEnabled(bEnableTouchScrolling);
+	MyTableViewBase->SetIsPointerScrollingEnabled(bIsPointerScrollingEnabled);
 	MyTableViewBase->SetFixedLineScrollOffset(bEnableFixedLineOffset ? TOptional<double>(FixedLineScrollOffset) : TOptional<double>());
 	MyTableViewBase->SetWheelScrollMultiplier(GetGlobalScrollAmount() * WheelScrollMultiplier);
 
@@ -164,6 +174,7 @@ void UListViewBase::SynchronizeProperties()
 		MyTableViewBase->SetEnableTouchAnimatedScrolling(bInEnableTouchAnimatedScrolling);
 		MyTableViewBase->SetIsRightClickScrollingEnabled(bEnableRightClickScrolling);
 		MyTableViewBase->SetIsTouchScrollingEnabled(bEnableTouchScrolling);
+		MyTableViewBase->SetIsPointerScrollingEnabled(bIsPointerScrollingEnabled);
 		MyTableViewBase->SetAllowOverscroll(AllowOverscroll ? EAllowOverscroll::Yes : EAllowOverscroll::No);
 		MyTableViewBase->SetFixedLineScrollOffset(bEnableFixedLineOffset ? TOptional<double>(FixedLineScrollOffset) : TOptional<double>());
 		MyTableViewBase->SetWheelScrollMultiplier(GetGlobalScrollAmount() * WheelScrollMultiplier);
diff --git a/Engine/Source/Runtime/UMG/Public/Components/ListViewBase.h b/Engine/Source/Runtime/UMG/Public/Components/ListViewBase.h
--- a/Engine/Source/Runtime/UMG/Public/Components/ListViewBase.h
+++ b/Engine/Source/Runtime/UMG/Public/Components/ListViewBase.h
@@ -539,6 +539,10 @@ public:
 	UFUNCTION(BlueprintCallable, Category = ListViewBase)
 	UMG_API void SetScrollbarVisibility(ESlateVisibility InVisibility);
 
+	/** Enable/Disable the ability of the list to scroll. This should be use as a temporary disable. */
+	UFUNCTION(BlueprintCallable, Category = ListViewBase)
+	UMG_API void SetIsPointerScrollingEnabled(bool bInIsPointerScrollingEnabled);
+
 	/**
 	 * Sets the list to refresh on the next tick.
 	 *
@@ -687,6 +691,10 @@ protected:
 	UPROPERTY(EditAnywhere, Category = Scrolling)
 	bool bEnableTouchScrolling = true;
 
+	/** Enable/Disable scrolling using Touch or Mouse. */
+	UPROPERTY(EditDefaultsOnly, Category = Scrolling)
+	bool bIsPointerScrollingEnabled = true;
+
 	UPROPERTY(EditAnywhere, Category = Scrolling)
 	bool bEnableFixedLineOffset = false;
 
