MultiBox: Fixed MenuEntries built using FMenuEntryParams having their type set to none.

diff --git a/Engine/Source/Runtime/Slate/Public/Framework/MultiBox/MultiBoxBuilder.h b/Engine/Source/Runtime/Slate/Public/Framework/MultiBox/MultiBoxBuilder.h
--- a/Engine/Source/Runtime/Slate/Public/Framework/MultiBox/MultiBoxBuilder.h
+++ b/Engine/Source/Runtime/Slate/Public/Framework/MultiBox/MultiBoxBuilder.h
@@ -177,6 +177,11 @@ protected:
 /** Helper struct that holds FMenuEntry params for construction */
 struct FMenuEntryParams : public FMultiBlock::FMultiBlockParams
 {
+	FMenuEntryParams()
+	{
+		Type = EMultiBlockType::MenuEntry;
+	}
+
 	/** Optional overridden text label for this menu entry.  If not set, then the action's label will be used instead. */
 	TAttribute<FText> LabelOverride;
 
