Fix tabs not being closed when tab manager permissions are changed
#rnx

diff --git a/Engine/Source/Runtime/Slate/Private/Widgets/Docking/SDockTab.cpp b/Engine/Source/Runtime/Slate/Private/Widgets/Docking/SDockTab.cpp
--- a/Engine/Source/Runtime/Slate/Private/Widgets/Docking/SDockTab.cpp
+++ b/Engine/Source/Runtime/Slate/Private/Widgets/Docking/SDockTab.cpp
@@ -13,6 +13,7 @@
 #include "Framework/Docking/SDockingArea.h"
 #include "Widgets/Colors/SComplexGradient.h"
 #include "Widgets/Layout/SBox.h"
+#include "Misc/NamePermissionList.h"
 
 namespace SDockTabDefs
 {
@@ -408,7 +409,23 @@ void SDockTab::ActivateInParent(ETabActivationCause InActivationCause)
 
 void SDockTab::SetTabManager( const TSharedPtr<FTabManager>& InTabManager)
 {
+	{
+		TSharedPtr<FTabManager> Previous = MyTabManager.Pin();
+		if (Previous && Previous != InTabManager)
+		{
+			Previous->GetTabPermissionList()->OnFilterChanged().RemoveAll(this);
+		}
+	}
+
 	MyTabManager = InTabManager;
+
+	if (InTabManager)
+	{
+		if (!InTabManager->GetTabPermissionList()->OnFilterChanged().IsBoundToObject(this))
+		{
+			InTabManager->GetTabPermissionList()->OnFilterChanged().AddSP(this, &SDockTab::CheckTabAllowed);
+		}
+	}
 }
 
 void SDockTab::SetOnPersistVisualState( const FOnPersistVisualState& Handler )
@@ -748,6 +765,22 @@ void SDockTab::NotifyTabRelocated()
 	OnTabRelocated.ExecuteIfBound();
 }
 
+void SDockTab::CheckTabAllowed()
+{
+	if (!LayoutIdentifier.TabType.IsNone())
+	{
+		TSharedPtr<FTabManager> Manager = MyTabManager.Pin();
+		if (Manager)
+		{
+			if (!Manager->GetTabPermissionList()->PassesFilter(LayoutIdentifier.TabType))
+			{
+				UE_LOG(LogSlate, Verbose, TEXT("Closing tab because type does not pass filter: %s"), *LayoutIdentifier.TabType.ToString());
+				RequestCloseTab();
+			}
+		}
+	}
+}
+
 const FDockTabStyle& SDockTab::GetCurrentStyle() const
 {
 	if ( GetVisualTabRole() == ETabRole::MajorTab )
diff --git a/Engine/Source/Runtime/Slate/Public/Widgets/Docking/SDockTab.h b/Engine/Source/Runtime/Slate/Public/Widgets/Docking/SDockTab.h
--- a/Engine/Source/Runtime/Slate/Public/Widgets/Docking/SDockTab.h
+++ b/Engine/Source/Runtime/Slate/Public/Widgets/Docking/SDockTab.h
@@ -401,6 +401,9 @@ private:
 
 	SLATE_API void NotifyTabRelocated();
 
+	/* Closes tab if permissions filter no longer allows this tab to be open. */
+	SLATE_API void CheckTabAllowed();
+
 	/** The handle to the active tab activation tick */
 	TWeakPtr<FActiveTimerHandle> DragDropTimerHandle;
 	TWeakPtr<FActiveTimerHandle> UpdateStyleTimerHandle;
