Skip repeated calls to RequestTreeRefresh/RequestListRefresh until the first call is handled.
Reduces cost of a large tree rebuild in SPathView from 6s to 1s without having to refactor SPathView.

diff --git a/Engine/Source/Runtime/Slate/Public/Widgets/Views/STreeView.h b/Engine/Source/Runtime/Slate/Public/Widgets/Views/STreeView.h
--- a/Engine/Source/Runtime/Slate/Public/Widgets/Views/STreeView.h
+++ b/Engine/Source/Runtime/Slate/Public/Widgets/Views/STreeView.h
@@ -470,7 +470,8 @@ private:
 
 	virtual bool Private_DoesItemHaveChildren( int32 ItemIndexInList ) const override
 	{
-		bool bHasChildren = false;
+		bool bHasChildren = false;
+
 		if (DenseItemInfos.IsValidIndex(ItemIndexInList))
 		{
 			bHasChildren = DenseItemInfos[ItemIndexInList].bHasChildren;
@@ -843,8 +844,11 @@ public:
 	/** Queue up a regeneration of the linearized items on the next tick. */
 	virtual void RequestListRefresh() override
 	{
-		bTreeItemsAreDirty = true;
-		SListView<ItemType>::RequestListRefresh();
+		if (!bTreeItemsAreDirty)
+		{
+			bTreeItemsAreDirty = true;
+			SListView<ItemType>::RequestListRefresh();
+		}
 	}
 
 	void RequestTreeRefresh()
