Make UToolMenu::bSearchable control whether the menu's search field is available

Setting UToolMenu::bSearchable to false in a menu delegate will now disable the search field completely for that menu (hidden and non-functional even if the user types into the menu). However, menu entries from this unsearchable menu will still show if a parent menu is searched.

diff --git a/Engine/Source/Developer/ToolMenus/Private/ToolMenus.cpp b/Engine/Source/Developer/ToolMenus/Private/ToolMenus.cpp
--- a/Engine/Source/Developer/ToolMenus/Private/ToolMenus.cpp
+++ b/Engine/Source/Developer/ToolMenus/Private/ToolMenus.cpp
@@ -1347,6 +1347,8 @@ TSharedRef<SWidget> UToolMenus::GenerateToolbarComboButtonMenu(TWeakObjectPtr<UT
 
 void UToolMenus::PopulateMenuBuilder(FMenuBuilder& MenuBuilder, UToolMenu* MenuData)
 {
+	MenuBuilder.SetSearchable(MenuData->bSearchable);
+
 	const bool bIsEditing = MenuData->IsEditing();
 	if (GetEditMenusMode() && !bIsEditing && EditMenuDelegate.IsBound())
 	{
diff --git a/Engine/Source/Runtime/Slate/Private/Framework/MultiBox/MultiBox.cpp b/Engine/Source/Runtime/Slate/Private/Framework/MultiBox/MultiBox.cpp
--- a/Engine/Source/Runtime/Slate/Private/Framework/MultiBox/MultiBox.cpp
+++ b/Engine/Source/Runtime/Slate/Private/Framework/MultiBox/MultiBox.cpp
@@ -874,18 +874,14 @@ bool SMultiBoxWidget::GetSearchable() const
 /** Creates the SearchTextWidget if the MultiBox has requested one */
 void SMultiBoxWidget::CreateSearchTextWidget()
 {
-	if (!MultiBox->bHasSearchWidget)
+	if (!MultiBox->bHasSearchWidget || !bSearchable)
 	{
 		return;
 	}
 
-	const FText SearchHint = ShouldShowMenuSearchField()
-							   ? LOCTEXT("SearchHintStartTyping", "Start typing to search")
-							   : LOCTEXT("SearchHint", "Search");
-
 	SearchTextWidget =
 		SNew(SSearchBox)
-			.HintText(SearchHint)
+			.HintText(LOCTEXT("SearchHintStartTyping", "Start typing to search"))
 			.SelectAllTextWhenFocused(false)
 			.OnTextChanged(this, &SMultiBoxWidget::OnFilterTextChanged);
 
@@ -913,7 +909,6 @@ void SMultiBoxWidget::OnFilterTextChanged(const FText& InFilterText)
 	{
 		if (SearchTextWidget.IsValid() && SearchBlockWidget.IsValid())
 		{
-			// We only have to do this if we're not always showing the search widget.
 			if (!ShouldShowMenuSearchField())
 			{
 				// Make the search box visible and focused
@@ -1650,7 +1645,6 @@ void SMultiBoxWidget::BeginSearch(const TCHAR InChar)
 
 		if (SearchTextWidget.IsValid() && SearchBlockWidget.IsValid())
 		{
-			// We only have to do this if we're not always showing the search widget.
 			if (!ShouldShowMenuSearchField())
 			{
 				// Make the search box visible and focused
diff --git a/Engine/Source/Runtime/Slate/Public/Framework/MultiBox/MultiBoxBuilder.h b/Engine/Source/Runtime/Slate/Public/Framework/MultiBox/MultiBoxBuilder.h
--- a/Engine/Source/Runtime/Slate/Public/Framework/MultiBox/MultiBoxBuilder.h
+++ b/Engine/Source/Runtime/Slate/Public/Framework/MultiBox/MultiBoxBuilder.h
@@ -425,6 +425,8 @@ public:
 	 */
 	SLATE_API void AddWrapperSubMenu( const FText& InMenuLabel, const FText& InToolTip, const TSharedPtr<SWidget>& InSubMenu, const FSlateIcon& InIcon );
 
+	SLATE_API void SetSearchable(bool bIsSearchable) { bSearchable = bIsSearchable; };
+
 private:
 	/** Current extension hook name for sections to determine where sections begin and end */
 	FName CurrentSectionExtensionHook;
