Fix crash in MultiBox::OnFocusChanging caused by missing check whether the new focus path has any windows.

diff --git a/Engine/Source/Runtime/Slate/Private/Framework/MultiBox/MultiBox.cpp b/Engine/Source/Runtime/Slate/Private/Framework/MultiBox/MultiBox.cpp
--- a/Engine/Source/Runtime/Slate/Private/Framework/MultiBox/MultiBox.cpp
+++ b/Engine/Source/Runtime/Slate/Private/Framework/MultiBox/MultiBox.cpp
@@ -1577,6 +1577,11 @@ FReply SMultiBoxWidget::OnFocusReceived( const FGeometry& MyGeometry, const FFoc
 
 void SMultiBoxWidget::OnFocusChanging(const FWeakWidgetPath& PreviousFocusPath, const FWidgetPath& NewWidgetPath, const FFocusEvent& InFocusEvent)
 {
+	if (!NewWidgetPath.IsValid())
+	{
+		return;
+	}
+
 	if (!ShouldShowMenuSearchField())
 	{
 		return;
