Simplify the logic for when menu search fields are shown

diff --git a/Engine/Source/Runtime/Slate/Private/Framework/MultiBox/MultiBox.cpp b/Engine/Source/Runtime/Slate/Private/Framework/MultiBox/MultiBox.cpp
--- a/Engine/Source/Runtime/Slate/Private/Framework/MultiBox/MultiBox.cpp
+++ b/Engine/Source/Runtime/Slate/Private/Framework/MultiBox/MultiBox.cpp
@@ -909,11 +909,8 @@ void SMultiBoxWidget::OnFilterTextChanged(const FText& InFilterText)
 	{
 		if (SearchTextWidget.IsValid() && SearchBlockWidget.IsValid())
 		{
-			if (!ShouldShowMenuSearchField())
-			{
-				// Make the search box visible and focused
-				SearchBlockWidget->SetVisibility(EVisibility::Visible);
-			}
+			// Make the search box visible and focused.
+			SearchBlockWidget->SetVisibility(EVisibility::Visible);
 			FSlateApplication::Get().SetUserFocus(0, SearchTextWidget);
 		}
 	}
@@ -1650,11 +1647,8 @@ void SMultiBoxWidget::BeginSearch(const TCHAR InChar)
 
 		if (SearchTextWidget.IsValid() && SearchBlockWidget.IsValid())
 		{
-			if (!ShouldShowMenuSearchField())
-			{
-				// Make the search box visible and focused
-				SearchBlockWidget->SetVisibility(EVisibility::Visible);
-			}
+			// Make the search box visible and focused.
+			SearchBlockWidget->SetVisibility(EVisibility::Visible);
 			FSlateApplication::Get().SetUserFocus(0, SearchTextWidget);
 
 			SearchTextWidget->SetText(FText::FromString(NewSearchText));
@@ -1803,13 +1797,10 @@ void SMultiBoxWidget::FilterMultiBoxEntries()
 			It.Key()->SetVisibility(EVisibility::Visible);
 		}
 
-		// We only have to do this if we're not always showing the search widget.
-		if (!ShouldShowMenuSearchField())
+		// Reset the visibility of the search field.
+		if (SearchBlockWidget.IsValid())
 		{
-			if (SearchBlockWidget.IsValid())
-			{
-				SearchBlockWidget->SetVisibility(EVisibility::Collapsed);
-			}
+			SearchBlockWidget->SetVisibility(ShouldShowMenuSearchField() ? EVisibility::Visible : EVisibility::Collapsed);
 		}
 
 		// Hide the sub-menus widgets that were made visible by searching this multi-box hierarchy.
@@ -1834,12 +1825,9 @@ void SMultiBoxWidget::FilterMultiBoxEntries()
 		const TSharedPtr<SWidget>& Widget = It.Key();
 
 		// Skip the search widget itself when scanning for searchable items.
-		if (ShouldShowMenuSearchField())
+		if (Widget == SearchBlockWidget)
 		{
-			if (Widget == SearchBlockWidget)
-			{
-				continue;
-			}
+			continue;
 		}
 
 		// Non-labeled elements should not be visible when searching
@@ -1885,16 +1873,6 @@ void SMultiBoxWidget::FilterMultiBoxEntries()
 			}
 		}
 	}
-
-	// If we always show the search widget, we're skipping it in the code above and do not need to show it here to compensate.
-	if (!ShouldShowMenuSearchField())
-	{
-		// Show the search widget again, it was hidden by the above code.
-		if (SearchBlockWidget.IsValid())
-		{
-			SearchBlockWidget->SetVisibility(EVisibility::Visible);
-		}
-	}
 }
 
 FText SMultiBoxWidget::GetSearchText() const
@@ -1910,6 +1888,7 @@ TSharedPtr<SWidget> SMultiBoxWidget::GetSearchTextWidget()
 void SMultiBoxWidget::SetSearchBlockWidget(TSharedPtr<SWidget> InWidget)
 {
 	SearchBlockWidget = InWidget;
+	InWidget->SetVisibility(ShouldShowMenuSearchField() ? EVisibility::Visible : EVisibility::Collapsed);
 }
 
 void SMultiBoxWidget::AddSearchElement( TSharedPtr<SWidget> BlockWidget, FText BlockDisplayText )
diff --git a/Engine/Source/Runtime/Slate/Private/Framework/MultiBox/SWidgetBlock.cpp b/Engine/Source/Runtime/Slate/Private/Framework/MultiBox/SWidgetBlock.cpp
--- a/Engine/Source/Runtime/Slate/Private/Framework/MultiBox/SWidgetBlock.cpp
+++ b/Engine/Source/Runtime/Slate/Private/Framework/MultiBox/SWidgetBlock.cpp
@@ -139,19 +139,15 @@ void SWidgetBlock::BuildMultiBlockWidget(const ISlateStyle* StyleSet, const FNam
 		}
 	}
 
+	const TSharedRef<SWidget> ThisWidget = this->AsWidget();
+
 	// Add this widget to the search list of the multibox
-	OwnerMultiBoxWidgetPinned->AddElement(this->AsWidget(), SearchLabel, MultiBlock->GetSearchable());
+	OwnerMultiBoxWidgetPinned->AddElement(ThisWidget, SearchLabel, MultiBlock->GetSearchable());
 
 	// This widget holds the search text, set it as the search block widget
 	if (OwnerMultiBoxWidgetPinned->GetSearchTextWidget() && OwnerMultiBoxWidgetPinned->GetSearchTextWidget()->GetParentWidget() == WidgetBlock->ContentWidget)
 	{
-		OwnerMultiBoxWidgetPinned->SetSearchBlockWidget(this->AsWidget());
-
-		// When we are always showing the search widget, we should not hide it here.
-		if (!OwnerMultiBoxWidgetPinned->ShouldShowMenuSearchField())
-		{
-			this->AsWidget()->SetVisibility(EVisibility::Collapsed);
-		}
+		OwnerMultiBoxWidgetPinned->SetSearchBlockWidget(ThisWidget);
 	}
 
 	ChildSlot
diff --git a/Engine/Source/Runtime/Slate/Public/Framework/MultiBox/MultiBox.h b/Engine/Source/Runtime/Slate/Public/Framework/MultiBox/MultiBox.h
--- a/Engine/Source/Runtime/Slate/Public/Framework/MultiBox/MultiBox.h
+++ b/Engine/Source/Runtime/Slate/Public/Framework/MultiBox/MultiBox.h
@@ -810,6 +810,7 @@ public:
 	 */
 	SLATE_API double GetSummonedMenuTime() const;
 
+	UE_DEPRECATED_FORGAME(5.5, "This function will be made private in the future.")
 	bool ShouldShowMenuSearchField();
 
 private:
