Make ESC clear filtering of a menu

This fixes an issue where if you first typed into a menu to search and filter it and then pressed ESC, then the query was cleared but the menu was still filtered. After this change, pressing ESC will in the same situation clear the search field and undo any filtering. Also, the search field will revert to its previous visibility, so if it was hidden before the user started typing it will hide after pressing ESC or if it was visible before the user started typing then the search field will just clear but stay visible when ESC is pressed.

diff --git a/Engine/Source/Runtime/Slate/Private/Framework/MultiBox/MultiBox.cpp b/Engine/Source/Runtime/Slate/Private/Framework/MultiBox/MultiBox.cpp
--- a/Engine/Source/Runtime/Slate/Private/Framework/MultiBox/MultiBox.cpp
+++ b/Engine/Source/Runtime/Slate/Private/Framework/MultiBox/MultiBox.cpp
@@ -883,7 +883,8 @@ void SMultiBoxWidget::CreateSearchTextWidget()
 		SNew(SSearchBox)
 			.HintText(LOCTEXT("SearchHintStartTyping", "Start typing to search"))
 			.SelectAllTextWhenFocused(false)
-			.OnTextChanged(this, &SMultiBoxWidget::OnFilterTextChanged);
+			.OnTextChanged(this, &SMultiBoxWidget::OnFilterTextChanged)
+			.OnTextCommitted(this, &SMultiBoxWidget::OnFilterTextCommitted);
 
 	TSharedRef<SBox> SearchBox =
 		SNew(SBox)
@@ -920,6 +921,15 @@ void SMultiBoxWidget::OnFilterTextChanged(const FText& InFilterText)
 	FilterMultiBoxEntries();
 }
 
+void SMultiBoxWidget::OnFilterTextCommitted(const FText& InFilterText, ETextCommit::Type CommitType)
+{
+	if (CommitType == ETextCommit::Type::OnCleared)
+	{
+		SearchText = InFilterText;
+		FilterMultiBoxEntries();
+	}
+}
+
 /**
  * Builds this MultiBox widget up from the MultiBox associated with it
  */
diff --git a/Engine/Source/Runtime/Slate/Public/Framework/MultiBox/MultiBox.h b/Engine/Source/Runtime/Slate/Public/Framework/MultiBox/MultiBox.h
--- a/Engine/Source/Runtime/Slate/Public/Framework/MultiBox/MultiBox.h
+++ b/Engine/Source/Runtime/Slate/Public/Framework/MultiBox/MultiBox.h
@@ -825,7 +825,10 @@ private:
 	/** Creates the SearchTextWidget if the MultiBox has requested one */
 	SLATE_API void CreateSearchTextWidget();
 
-	/** Called when the SearchText changes */
+	/** Called when the search field text is committed (such as pressing ESC). */
+	SLATE_API void OnFilterTextCommitted(const FText& InFilterText, ETextCommit::Type CommitType);
+	
+	/** Called when the search field text changes. */
 	SLATE_API void OnFilterTextChanged(const FText& InFilterText);
 
 	/**
