Fix typo in comment

diff --git a/Engine/Source/Runtime/Slate/Private/Framework/Application/SlateApplication.cpp b/Engine/Source/Runtime/Slate/Private/Framework/Application/SlateApplication.cpp
--- a/Engine/Source/Runtime/Slate/Private/Framework/Application/SlateApplication.cpp
+++ b/Engine/Source/Runtime/Slate/Private/Framework/Application/SlateApplication.cpp
@@ -4758,7 +4758,7 @@ bool FSlateApplication::ProcessKeyDownEvent( const FKeyEvent& InKeyEvent )
 		TSharedRef<FWidgetPath> EventPathRef = SlateUser->GetFocusPath();
 		const FWidgetPath& EventPath = EventPathRef.Get();
 
-		// Switch worlds for widgets inOnPreviewMouseButtonDown the current path
+		// Switch worlds for widgets in the current path
 		FScopedSwitchWorldHack SwitchWorld(EventPath);
 
 		// Tunnel the keyboard event
