Fixing material editor crash when trying to restore previously open shader code tab before shader instances are populated

diff --git a/Engine/Source/Editor/MaterialEditor/Private/MaterialStats.cpp b/Engine/Source/Editor/MaterialEditor/Private/MaterialStats.cpp
--- a/Engine/Source/Editor/MaterialEditor/Private/MaterialStats.cpp
+++ b/Engine/Source/Editor/MaterialEditor/Private/MaterialStats.cpp
@@ -834,6 +834,9 @@ void FMaterialStats::Update()
 			GetStatsGrid()->OnShaderChanged();
 		}
 
+		// update any dependent tabs
+		RefreshDependentTabs.Broadcast();
+
 		bNeedsGridRefresh = false;
 	}
 
@@ -999,8 +1002,18 @@ TSharedRef<SDockTab> FMaterialStats::SpawnTab_ShaderCode(const FSpawnTabArgs& Ar
 	TSharedPtr<FShaderPlatformSettings> PlatformPtr = GetPlatformSettings(PlatformID);
 	check(PlatformPtr.IsValid());
 
+	const TArray<TSharedPtr<FMaterialShaderEntry>> *ShaderEntriesPtr = PlatformPtr->GetShaderEntries(QualityLevel, InstanceIndex);
+
+	if (ShaderEntriesPtr == nullptr)
+	{
+		// SComboBox needs a valid pointer otherwise follow up SetItemsSource will fail.
+		static TArray<TSharedPtr<FMaterialShaderEntry>> EmptyShaderEntriesArray = {};
+		ShaderEntriesPtr = &EmptyShaderEntriesArray;
+	}
+
+	ensure(ShaderEntriesPtr != nullptr);
 	TSharedRef<SComboBox<TSharedPtr<FMaterialShaderEntry>>> ShaderBox = SNew(SComboBox<TSharedPtr<FMaterialShaderEntry>>)
-		.OptionsSource(PlatformPtr->GetShaderEntries(QualityLevel, InstanceIndex))
+		.OptionsSource(ShaderEntriesPtr)
 		.OnGenerateWidget_Lambda([](TSharedPtr<FMaterialShaderEntry> Value) { return SNew(STextBlock).Text(FText::FromString(Value->Text)); })
 		.OnSelectionChanged_Lambda([PlatformPtr, QualityLevel, InstanceIndex](TSharedPtr<FMaterialShaderEntry> Item, ESelectInfo::Type SelectInfo) { PlatformPtr->OnShaderViewComboSelectionChanged(Item, QualityLevel, InstanceIndex); })
 		[
@@ -1008,7 +1021,26 @@ TSharedRef<SDockTab> FMaterialStats::SpawnTab_ShaderCode(const FSpawnTabArgs& Ar
 			.Text_Lambda([PlatformPtr, QualityLevel, InstanceIndex]() { return PlatformPtr->GetSelectedShaderViewComboText(QualityLevel, InstanceIndex); })
 		];
 
+	// Refresh ShaderBox when shaders are updated
+	FDelegateHandle ShaderBoxUpdater = RefreshDependentTabs.AddLambda([PlatformPtrWeakPtr = PlatformPtr.ToWeakPtr(), ShaderBoxWeakPtr = ShaderBox.ToWeakPtr(), QualityLevel, InstanceIndex]()
+	{
+		TSharedPtr<FShaderPlatformSettings> PlatformPtr = PlatformPtrWeakPtr.Pin();
+		TSharedPtr<SComboBox<TSharedPtr<FMaterialShaderEntry>>> ShaderBox = ShaderBoxWeakPtr.Pin();
+		if (PlatformPtr.IsValid() && ShaderBox.IsValid())
+		{
+			ShaderBox->SetItemsSource(PlatformPtr->GetShaderEntries(QualityLevel, InstanceIndex));
+		}
+	});
+
 	TSharedRef<SDockTab> SpawnedTab = SNew(SDockTab)
+		.OnTabClosed_Lambda([MaterialStatsWeakPtr = this->AsWeak(), ShaderBoxUpdater](TSharedRef<SDockTab> Tab)
+		{
+			TSharedPtr<FMaterialStats> MaterialStats = MaterialStatsWeakPtr.Pin();
+			if (MaterialStats.IsValid() && ShaderBoxUpdater.IsValid())
+			{
+				MaterialStats->RefreshDependentTabs.Remove(ShaderBoxUpdater);
+			}
+		})
 		[
 			SNew(SVerticalBox)
 			+ SVerticalBox::Slot()
diff --git a/Engine/Source/Editor/MaterialEditor/Private/MaterialStats.h b/Engine/Source/Editor/MaterialEditor/Private/MaterialStats.h
--- a/Engine/Source/Editor/MaterialEditor/Private/MaterialStats.h
+++ b/Engine/Source/Editor/MaterialEditor/Private/MaterialStats.h
@@ -228,7 +228,7 @@ public:
 	/** retrives the pointer to the spawned window tab that contains the shader code */
 	FORCEINLINE TWeakPtr<class SDockTab> GetCodeViewerTab(const EMaterialQualityLevel::Type QualityLevel);
 
-	/** returns an array with the names of all the compiled shaders for this material with specified quality level */
+	/** returns an array with the names of all the compiled shaders for this material with specified quality level, can return nullptr if no instance index is not available */
 	FORCEINLINE const TArray<TSharedPtr<FMaterialShaderEntry>> *GetShaderEntries(const EMaterialQualityLevel::Type QualityLevel, const int32 InstanceIndex);
 
 	/** when set this flag will indicate the presence of this material with a particular quality level inside the stats widget */
@@ -352,6 +352,8 @@ class FMaterialStats : public FGCObject, public TSharedFromThis<FMaterialStats>
 
 	bool bNeedsGridRefresh = false;
 
+	TMulticastDelegate<void()> RefreshDependentTabs;
+
 private:
 	/** adds a specified platform in the grid widget for analysis; usually called from BuildShaderPlatformDB() */
 	TSharedPtr<FShaderPlatformSettings> AddShaderPlatform(const EPlatformCategoryType PlatformType, const EShaderPlatform PlatformID, const FName PlatformName, const bool bAllowCodeView, const FString& Description, const bool bAlwaysOn = false);
@@ -646,8 +648,16 @@ FORCEINLINE TWeakPtr<class SDockTab> FShaderPlatformSettings::GetCodeViewerTab(c
 
 FORCEINLINE const TArray<TSharedPtr<FMaterialShaderEntry>> *FShaderPlatformSettings::GetShaderEntries(const EMaterialQualityLevel::Type QualityLevel, const int32 InstanceIndex)
 {
-	auto& InstanceData = GetInstanceData(QualityLevel, InstanceIndex);
-	return &InstanceData.ArrShaderEntries;
+	check(QualityLevel < EMaterialQualityLevel::Num);
+	if (InstanceIndex < PlatformData[QualityLevel].Instances.Num())
+	{
+		auto& InstanceData = GetInstanceData(QualityLevel, InstanceIndex);
+		return &InstanceData.ArrShaderEntries;
+	}
+	else
+	{
+		return nullptr;
+	}
 }
 
 FORCEINLINE void FShaderPlatformSettings::SetExtractStatsFlag(const EMaterialQualityLevel::Type QualityType, const bool bValue)
diff --git a/Engine/Source/Runtime/Slate/Private/Widgets/Input/STextComboBox.cpp b/Engine/Source/Runtime/Slate/Private/Widgets/Input/STextComboBox.cpp
--- a/Engine/Source/Runtime/Slate/Private/Widgets/Input/STextComboBox.cpp
+++ b/Engine/Source/Runtime/Slate/Private/Widgets/Input/STextComboBox.cpp
@@ -15,7 +15,7 @@ void STextComboBox::Construct( const FArguments& InArgs )
 		SAssignNew(StringCombo, SComboBox< TSharedPtr<FString> > )
 		.ComboBoxStyle(InArgs._ComboBoxStyle)
 		.ButtonStyle(InArgs._ButtonStyle)
-		.OptionsSource(InArgs._OptionsSource)
+		.OptionsSource(InArgs.GetOptionsSource())
 		.OnGenerateWidget(this, &STextComboBox::MakeItemWidget)
 		.OnSelectionChanged(this, &STextComboBox::OnSelectionChanged)
 		.OnComboBoxOpening(InArgs._OnComboBoxOpening)
diff --git a/Engine/Source/Runtime/Slate/Public/Widgets/Input/SComboBox.h b/Engine/Source/Runtime/Slate/Public/Widgets/Input/SComboBox.h
--- a/Engine/Source/Runtime/Slate/Public/Widgets/Input/SComboBox.h
+++ b/Engine/Source/Runtime/Slate/Public/Widgets/Input/SComboBox.h
@@ -114,7 +114,6 @@ public:
 		, _ScrollBarStyle(&FAppStyle::Get().GetWidgetStyle<FScrollBarStyle>("ScrollBar"))
 		, _ContentPadding(_ComboBoxStyle->ContentPadding)
 		, _ForegroundColor(FSlateColor::UseStyle())
-		, _OptionsSource()
 		, _OnSelectionChanged()
 		, _OnGenerateWidget()
 		, _InitiallySelectedItem(ListTypeTraits::MakeNullPtr())
@@ -140,7 +139,7 @@ public:
 		SLATE_ATTRIBUTE( FMargin, ContentPadding )
 		SLATE_ATTRIBUTE( FSlateColor, ForegroundColor )
 
-		SLATE_ARGUMENT( const TArray< OptionType >*, OptionsSource )
+		SLATE_ITEMS_SOURCE_ARGUMENT( OptionType, OptionsSource )
 		SLATE_EVENT( FOnSelectionChanged, OnSelectionChanged )
 		SLATE_EVENT( FOnGenerateWidget, OnGenerateWidget )
 
@@ -210,7 +209,6 @@ public:
 		this->EnableGamepadNavigationMode = InArgs._EnableGamepadNavigationMode;
 		this->bControllerInputCaptured = false;
 
-		OptionsSource = InArgs._OptionsSource;
 		CustomScrollbar = InArgs._CustomScrollbar;
 
 		ComboBoxMenuContent =
@@ -218,7 +216,7 @@ public:
 			.MaxDesiredHeight(InArgs._MaxListHeight)
 			[
 				SAssignNew(this->ComboListView, SComboListType)
-				.ListItemsSource(InArgs._OptionsSource)
+				.ListItemsSource(InArgs.GetOptionsSource())
 				.OnGenerateRow(this, &SComboBox< OptionType >::GenerateMenuItemRow)
 				.OnSelectionChanged(this, &SComboBox< OptionType >::OnSelectionChanged_Internal)
 				.OnKeyDownHandler(this, &SComboBox< OptionType >::OnKeyDownHandler)
@@ -410,6 +408,23 @@ public:
 		return SelectedItem;
 	}
 
+	/** Sets new item source */
+	void SetItemsSource(const TArray<OptionType>* InListItemsSource)
+	{
+		ComboListView->SetItemsSource(InListItemsSource);
+	}
+
+	/** Sets new item source */
+	void SetItemsSource(TSharedRef<::UE::Slate::Containers::TObservableArray<OptionType>> InListItemsSource)
+	{
+		ComboListView->SetItemsSource(InListItemsSource);
+	}
+
+	/** Clears current item source */
+	void ClearItemsSource()
+	{
+		ComboListView->ClearItemsSource();
+	}
 
 	/** 
 	 * Requests a list refresh after updating options 
@@ -488,11 +503,12 @@ protected:
 					if (TListTypeTraits<OptionType>::IsPtrValid(NullableSelected))
 					{
 						OptionType ActuallySelected = TListTypeTraits<OptionType>::NullableItemTypeConvertToItemType(NullableSelected);
-						const int32 SelectionIndex = OptionsSource->Find(ActuallySelected);
+						const TArrayView<const OptionType> OptionsSource = ComboListView->GetItems();
+						const int32 SelectionIndex = OptionsSource.Find(ActuallySelected);
 						if (SelectionIndex >= 1)
 						{
 							// Select an item on the prev row
-							SetSelectedItem((*OptionsSource)[SelectionIndex - 1]);
+							SetSelectedItem(OptionsSource[SelectionIndex - 1]);
 						}
 					}
 
@@ -504,11 +520,12 @@ protected:
 					if (TListTypeTraits<OptionType>::IsPtrValid(NullableSelected))
 					{
 						OptionType ActuallySelected = TListTypeTraits<OptionType>::NullableItemTypeConvertToItemType(NullableSelected);
-						const int32 SelectionIndex = OptionsSource->Find(ActuallySelected);
-						if (SelectionIndex < OptionsSource->Num() - 1)
+						const TArrayView<const OptionType> OptionsSource = ComboListView->GetItems();
+						const int32 SelectionIndex = OptionsSource.Find(ActuallySelected);
+						if (SelectionIndex < OptionsSource.Num() - 1)
 						{
 							// Select an item on the next row
-							SetSelectedItem((*OptionsSource)[SelectionIndex + 1]);
+							SetSelectedItem(OptionsSource[SelectionIndex + 1]);
 						}
 					}
 					return FReply::Handled();
@@ -687,7 +704,5 @@ private:
 	bool bControllerInputCaptured;
 
 	TSharedPtr<SBox> ComboBoxMenuContent;
-
-	const TArray< OptionType >* OptionsSource;
 };
 
diff --git a/Engine/Source/Runtime/Slate/Public/Widgets/Input/STextComboBox.h b/Engine/Source/Runtime/Slate/Public/Widgets/Input/STextComboBox.h
--- a/Engine/Source/Runtime/Slate/Public/Widgets/Input/STextComboBox.h
+++ b/Engine/Source/Runtime/Slate/Public/Widgets/Input/STextComboBox.h
@@ -39,7 +39,7 @@ public:
 		SLATE_STYLE_ARGUMENT(FButtonStyle, ButtonStyle)
 
 		/** Selection of strings to pick from */
-		SLATE_ARGUMENT( TArray< TSharedPtr<FString> >*, OptionsSource )
+		SLATE_ITEMS_SOURCE_ARGUMENT(TSharedPtr<FString>, OptionsSource)
 
 		/** Text color and opacity */
 		SLATE_ATTRIBUTE( FSlateColor, ColorAndOpacity )
@@ -76,6 +76,24 @@ public:
 		return SelectedItem;
 	}
 
+	/** Sets new item source */
+	void SetItemsSource(const TArray<TSharedPtr<FString>>* InListItemsSource)
+	{
+		StringCombo->SetItemsSource(InListItemsSource);
+	}
+
+	/** Sets new item source */
+	void SetItemsSource(TSharedRef<::UE::Slate::Containers::TObservableArray<TSharedPtr<FString>>> InListItemsSource)
+	{
+		StringCombo->SetItemsSource(InListItemsSource);
+	}
+
+	/** Clears current item source */
+	void ClearItemsSource()
+	{
+		StringCombo->ClearItemsSource();
+	}
+
 	/** Request to reload the text options in the combobox from the OptionsSource attribute */
 	SLATE_API void RefreshOptions();
 
diff --git a/Engine/Source/Runtime/Slate/Public/Widgets/Views/IItemsSource.h b/Engine/Source/Runtime/Slate/Public/Widgets/Views/IItemsSource.h
--- a/Engine/Source/Runtime/Slate/Public/Widgets/Views/IItemsSource.h
+++ b/Engine/Source/Runtime/Slate/Public/Widgets/Views/IItemsSource.h
@@ -7,8 +7,15 @@
 #include "Containers/ObservableArray.h"
 #include "Widgets/Views/STableViewBase.h"
 
- namespace UE::Slate::ItemsSource
- {
+namespace UE::Slate::ItemsSource
+{
+template<typename ArgType>
+struct ForwardedSlateItemsSourceArgument
+{
+	const TArray<ArgType>* ArrayPointer;
+	::UE::Slate::Containers::TObservableArray<ArgType>* ObservableArrayPointer;
+	TSharedPtr<::UE::Slate::Containers::TObservableArray<ArgType>> SharedObservableArray;
+};
 
 #define SLATE_ITEMS_SOURCE_ARGUMENT( ArgType, ArgName ) \
 	private: \
@@ -46,7 +53,19 @@
 			_##ArgName##_SharedObservableArray = InArg; \
 			return static_cast<WidgetArgsType*>(this)->Me(); \
 		} \
-		TUniquePtr<::UE::Slate::ItemsSource::IItemsSource<ArgType>> Make##ArgName (TSharedRef<STableViewBase> InWidget) const \
+		WidgetArgsType& ArgName(::UE::Slate::ItemsSource::ForwardedSlateItemsSourceArgument<ArgType> InArg) \
+		{ \
+			_Reset##ArgName(); \
+			_##ArgName##_ArrayPointer = InArg.ArrayPointer; \
+			_##ArgName##_ObservableArrayPointer = InArg.ObservableArrayPointer; \
+			_##ArgName##_SharedObservableArray = InArg.SharedObservableArray; \
+			return static_cast<WidgetArgsType*>(this)->Me(); \
+		} \
+		::UE::Slate::ItemsSource::ForwardedSlateItemsSourceArgument<ArgType> Get##ArgName() const \
+		{ \
+			return {_##ArgName##_ArrayPointer, _##ArgName##_ObservableArrayPointer, _##ArgName##_SharedObservableArray}; \
+		} \
+		TUniquePtr<::UE::Slate::ItemsSource::IItemsSource<ArgType>> Make##ArgName(TSharedRef<STableViewBase> InWidget) const \
 		{ \
 			if (_##ArgName##_ArrayPointer) \
 			{ \
