[Details Panel] Fix in progress spin box transactions which would never complete and hold undo transactions open until the editor restarted

Also fix an issue where widgets from previously set objects in a details panel were never destructed if details panel was set to have no objects later

diff --git a/Engine/Source/Editor/PropertyEditor/Private/SDetailsViewBase.cpp b/Engine/Source/Editor/PropertyEditor/Private/SDetailsViewBase.cpp
--- a/Engine/Source/Editor/PropertyEditor/Private/SDetailsViewBase.cpp
+++ b/Engine/Source/Editor/PropertyEditor/Private/SDetailsViewBase.cpp
@@ -488,7 +488,7 @@ EVisibility SDetailsViewBase::GetTreeVisibility() const
 		}
 	}
 
-	return EVisibility::Collapsed;
+	return DetailTree->IsPendingRefresh() ? EVisibility::Visible : EVisibility::Collapsed;
 }
 
 EVisibility SDetailsViewBase::GetScrollBarVisibility() const
diff --git a/Engine/Source/Runtime/Slate/Public/Widgets/Input/SSpinBox.h b/Engine/Source/Runtime/Slate/Public/Widgets/Input/SSpinBox.h
--- a/Engine/Source/Runtime/Slate/Public/Widgets/Input/SSpinBox.h
+++ b/Engine/Source/Runtime/Slate/Public/Widgets/Input/SSpinBox.h
@@ -176,6 +176,14 @@ public:
 	{
 	}
 
+	virtual ~SSpinBox()
+	{
+		if (bDragging || PointerDraggingSliderIndex != INDEX_NONE)
+		{
+			CancelMouseCapture();
+		}
+	}
+
 	/**
 	 * Construct the widget
 	 *
@@ -772,11 +780,7 @@ public:
 		const FKey Key = InKeyEvent.GetKey();
 		if (Key == EKeys::Escape && HasMouseCapture())
 		{
-			bDragging = false;
-			PointerDraggingSliderIndex = INDEX_NONE;
-
-			InternalValue = (double)PreDragValue;
-			NotifyValueCommitted(PreDragValue);
+			CancelMouseCapture();
 			return FReply::Handled().ReleaseMouseCapture().SetMousePos(CachedMousePosition);
 		}
 		else if (Key == EKeys::Up || Key == EKeys::Right)
@@ -1225,6 +1229,15 @@ private:
 		}
 	}
 
+	void CancelMouseCapture()
+	{
+		bDragging = false;
+		PointerDraggingSliderIndex = INDEX_NONE;
+
+		InternalValue = (double)PreDragValue;
+		NotifyValueCommitted(PreDragValue);
+	}
+
 	/** Tracks which cursor is currently dragging the slider (e.g., the mouse cursor or a specific finger) */
 	int32 PointerDraggingSliderIndex;
 
