Fix for crash in TabStack due to trying to remove a non-existent tab.

diff --git a/Engine/Source/Runtime/Slate/Private/Framework/Docking/SDockingTabStack.cpp b/Engine/Source/Runtime/Slate/Private/Framework/Docking/SDockingTabStack.cpp
--- a/Engine/Source/Runtime/Slate/Private/Framework/Docking/SDockingTabStack.cpp
+++ b/Engine/Source/Runtime/Slate/Private/Framework/Docking/SDockingTabStack.cpp
@@ -1317,7 +1317,10 @@ int32 SDockingTabStack::ClosePersistentTab( const FTabId& TabId )
 void SDockingTabStack::RemovePersistentTab( const FTabId& TabId )
 {
 	const int32 TabIndex = Tabs.IndexOfByPredicate(FTabMatcher(TabId));
-	Tabs.RemoveAtSwap(TabIndex);
+	if(TabIndex != INDEX_NONE)
+	{
+		Tabs.RemoveAtSwap(TabIndex);
+	}
 }
 
 EVisibility SDockingTabStack::GetMaximizeSpacerVisibility() const
