PR #11629: Let STableViewBase use scrollbar thickness from scrollbar style

diff --git a/Engine/Source/Runtime/Slate/Private/Widgets/Views/STableViewBase.cpp b/Engine/Source/Runtime/Slate/Private/Widgets/Views/STableViewBase.cpp
--- a/Engine/Source/Runtime/Slate/Private/Widgets/Views/STableViewBase.cpp
+++ b/Engine/Source/Runtime/Slate/Private/Widgets/Views/STableViewBase.cpp
@@ -97,7 +97,7 @@ void STableViewBase::ConstructChildren( const TAttribute<float>& InItemWidth, co
 			.Style(InScrollBarStyle ? InScrollBarStyle : &FAppStyle::Get().GetWidgetStyle<FScrollBarStyle>("ScrollBar"))
 			.PreventThrottling(bInPreventThrottling);
 
-		const FOptionalSize ScrollBarSize(16.f);
+		const FOptionalSize ScrollBarSize(InScrollBarStyle ? InScrollBarStyle->Thickness : 16.f);
 
 		if (Orientation == Orient_Vertical)
 		{
