Add support for disabling gamepad navigation scrolling in listview

#tests PIE, use case tested by JIRA reporter

diff --git a/Engine/Source/Runtime/Slate/Private/Widgets/Views/STableViewBase.cpp b/Engine/Source/Runtime/Slate/Private/Widgets/Views/STableViewBase.cpp
--- a/Engine/Source/Runtime/Slate/Private/Widgets/Views/STableViewBase.cpp
+++ b/Engine/Source/Runtime/Slate/Private/Widgets/Views/STableViewBase.cpp
@@ -936,6 +936,11 @@ void STableViewBase::SetIsPointerScrollingEnabled(bool bInIsPointerScrollingEnab
 	bIsPointerScrollingEnabled = bInIsPointerScrollingEnabled;
 }
 
+void STableViewBase::SetIsGamepadScrollingEnabled(bool bInIsGamepadScrollingEnabled)
+{
+	bIsGamepadScrollingEnabled = bInIsGamepadScrollingEnabled;
+}
+
 void STableViewBase::SetBackgroundBrush(const TAttribute<const FSlateBrush*>& InBackgroundBrush)
 {
 	BackgroundBrush.SetImage(*this, InBackgroundBrush);
diff --git a/Engine/Source/Runtime/Slate/Public/Widgets/Views/SListView.h b/Engine/Source/Runtime/Slate/Public/Widgets/Views/SListView.h
--- a/Engine/Source/Runtime/Slate/Public/Widgets/Views/SListView.h
+++ b/Engine/Source/Runtime/Slate/Public/Widgets/Views/SListView.h
@@ -525,7 +525,7 @@ public:
 			}
 
 			// If it's valid we'll scroll it into view and return an explicit widget in the FNavigationReply
-			if (ItemsSourceRef.IsValidIndex(AttemptSelectIndex))
+			if (ItemsSourceRef.IsValidIndex(AttemptSelectIndex) && this->bIsGamepadScrollingEnabled)
 			{
 				TOptional<ItemType> ItemToSelect = Private_FindNextSelectableOrNavigableWithIndexAndDirection(ItemsSourceRef[AttemptSelectIndex], AttemptSelectIndex, AttemptSelectIndex >= CurSelectionIndex);
 				if (ItemToSelect.IsSet())
diff --git a/Engine/Source/Runtime/Slate/Public/Widgets/Views/STableViewBase.h b/Engine/Source/Runtime/Slate/Public/Widgets/Views/STableViewBase.h
--- a/Engine/Source/Runtime/Slate/Public/Widgets/Views/STableViewBase.h
+++ b/Engine/Source/Runtime/Slate/Public/Widgets/Views/STableViewBase.h
@@ -185,9 +185,12 @@ public:
 	/** Sets the multiplier applied when wheel scrolling. Higher numbers will cover more distance per click of the wheel. */
 	SLATE_API void SetWheelScrollMultiplier(float NewWheelScrollMultiplier);
 
-	/** Enables/disables being able to scroll. This should be use as a temporary mean to disable scrolling. */
+	/** Enables/disables being able to scroll via pointer. This should be use as a temporary mean to disable scrolling. */
 	SLATE_API void SetIsPointerScrollingEnabled(bool bInIsPointerScrollingEnabled);
 
+	/** Enables/disables being able to scroll via gamepad. */
+	SLATE_API void SetIsGamepadScrollingEnabled(bool bInIsGamepadScrollingEnabled);
+
 	/** Sets the Background Brush */
 	SLATE_API void SetBackgroundBrush(const TAttribute<const FSlateBrush*>& InBackgroundBrush);
 
@@ -478,9 +481,12 @@ protected:
 	/** How much to scroll when using mouse wheel */
 	float WheelScrollMultiplier;
 
-	/** Wheter the list is allowed to scroll. */
+	/** Whether the list is allowed to scroll via cursor. */
 	bool bIsPointerScrollingEnabled = true;
 
+	/** Whether the list is allowed to scroll via gamepad. */
+	bool bIsGamepadScrollingEnabled = true;
+
 	/** The layout and scroll orientation of the list */
 	EOrientation Orientation = Orient_Vertical;
 
diff --git a/Engine/Source/Runtime/Slate/Public/Widgets/Views/STileView.h b/Engine/Source/Runtime/Slate/Public/Widgets/Views/STileView.h
--- a/Engine/Source/Runtime/Slate/Public/Widgets/Views/STileView.h
+++ b/Engine/Source/Runtime/Slate/Public/Widgets/Views/STileView.h
@@ -256,7 +256,7 @@ public:
 			}
 
 			// If it's valid we'll scroll it into view and return an explicit widget in the FNavigationReply
-			if (ItemsSourceRef.IsValidIndex(AttemptSelectIndex))
+			if (ItemsSourceRef.IsValidIndex(AttemptSelectIndex) && this->bIsGamepadScrollingEnabled)
 			{
 				this->NavigationSelect(ItemsSourceRef[AttemptSelectIndex], InNavigationEvent);
 				return FNavigationReply::Explicit(nullptr);
diff --git a/Engine/Source/Runtime/UMG/Private/Components/ListViewBase.cpp b/Engine/Source/Runtime/UMG/Private/Components/ListViewBase.cpp
--- a/Engine/Source/Runtime/UMG/Private/Components/ListViewBase.cpp
+++ b/Engine/Source/Runtime/UMG/Private/Components/ListViewBase.cpp
@@ -99,6 +99,15 @@ UMG_API void UListViewBase::SetIsPointerScrollingEnabled(bool bInIsPointerScroll
 	}
 }
 
+UMG_API void UListViewBase::SetIsGamepadScrollingEnabled(bool bInIsGamepadScrollingEnabled)
+{
+	bIsGamepadScrollingEnabled = bInIsGamepadScrollingEnabled;
+	if (MyTableViewBase)
+	{
+		MyTableViewBase->SetIsGamepadScrollingEnabled(bInIsGamepadScrollingEnabled);
+	}
+}
+
 const TArray<UUserWidget*>& UListViewBase::GetDisplayedEntryWidgets() const
 { 
 	return EntryWidgetPool.GetActiveWidgets(); 
diff --git a/Engine/Source/Runtime/UMG/Public/Components/ListViewBase.h b/Engine/Source/Runtime/UMG/Public/Components/ListViewBase.h
--- a/Engine/Source/Runtime/UMG/Public/Components/ListViewBase.h
+++ b/Engine/Source/Runtime/UMG/Public/Components/ListViewBase.h
@@ -546,6 +546,10 @@ public:
 	UFUNCTION(BlueprintCallable, Category = ListViewBase)
 	UMG_API void SetIsPointerScrollingEnabled(bool bInIsPointerScrollingEnabled);
 
+	/** Enable/Disable the ability of the list to scroll via gamepad. */
+	UFUNCTION(BlueprintCallable, Category = ListViewBase)
+	UMG_API void SetIsGamepadScrollingEnabled(bool bInIsGamepadScrollingEnabled);
+
 	/**
 	 * Sets the list to refresh on the next tick.
 	 *
@@ -698,6 +702,10 @@ protected:
 	UPROPERTY(EditDefaultsOnly, Category = Scrolling)
 	bool bIsPointerScrollingEnabled = true;
 
+	/** Enable/Disable scrolling using Gamepad. */
+	UPROPERTY(EditDefaultsOnly, Category = Scrolling)
+	bool bIsGamepadScrollingEnabled = true;
+
 	UPROPERTY(EditAnywhere, Category = Scrolling)
 	bool bEnableFixedLineOffset = false;
 
