Slate: Fix overflow on rich text block.
- Rich text box has property to set OverflowPolicy but it was never used

diff --git a/Engine/Source/Runtime/Slate/Private/Widgets/Text/SRichTextBlock.cpp b/Engine/Source/Runtime/Slate/Private/Widgets/Text/SRichTextBlock.cpp
--- a/Engine/Source/Runtime/Slate/Private/Widgets/Text/SRichTextBlock.cpp
+++ b/Engine/Source/Runtime/Slate/Private/Widgets/Text/SRichTextBlock.cpp
@@ -54,6 +54,7 @@ void SRichTextBlock::Construct( const FArguments& InArgs )
 
 		TextLayoutCache = MakeUnique<FSlateTextBlockLayout>(this, TextStyle, InArgs._TextShapingMethod, InArgs._TextFlowDirection, InArgs._CreateSlateTextLayout, Marshaller.ToSharedRef(), nullptr);
 		TextLayoutCache->SetDebugSourceInfo(TAttribute<FString>::Create(TAttribute<FString>::FGetter::CreateLambda([this]{ return FReflectionMetaData::GetWidgetDebugInfo(this); })));
+		TextLayoutCache->SetTextOverflowPolicy(InArgs._OverflowPolicy.IsSet() ? InArgs._OverflowPolicy : TextStyle.OverflowPolicy);
 	}
 
 	SetCanTick(false);
