Slate: Remove visibility attribute in TableViewBase.

diff --git a/Engine/Source/Runtime/Slate/Private/Widgets/Views/STableViewBase.cpp b/Engine/Source/Runtime/Slate/Private/Widgets/Views/STableViewBase.cpp
--- a/Engine/Source/Runtime/Slate/Private/Widgets/Views/STableViewBase.cpp
+++ b/Engine/Source/Runtime/Slate/Private/Widgets/Views/STableViewBase.cpp
@@ -78,7 +78,7 @@ void STableViewBase::ConstructChildren( const TAttribute<float>& InItemWidth, co
 		.NumDesiredItems(this, &STableViewBase::GetNumPinnedItems)
 		.ItemAlignment(InItemAlignment)
 		.ListOrientation(Orientation)
-		.Visibility(this, &STableViewBase::GetPinnedItemsVisiblity);
+		.Visibility(EVisibility::Collapsed);
 
 	TSharedPtr<SWidget> ListAndScrollbar;
 	if (InScrollBar)
@@ -978,6 +978,7 @@ void STableViewBase::InsertPinnedWidget( const TSharedRef<SWidget> & WidgetToIns
 	[
 		WidgetToInset
 	];
+	PinnedItemsPanel->SetVisibility(EVisibility::Visible);
 }
 
 void STableViewBase::AppendPinnedWidget( const TSharedRef<SWidget>& WidgetToAppend )
@@ -986,10 +987,12 @@ void STableViewBase::AppendPinnedWidget( const TSharedRef<SWidget>& WidgetToAppe
 	[
 		WidgetToAppend
 	];
+	PinnedItemsPanel->SetVisibility(EVisibility::Visible);
 }
 
 void STableViewBase::ClearPinnedWidgets()
 {
+	PinnedItemsPanel->SetVisibility(EVisibility::Collapsed);
 	PinnedItemsPanel->ClearItems();
 }
 
diff --git a/Engine/Source/Runtime/Slate/Public/Widgets/Views/STableViewBase.h b/Engine/Source/Runtime/Slate/Public/Widgets/Views/STableViewBase.h
--- a/Engine/Source/Runtime/Slate/Public/Widgets/Views/STableViewBase.h
+++ b/Engine/Source/Runtime/Slate/Public/Widgets/Views/STableViewBase.h
@@ -363,6 +363,7 @@ protected:
 	/** @return how many pinned items are in the table */
 	SLATE_API int32 GetNumPinnedItems() const;
 
+	UE_DEPRECATED(5.5, "The Pinned items visibility is no more an attribute.")
 	SLATE_API EVisibility GetPinnedItemsVisiblity() const;
 
 	enum class EScrollIntoViewResult
