Remove outdated comment from FMenuBuilder

diff --git a/Engine/Source/Runtime/Slate/Private/Framework/MultiBox/MultiBoxBuilder.cpp b/Engine/Source/Runtime/Slate/Private/Framework/MultiBox/MultiBoxBuilder.cpp
--- a/Engine/Source/Runtime/Slate/Private/Framework/MultiBox/MultiBoxBuilder.cpp
+++ b/Engine/Source/Runtime/Slate/Private/Framework/MultiBox/MultiBoxBuilder.cpp
@@ -205,7 +205,6 @@ TSharedRef< class SWidget > FMenuBuilder::MakeWidget( FMultiBox::FOnMakeMultiBox
 
 TSharedRef< class SWidget > FMenuBuilder::MakeWidget( FMultiBox::FOnMakeMultiBoxBuilderOverride* InMakeMultiBoxBuilderOverride, uint32 MaxHeight)
 {
-	// Make menu builders searchable (by default)
 	TAttribute<float> MaxHeightAttribute;
 	if (MaxHeight < INT_MAX)
 	{
