Engine
- Stop the inertial scrolling the same frame that it's been requested.

#rnx

diff --git a/Engine/Source/Runtime/Slate/Private/Framework/Layout/InertialScrollManager.cpp b/Engine/Source/Runtime/Slate/Private/Framework/Layout/InertialScrollManager.cpp
--- a/Engine/Source/Runtime/Slate/Private/Framework/Layout/InertialScrollManager.cpp
+++ b/Engine/Source/Runtime/Slate/Private/Framework/Layout/InertialScrollManager.cpp
@@ -63,7 +63,8 @@ void FInertialScrollManager::UpdateScrollVelocity(const float InDeltaTime)
 	}
 }
 
-void FInertialScrollManager::ClearScrollVelocity()
+void FInertialScrollManager::ClearScrollVelocity(bool bInShouldStopScrollNow)
 {
+	bShouldStopScrollNow = ScrollVelocity != 0 ? bInShouldStopScrollNow : false;
 	ScrollVelocity = 0;
 }
diff --git a/Engine/Source/Runtime/Slate/Private/Widgets/Views/STableViewBase.cpp b/Engine/Source/Runtime/Slate/Private/Widgets/Views/STableViewBase.cpp
--- a/Engine/Source/Runtime/Slate/Private/Widgets/Views/STableViewBase.cpp
+++ b/Engine/Source/Runtime/Slate/Private/Widgets/Views/STableViewBase.cpp
@@ -290,6 +290,13 @@ void STableViewBase::Tick( const FGeometry& AllottedGeometry, const double InCur
 			const EScrollIntoViewResult ScrollIntoViewResult = ScrollIntoView(PanelGeometry);
 
 			double TargetScrollOffset = GetTargetScrollOffset();
+			
+			if (InertialScrollManager.GetShouldStopScrollNow())
+			{
+				TargetScrollOffset = DesiredScrollOffset = CurrentScrollOffset;
+				InertialScrollManager.ResetShouldStopScrollNow();
+			}
+
 			if((bStartedTouchInteraction && bEnableTouchAnimatedScrolling) || (!bStartedTouchInteraction && bEnableAnimatedScrolling))
 			{
 				CurrentScrollOffset = FMath::FInterpTo(CurrentScrollOffset, TargetScrollOffset, (double)InDeltaTime, 12.0);
@@ -398,8 +405,8 @@ FReply STableViewBase::OnPreviewMouseButtonDown( const FGeometry& MyGeometry, co
 {
 	if (bEnableTouchScrolling && MouseEvent.IsTouchEvent())
 	{
-		// Clear any inertia 
-		this->InertialScrollManager.ClearScrollVelocity();
+		// Clear any inertia
+		InertialScrollManager.ClearScrollVelocity(true);
 		// We have started a new interaction; track how far the user has moved since they put their finger down.
 		AmountScrolledWhileRightMouseDown = 0;
 
@@ -419,7 +426,7 @@ FReply STableViewBase::OnPreviewMouseButtonDown( const FGeometry& MyGeometry, co
 FReply STableViewBase::OnMouseButtonDown( const FGeometry& MyGeometry, const FPointerEvent& MouseEvent )
 {
 	// Zero the scroll velocity so the list stops immediately on mouse down, even if the user does not drag
-	this->InertialScrollManager.ClearScrollVelocity();
+	InertialScrollManager.ClearScrollVelocity(true);
 
 	if (MouseEvent.GetEffectingButton() == EKeys::RightMouseButton)
 	{
diff --git a/Engine/Source/Runtime/Slate/Public/Framework/Layout/InertialScrollManager.h b/Engine/Source/Runtime/Slate/Public/Framework/Layout/InertialScrollManager.h
--- a/Engine/Source/Runtime/Slate/Public/Framework/Layout/InertialScrollManager.h
+++ b/Engine/Source/Runtime/Slate/Public/Framework/Layout/InertialScrollManager.h
@@ -37,12 +37,21 @@ public:
 	/** Updates the current scroll velocity. Call every frame. */
 	SLATE_API void UpdateScrollVelocity(const float InDeltaTime);
 
-	/** Instantly end the inertial scroll */
-	SLATE_API void ClearScrollVelocity();
+	/** 
+	 * Stop the accumulation of inertial scroll. 
+	 * @param bInShouldStopScrollNow true implies the scroll will stop instantly, else the list will scroll until any accumulated scroll offset is cleared.
+	 */
+	SLATE_API void ClearScrollVelocity(bool bInShouldStopScrollNow = false);
 
 	/** Gets the calculated velocity of the scroll. */
 	float GetScrollVelocity() const { return ScrollVelocity; }
 
+	/** Gets the the value of bShouldStopScrollNow. */
+	bool GetShouldStopScrollNow() const { return bShouldStopScrollNow; }
+
+	/** Set the value of bShouldStopScrollNow to false. */
+	void ResetShouldStopScrollNow() { bShouldStopScrollNow = false; }
+
 private:
 	struct FScrollSample
 	{
@@ -61,6 +70,9 @@ private:
 	/** The current velocity of the scroll */
 	float ScrollVelocity;
 
+	/** When true, the list will stop scrolling */
+	bool bShouldStopScrollNow = false;
+
 	/** Samples older than this amount of time will be discarded. */
 	double SampleTimeout;
 };
