Fix issue where menu search fields were visible when they shouldn't

The issue happened when an FMenuBuilder built a menu widget of N items while the "Menu Search Field Visibility Threshold" setting was at some value lower than N - in this case the menu is created with a visible search field. If the user later changes the setting to a value higher than N, and the widget is still in the UI and not rebuilt, then the search field will still be visible but disabled.

This is fixed to properly hide the search field when the number of menu items is lower than the "Menu Search Field Visibility Threshold" setting by making the Visibility attribute of the SBox containing the search field read the setting. If the setting indicates the search field should be hidden, the SBox is collapsed and the search field disappears.

diff --git a/Engine/Source/Runtime/Slate/Private/Framework/MultiBox/MultiBox.cpp b/Engine/Source/Runtime/Slate/Private/Framework/MultiBox/MultiBox.cpp
--- a/Engine/Source/Runtime/Slate/Private/Framework/MultiBox/MultiBox.cpp
+++ b/Engine/Source/Runtime/Slate/Private/Framework/MultiBox/MultiBox.cpp
@@ -889,6 +889,9 @@ void SMultiBoxWidget::CreateSearchTextWidget()
 	TSharedRef<SBox> SearchBox =
 		SNew(SBox)
 			.Padding(FMargin(8, 0, 8, 0))
+			.Visibility_Lambda([this]() -> EVisibility {
+				return ShouldShowMenuSearchField() ? EVisibility::Visible : EVisibility::Collapsed;
+			})
 			[
 				SearchTextWidget.ToSharedRef()
 			];
@@ -1807,12 +1810,6 @@ void SMultiBoxWidget::FilterMultiBoxEntries()
 			It.Key()->SetVisibility(EVisibility::Visible);
 		}
 
-		// Reset the visibility of the search field.
-		if (SearchBlockWidget.IsValid())
-		{
-			SearchBlockWidget->SetVisibility(ShouldShowMenuSearchField() ? EVisibility::Visible : EVisibility::Collapsed);
-		}
-
 		// Hide the sub-menus widgets that were made visible by searching this multi-box hierarchy.
 		for (TPair<TSharedPtr<const FMultiBlock>, TSharedPtr<FFlattenSearchableBlockInfo>>& Pair: FlattenSearchableBlocks)
 		{
@@ -1898,7 +1895,6 @@ TSharedPtr<SWidget> SMultiBoxWidget::GetSearchTextWidget()
 void SMultiBoxWidget::SetSearchBlockWidget(TSharedPtr<SWidget> InWidget)
 {
 	SearchBlockWidget = InWidget;
-	InWidget->SetVisibility(ShouldShowMenuSearchField() ? EVisibility::Visible : EVisibility::Collapsed);
 }
 
 void SMultiBoxWidget::AddSearchElement( TSharedPtr<SWidget> BlockWidget, FText BlockDisplayText )
