[PropertyEditor] Displayed units are no longer overridden when ForceUnits is specified

diff --git a/Engine/Source/Editor/PropertyEditor/Private/UserInterface/PropertyEditor/SPropertyEditorNumeric.h b/Engine/Source/Editor/PropertyEditor/Private/UserInterface/PropertyEditor/SPropertyEditorNumeric.h
--- a/Engine/Source/Editor/PropertyEditor/Private/UserInterface/PropertyEditor/SPropertyEditorNumeric.h
+++ b/Engine/Source/Editor/PropertyEditor/Private/UserInterface/PropertyEditor/SPropertyEditorNumeric.h
@@ -287,7 +287,7 @@ public:
 			{
 				// Create the type interface and set up the default input units if they are compatible
 				TypeInterface = MakeShareable(new TNumericUnitTypeInterface<NumericType>(PropertyUnits.GetValue()));
-				TypeInterface->FixedDisplayUnits = PropertyUnits.GetValue();
+				TypeInterface->UserDisplayUnits = TypeInterface->FixedDisplayUnits = PropertyUnits.GetValue();
 			}
 			// If that's not set, we fall back to Units=xxx which calculates the most appropriate unit to display in
 			else
@@ -500,7 +500,7 @@ private:
 			{
 				// Value hasn't changed, so lets return now
 				if (OrgValue == NewValue)
-				{ 
+				{
 					return;
 				}
 			}
diff --git a/Engine/Source/Runtime/Slate/Public/Widgets/Input/NumericTypeInterface.h b/Engine/Source/Runtime/Slate/Public/Widgets/Input/NumericTypeInterface.h
--- a/Engine/Source/Runtime/Slate/Public/Widgets/Input/NumericTypeInterface.h
+++ b/Engine/Source/Runtime/Slate/Public/Widgets/Input/NumericTypeInterface.h
@@ -157,9 +157,12 @@ struct TNumericUnitTypeInterface : TDefaultNumericTypeInterface<NumericType>
 	/** The underlying units which the numeric type are specified in. */
 	const EUnit UnderlyingUnits;
 
-	/** Optional units that this type interface will be fixed on */
+	/** Optional units that this type interface will be fixed on. These are usually auto-calculated by SetupFixedDisplay. */
 	TOptional<EUnit> FixedDisplayUnits;
 
+	/** Optional user-specified units that this type interface will be displayed in. If set, FixedDisplayUnits will be ignored. */
+	TOptional<EUnit> UserDisplayUnits;
+
 	/** Constructor */
 	TNumericUnitTypeInterface(EUnit InUnits);
 
@@ -172,6 +175,7 @@ struct TNumericUnitTypeInterface : TDefaultNumericTypeInterface<NumericType>
 	/** Check whether the specified typed character is valid */
 	virtual bool IsCharacterValid(TCHAR InChar) const override;
 
-	/** Set up this interface to use a fixed display unit based on the specified value */
+	/** Set up this interface to use a fixed display unit, calculated based on the specified value.
+	 * For example, if underlying units are cm/s, but the provided value is over 100, it will use m/s instead. */
 	void SetupFixedDisplay(const NumericType& InValue);
 };
diff --git a/Engine/Source/Runtime/Slate/Public/Widgets/Input/NumericUnitTypeInterface.inl b/Engine/Source/Runtime/Slate/Public/Widgets/Input/NumericUnitTypeInterface.inl
--- a/Engine/Source/Runtime/Slate/Public/Widgets/Input/NumericUnitTypeInterface.inl
+++ b/Engine/Source/Runtime/Slate/Public/Widgets/Input/NumericUnitTypeInterface.inl
@@ -34,6 +34,15 @@ FString TNumericUnitTypeInterface<NumericType>::ToString(const NumericType& Valu
 
 	FNumericUnit<NumericType> FinalValue(Value, UnderlyingUnits);
 
+	if (UserDisplayUnits.IsSet())
+	{
+		auto Converted = FinalValue.ConvertTo(UserDisplayUnits.GetValue());
+		if (Converted.IsSet())
+		{
+			return ToUnitString(Converted.GetValue());
+		}
+	}
+
 	if (FixedDisplayUnits.IsSet())
 	{
 		auto Converted = FinalValue.ConvertTo(FixedDisplayUnits.GetValue());
@@ -54,7 +63,12 @@ TOptional<NumericType> TNumericUnitTypeInterface<NumericType>::FromString(const
 		return TDefaultNumericTypeInterface<NumericType>::FromString(InString, InExistingValue);
 	}
 
-	EUnit DefaultUnits = FixedDisplayUnits.IsSet() ? FixedDisplayUnits.GetValue() : UnderlyingUnits;
+	// Get displayed units from the user, if set, otherwise the fixed units, or finally the underlying units
+	EUnit DefaultUnits = UserDisplayUnits.IsSet()
+		? UserDisplayUnits.GetValue()
+		: FixedDisplayUnits.IsSet()
+			? FixedDisplayUnits.GetValue()
+			: UnderlyingUnits;
 
 	// Always parse in as a double, to allow for input of higher-order units with decimal numerals into integral types (eg, inputting 0.5km as 500m)
 	TValueOrError<FNumericUnit<double>, FText> NewValue = FNumericUnit<double>::TryParseExpression( *InString, DefaultUnits, InExistingValue );
@@ -62,10 +76,18 @@ TOptional<NumericType> TNumericUnitTypeInterface<NumericType>::FromString(const
 	{
 		// Convert the number into the correct units
 		EUnit SourceUnits = NewValue.GetValue().Units;
-		if (SourceUnits == EUnit::Unspecified && FixedDisplayUnits.IsSet())
+		if (SourceUnits == EUnit::Unspecified)
 		{
+			// Use the user overridden input units if set
+			if (UserDisplayUnits.IsSet()) 
+			{
+				SourceUnits = UserDisplayUnits.GetValue();
+			}
 			// Use the default supplied input units
-			SourceUnits = FixedDisplayUnits.GetValue();
+			else if (FixedDisplayUnits.IsSet())
+			{
+				SourceUnits = FixedDisplayUnits.GetValue();
+			}
 		}
 		double ConvertedValue = FUnitConversion::Convert(NewValue.GetValue().Value, SourceUnits, UnderlyingUnits);
 		return FMath::Clamp((NumericType)ConvertedValue, TNumericLimits<NumericType>::Lowest(), TNumericLimits<NumericType>::Max());
@@ -83,6 +105,7 @@ bool TNumericUnitTypeInterface<NumericType>::IsCharacterValid(TCHAR InChar) cons
 template<typename NumericType>
 void TNumericUnitTypeInterface<NumericType>::SetupFixedDisplay(const NumericType& InValue)
 {
+	// We calculate this regardless of whether FixedDisplayUnits is used, so that the moment it is used, it's correct
 	EUnit DisplayUnit = FUnitConversion::CalculateDisplayUnit(InValue, UnderlyingUnits);
 	if (DisplayUnit != EUnit::Unspecified)
 	{
