Hide widget tooltip while LMB is down

#rnx

diff --git a/Engine/Source/Runtime/Slate/Private/Framework/Application/SlateUser.cpp b/Engine/Source/Runtime/Slate/Private/Framework/Application/SlateUser.cpp
--- a/Engine/Source/Runtime/Slate/Private/Framework/Application/SlateUser.cpp
+++ b/Engine/Source/Runtime/Slate/Private/Framework/Application/SlateUser.cpp
@@ -1204,6 +1204,7 @@ void FSlateUser::UpdateTooltip(const FMenuStack& MenuStack, bool bCanSpawnNewToo
 		IsInGameThread() &&					// We should never allow the slate loading thread to create new windows or interact with the hittest grid
 		!SlateApp.IsUsingHighPrecisionMouseMovment() && // If we are using HighPrecision movement then we can't rely on the OS cursor to be accurate
 		!IsDragDropping() &&				// We must not currently be in the middle of a drag-drop action
+		!SlateApp.GetPressedMouseButtons().Contains(EKeys::LeftMouseButton) && // We must not currently be clicking on a widget
 		(
 			SlateApp.IsActive() || // Assume we need update if app is active
 			//@todo DanH: We need to check if OUR cursor is over a slate window, not just the platform cursor. 
