Add favorites drop target for place actors

For toolbars, we need the ability to provide buttons with dropzone (or othere) capabilities, and we need to be able to track the selected index for the toolbar in the overflow button. I added the selected index for the toolbar, and for buttons on the toolbar I added a delegate to toolbar buttons which can decorate the button as needed (in this case adding a droptarget). In this case, using the asset droptarget is the first step that will be covered in this review, the intention is to ultimately support more than just assets for this droptarget.

The reason that we need to track the last selected index is because the overflow button in the SClippingVerticalBox needs to apply a certain style based upon whether something in the overflow has been selected. To enable that there are changes in SClippingVerticalBox to track and use the last selected index to provide the appropriate styling.

diff --git a/Engine/Source/Developer/WidgetRegistration/Private/Inputs/BuilderInputManager.cpp b/Engine/Source/Developer/WidgetRegistration/Private/Inputs/BuilderInputManager.cpp
--- a/Engine/Source/Developer/WidgetRegistration/Private/Inputs/BuilderInputManager.cpp
+++ b/Engine/Source/Developer/WidgetRegistration/Private/Inputs/BuilderInputManager.cpp
@@ -3,7 +3,6 @@
 #include "Inputs/BuilderInputManager.h"
 #include "BuilderCommandCreationManager.h"
 #include "Framework/Commands/UICommandInfo.h"
-// #include "Inputs/DragAndDrop/BuilderDragAndDropManager.h"
 
 FBuilderInputManager& FBuilderInputManager::Get()
 {
@@ -26,9 +25,3 @@ const FBuilderCommandCreationManager& FBuilderInputManager::GetCommandManager()
 	static const FBuilderCommandCreationManager& CommandManager = FBuilderCommandCreationManager::Get();
 	return CommandManager;
 }
-
-/*FBuilderDragAndDropManager& FBuilderInputManager::GetDragAndDropManager()
-{
-	static FBuilderDragAndDropManager& DragAndDropManager = FBuilderDragAndDropManager::Get();
-	return DragAndDropManager;
-}*/
\ No newline at end of file
diff --git a/Engine/Source/Developer/WidgetRegistration/Private/Inputs/BuilderInputManager.h b/Engine/Source/Developer/WidgetRegistration/Private/Inputs/BuilderInputManager.h
--- a/Engine/Source/Developer/WidgetRegistration/Private/Inputs/BuilderInputManager.h
+++ b/Engine/Source/Developer/WidgetRegistration/Private/Inputs/BuilderInputManager.h
@@ -4,7 +4,6 @@
 
 #include "Framework/Commands/UICommandInfo.h"
 
-// class FBuilderDragAndDropManager;
 class FBuilderCommandCreationManager;
 
 namespace UE::DisplayBuilders
@@ -38,9 +37,4 @@ public:
 	 * @return the Command manager for Builders
 	 */
 	const FBuilderCommandCreationManager& GetCommandManager();
-	
-	/**
-	 * @return the Drag and Drop manager for Builders
-	 */
-	// FBuilderDragAndDropManager& GetDragAndDropManager();
 };
diff --git a/Engine/Source/Developer/WidgetRegistration/Private/Layout/CategoryDrivenContentBuilder.cpp b/Engine/Source/Developer/WidgetRegistration/Private/Layout/CategoryDrivenContentBuilder.cpp
--- a/Engine/Source/Developer/WidgetRegistration/Private/Layout/CategoryDrivenContentBuilder.cpp
+++ b/Engine/Source/Developer/WidgetRegistration/Private/Layout/CategoryDrivenContentBuilder.cpp
@@ -10,7 +10,6 @@
 #include "Framework/MultiBox/SToolBarButtonBlock.h"
 #include "Inputs/BuilderInputManager.h"
 #include "Input/DragAndDrop.h"
-// #include "Inputs/DragAndDrop/AcceptSingleItemDropZoneHandler.h"
 #include "Layout/Containers/ColumnWrappingContainer.h"
 #include "Layout/Containers/SimpleTitleContainer.h"
 #include "Persistence/BuilderPersistenceManager.h"
@@ -50,7 +49,6 @@ FCategoryDrivenContentBuilder::FCategoryDrivenContentBuilder(FCategoryDrivenCont
 	, TitleContainer( nullptr )
 	, bIsFilledWithWidget( false )
 	, bShowNoCategorySelection( false )
-	/*, FavoritesDropZoneHandler( Args.FavoritesDropZoneHandler )*/
 {
 	Favorites = UBuilderPersistenceManager::Get()->GetPersistedFavoritesNamesArray( BuilderKey );
 	ActiveCategoryName = Args.ActiveCategoryName;
@@ -165,7 +163,7 @@ void FCategoryDrivenContentBuilder::UpdateWidget()
 		
 			if ( LoadPaletteToolBarBuilder && !bShowNoCategorySelection )
 			{
-				//	LoadPaletteToolBarBuilder->SetLastSelectedCommandIndex( (*Input).Index );
+				LoadPaletteToolBarBuilder->SetLastSelectedCommandIndex( (*Input).Index );
 			}
 		}
 
@@ -233,14 +231,14 @@ void FCategoryDrivenContentBuilder::SetShowNoCategorySelection(bool bInShowNoCat
 		bShowNoCategorySelection = bInShowNoCategorySelection;
 		const UE::DisplayBuilders::FBuilderInput* Input = CategoryNameToBuilderInputMap.Find( ActiveCategoryName );
 
-		/*if ( bShowNoCategorySelection )
+		if ( bShowNoCategorySelection )
 		{
 			LoadPaletteToolBarBuilder->SetLastSelectedCommandIndex( INDEX_NONE );
 		}
 		else if ( Input != nullptr)
 		{
 			LoadPaletteToolBarBuilder->SetLastSelectedCommandIndex( (*Input).Index );
-		}*/
+		}
 	}
 }
 
@@ -258,14 +256,9 @@ void FCategoryDrivenContentBuilder::InitializeCategoryButtons(TArray<UE::Display
 		UE::DisplayBuilders::FBuilderInput& BuilderInput = InBuilderInputArray[Index];
 		BuilderInput.Index = Index;
 		
-		if ( BuilderInput.Name == FavoritesCategoryName /*&& FavoritesDropZoneHandler.IsValid()*/ )
+		if ( BuilderInput.Name == FavoritesCategoryName && GetDecoratedButtonDelegate.IsBound() )
 		{
-			/*if ( FavoritesDropZoneHandler.IsValid() )
-			{
-				 FavoritesDropZoneHandler->OnAcceptDrop =
-					FAcceptSingleItemDropZoneHandler::FOnAcceptDrop::CreateSP( this, &FCategoryDrivenContentBuilder::AddFavorite );
-				 BuilderInput.ButtonArgs.DropZoneHandler = FavoritesDropZoneHandler; 
-			}*/
+			BuilderInput.ButtonArgs.GetDecoratedButtonDelegate = GetDecoratedButtonDelegate;
 			BuilderInput.ButtonArgs.IconOverride = FSlateIcon(FAppStyle::Get().GetStyleSetName(), "Icons.Favorites.Small");
 		}
 		CategoryNameToBuilderInputMap.Add( BuilderInput.Name, BuilderInput );
diff --git a/Engine/Source/Developer/WidgetRegistration/Private/Layout/CategoryDrivenContentBuilderBase.cpp b/Engine/Source/Developer/WidgetRegistration/Private/Layout/CategoryDrivenContentBuilderBase.cpp
--- a/Engine/Source/Developer/WidgetRegistration/Private/Layout/CategoryDrivenContentBuilderBase.cpp
+++ b/Engine/Source/Developer/WidgetRegistration/Private/Layout/CategoryDrivenContentBuilderBase.cpp
@@ -27,6 +27,7 @@ FCategoryDrivenContentBuilderBase::FCategoryDrivenContentBuilderBase ( FName InB
 
 FCategoryDrivenContentBuilderBase::FCategoryDrivenContentBuilderBase( FCategoryDrivenContentBuilderArgs& Args ) :
 	FToolElementRegistrationArgs( Args.Key )
+	, GetDecoratedButtonDelegate( Args.GetDecoratedButtonDelegate )
 	, CategoryReclickBehavior( Args.CategoryReclickBehavior )
 	, BuilderName(Args.BuilderName)
 	, ActiveCategoryName( Args.ActiveCategoryName )
diff --git a/Engine/Source/Developer/WidgetRegistration/Public/Inputs/BuilderInput.cpp b/Engine/Source/Developer/WidgetRegistration/Public/Inputs/BuilderInput.cpp
--- a/Engine/Source/Developer/WidgetRegistration/Public/Inputs/BuilderInput.cpp
+++ b/Engine/Source/Developer/WidgetRegistration/Public/Inputs/BuilderInput.cpp
@@ -38,8 +38,7 @@ UE::DisplayBuilders::FBuilderInput::FBuilderInput(
 	FName InUiStyle ,
 	FName InBindingContext  ,
 	FName InBundle  ,
-	bool bInUseLongDisplayName 
-	/*TSharedPtr<FDropZoneHandler> InDropZoneHandler*/   )  :
+	bool bInUseLongDisplayName )  :
 	FLabelAndIconArgs( InLabel.IsEmpty() ? FText::FromName( InName ) : InLabel, InIcon )
 	, Name( InName )
 	, UserInterfaceType( InUserInterfaceType )
@@ -50,7 +49,6 @@ UE::DisplayBuilders::FBuilderInput::FBuilderInput(
 	, BindingContext( InBindingContext )
 	, Bundle( InBundle )
 	, bUseLongDisplayName( bInUseLongDisplayName )
-	//, DropZoneHandler( InDropZoneHandler )
 	, Index( INDEX_NONE )
 	, Tooltip( InToolTip.IsEmpty() ? InLabel : InToolTip )
 {
diff --git a/Engine/Source/Developer/WidgetRegistration/Public/Inputs/BuilderInput.h b/Engine/Source/Developer/WidgetRegistration/Public/Inputs/BuilderInput.h
--- a/Engine/Source/Developer/WidgetRegistration/Public/Inputs/BuilderInput.h
+++ b/Engine/Source/Developer/WidgetRegistration/Public/Inputs/BuilderInput.h
@@ -34,8 +34,7 @@ namespace UE::DisplayBuilders
 			FName InUiStyle = NAME_None,
 			FName InBindingContext = NAME_None,
 			FName InBundle = NAME_None,
-			bool bInUseLongDisplayName = false
-			/*TSharedPtr<FDropZoneHandler> InDropZoneHandler = nullptr*/ );
+			bool bInUseLongDisplayName = false );
 
 		/**
 		 * The destructor, destroys any command related information
@@ -69,9 +68,6 @@ namespace UE::DisplayBuilders
 		/** True if should we use long names for when getting text for input chords */
 		bool bUseLongDisplayName;
 
-		/** A pointer to the handler for a dropzone on the input, if one exists */
-		// TSharedPtr<FDropZoneHandler> DropZoneHandler;
-
 		/** the index of the BuilderInput in whatever container it is in */
 		int32 Index;
 		
diff --git a/Engine/Source/Developer/WidgetRegistration/Public/Layout/CategoryDrivenContentBuilder.h b/Engine/Source/Developer/WidgetRegistration/Public/Layout/CategoryDrivenContentBuilder.h
--- a/Engine/Source/Developer/WidgetRegistration/Public/Layout/CategoryDrivenContentBuilder.h
+++ b/Engine/Source/Developer/WidgetRegistration/Public/Layout/CategoryDrivenContentBuilder.h
@@ -9,7 +9,6 @@
 
 class FColumnWrappingContainer;
 class FSimpleTitleContainer;
-// class FAcceptSingleItemDropZoneHandler;
 class FSlateBuilder;
 
 namespace UE::DisplayBuilders
@@ -93,7 +92,7 @@ public:
 	 * @param InBuilderInputArray the array of FBuilderInput instances that will initialize the Category buttons
 	 */
 	WIDGETREGISTRATION_API void InitializeCategoryButtons(TArray<UE::DisplayBuilders::FBuilderInput> InBuilderInputArray);
-
+	
 	/**
 	 * Converts the SWidget Widget to a FSlateBuilder and adds it to the main content for the currently selected category
 	 *
@@ -113,6 +112,12 @@ public:
 	 */
 	WIDGETREGISTRATION_API void ClearCategoryContent();
 	
+	/**
+	 * Adds the favorites with the name InFavoriteCommandName to the favorites list. If InFavoriteCommandName is already in
+	 * the Favorites list, this is a no-op
+	 */
+	WIDGETREGISTRATION_API void AddFavorite( FName InFavoriteCommandName );
+	
 private:
 	/**
      * Given the active category name, return an SWidget to provide the content
@@ -142,12 +147,6 @@ private:
 	 * Toggles the favorite with the name InFavoriteCommandName
 	 */
 	void ToggleFavorite( FName InFavoriteCommandName );
-	
-	/**
-	 * Adds the favorites with the name InFavoriteCommandName to the favorites list. If InFavoriteCommandName is already in
-	 * the Favorites list, this is a no-op
-	 */
-	void AddFavorite( FName InFavoriteCommandName );
 
 private:
 
@@ -169,7 +168,4 @@ private:
 
 	/** if true, we need to show no selection on the category picker and no category title, else we should show both  */
 	bool bShowNoCategorySelection;
-	
-	/** The dropzone handler for the favorites */
-	// TSharedPtr<FAcceptSingleItemDropZoneHandler> FavoritesDropZoneHandler;
 };
diff --git a/Engine/Source/Developer/WidgetRegistration/Public/Layout/CategoryDrivenContentBuilderBase.h b/Engine/Source/Developer/WidgetRegistration/Public/Layout/CategoryDrivenContentBuilderBase.h
--- a/Engine/Source/Developer/WidgetRegistration/Public/Layout/CategoryDrivenContentBuilderBase.h
+++ b/Engine/Source/Developer/WidgetRegistration/Public/Layout/CategoryDrivenContentBuilderBase.h
@@ -9,12 +9,12 @@
 #include "Widgets/Layout/SSplitter.h"
 #include "Framework/MultiBox/MultiBoxBuilder.h"
 #include "BuilderKey.h"
+#include "Framework/MultiBox/SToolBarButtonBlock.h"
 
 class SWidget;
 class FToolElement;
 
 class FCategoryDrivenContentBuilderArgs;
-// class FAcceptSingleItemDropZoneHandler;
 
 /**
  * The FToolElementRegistrationArgs which is specified for Toolkits
@@ -83,9 +83,9 @@ public:
 	WIDGETREGISTRATION_API void SetCategoryButtonLabelVisibility(bool bIsCategoryButtonLabelVisible);
 
 	/**
-	 * refreshes the UI display of the category toolbar 
+	 * refreshes the UI display of the category toolbar
 	 */
-	WIDGETREGISTRATION_API void RefreshCategoryToolbarWidget();
+    	WIDGETREGISTRATION_API void RefreshCategoryToolbarWidget();
 
 	/** Implements the generation of the TSharedPtr<SWidget> */
 	WIDGETREGISTRATION_API virtual TSharedPtr<SWidget> GenerateWidget() override;
@@ -101,17 +101,20 @@ public:
 
 	/**   OnActivePaletteChanged is broadcast when the active palette changes to a different palette  */
 	FSimpleMulticastDelegate OnActivePaletteChanged;
+
+	/**  Delegate that takes the default toolbar button created by this as a parameter, adds any decorator needed and returns the new widget*/
+	FGetDecoratedButtonDelegate GetDecoratedButtonDelegate;
 	
 private:
 
 	/* The SWidget that is the whole Toolkit */
 	TSharedPtr<SWidget> ToolkitWidget;
 
-	/**
-	 * Creates the SWidget
-	 */
-	void CreateWidget();
-
+    	/**
+    	 * Creates the SWidget
+    	 */
+    	void CreateWidget();
+     
 	/** The SVerticalBox which contains the category toolbar */
 	TSharedPtr<SVerticalBox> CategoryToolbarVBox;
 
@@ -195,6 +198,6 @@ public:
 	/** The name of the initially selected category */
 	FName ActiveCategoryName;
 	
-	
-	// TSharedPtr<FAcceptSingleItemDropZoneHandler> FavoritesDropZoneHandler;
+	/**  Delegate that takes the default toolbar button created by this as a parameter, adds any decorator needed and returns the new widget */
+	FGetDecoratedButtonDelegate GetDecoratedButtonDelegate;
 };
\ No newline at end of file
diff --git a/Engine/Source/Runtime/Slate/Private/Framework/MultiBox/MultiBox.cpp b/Engine/Source/Runtime/Slate/Private/Framework/MultiBox/MultiBox.cpp
--- a/Engine/Source/Runtime/Slate/Private/Framework/MultiBox/MultiBox.cpp
+++ b/Engine/Source/Runtime/Slate/Private/Framework/MultiBox/MultiBox.cpp
@@ -31,6 +31,7 @@
 #include "Framework/Commands/UICommandDragDropOp.h"
 #include "Framework/MultiBox/SUniformToolbarPanel.h"
 #include "Styling/ToolBarStyle.h"
+#include "Misc/CoreMiscDefines.h"
 
 #define LOCTEXT_NAMESPACE "MultiBox"
 
@@ -252,6 +253,7 @@ FMultiBox::FMultiBox(const EMultiBoxType InType, FMultiBoxCustomization InCustom
 	, StyleName( "ToolBar" )
 	, Type( InType )
 	, bShouldCloseWindowAfterMenuSelection( bInShouldCloseWindowAfterMenuSelection )
+	, LastSelectedCommandIndex( INDEX_NONE )
 {
 
 	if ((InType == EMultiBoxType::SlimHorizontalToolBar ||  InType == EMultiBoxType::SlimHorizontalUniformToolBar) && FCoreStyle::IsStarshipStyle())
@@ -462,6 +464,16 @@ bool FMultiBox::IsCustomizable() const
 	return false;
 }
 
+int32 FMultiBox::GetLastSelectedCommandIndex() const
+{
+	return LastSelectedCommandIndex;
+}
+
+void FMultiBox::SetLastSelectedCommandIndex(const int32 InLastSelectedCommandIndex)
+{
+	LastSelectedCommandIndex = InLastSelectedCommandIndex;
+}
+
 FName FMultiBox::GetCustomizationName() const
 {
 	return NAME_None;
@@ -1007,6 +1019,10 @@ void SMultiBoxWidget::BuildMultiBoxWidget()
 			MainWidget = VerticalBox = ClippedVerticalBox = SNew(SClippingVerticalBox)
 				.OnWrapButtonClicked(FOnGetContent::CreateSP(this, &SMultiBoxWidget::OnWrapButtonClicked))
 				.IsFocusable(MultiBox->bIsFocusable)
+				.SelectedIndex_Lambda( [this] ()
+				{
+					return MultiBox->GetLastSelectedCommandIndex();
+				})
 				.StyleSet(StyleSet)
 				.StyleName(StyleName);
 		}
diff --git a/Engine/Source/Runtime/Slate/Private/Framework/MultiBox/MultiBoxBuilder.cpp b/Engine/Source/Runtime/Slate/Private/Framework/MultiBox/MultiBoxBuilder.cpp
--- a/Engine/Source/Runtime/Slate/Private/Framework/MultiBox/MultiBoxBuilder.cpp
+++ b/Engine/Source/Runtime/Slate/Private/Framework/MultiBox/MultiBoxBuilder.cpp
@@ -53,6 +53,11 @@ void FMultiBoxBuilder::PushCommandList( const TSharedRef< const FUICommandList >
 	CommandListStack.Push( CommandList );
 }
 
+void FMultiBoxBuilder::SetLastSelectedCommandIndex(int32 InLastSelectedCommandIndex)
+{
+	MultiBox->SetLastSelectedCommandIndex( InLastSelectedCommandIndex );			
+}
+
 void FMultiBoxBuilder::PopCommandList()
 {
 	// Never allowed to pop the last command-list!  This command-list was set when the multibox was first created and is canonical.
@@ -205,6 +210,7 @@ TSharedRef< class SWidget > FMenuBuilder::MakeWidget( FMultiBox::FOnMakeMultiBox
 
 TSharedRef< class SWidget > FMenuBuilder::MakeWidget( FMultiBox::FOnMakeMultiBoxBuilderOverride* InMakeMultiBoxBuilderOverride, uint32 MaxHeight)
 {
+	// Make menu builders searchable (by default)
 	TAttribute<float> MaxHeightAttribute;
 	if (MaxHeight < INT_MAX)
 	{
@@ -492,6 +498,7 @@ void FToolBarBuilder::AddToolBarButton(const FButtonArgs& Args)
 	NewToolBarButtonBlock->SetStyleNameOverride(CurrentStyleOverride);
 	NewToolBarButtonBlock->SetCustomMenuDelegate(Args.CustomMenuDelegate);
 	NewToolBarButtonBlock->SetOnGetMenuContent(Args.OnGetMenuContent);
+	NewToolBarButtonBlock->SetGetDecoratedButtonDelegate( Args.GetDecoratedButtonDelegate );
 	
 	if (bHasUserInteractionType)
 	{
diff --git a/Engine/Source/Runtime/Slate/Private/Framework/MultiBox/SClippingVerticalBox.cpp b/Engine/Source/Runtime/Slate/Private/Framework/MultiBox/SClippingVerticalBox.cpp
--- a/Engine/Source/Runtime/Slate/Private/Framework/MultiBox/SClippingVerticalBox.cpp
+++ b/Engine/Source/Runtime/Slate/Private/Framework/MultiBox/SClippingVerticalBox.cpp
@@ -1,6 +1,9 @@
 // Copyright Epic Games, Inc. All Rights Reserved.
 
 #include "Framework/MultiBox/SClippingVerticalBox.h"
+
+#include "Brushes/SlateRoundedBoxBrush.h"
+#include "Styling/StyleColors.h"
 #include "Layout/ArrangedChildren.h"
 #include "Rendering/DrawElements.h"
 #include "Widgets/Images/SImage.h"
@@ -9,6 +12,9 @@
 #include "Framework/MultiBox/SToolBarButtonBlock.h"
 #include "Styling/ToolBarStyle.h"
 #include "Widgets/Images/SLayeredImage.h"
+#include "Widgets/Text/STextBlock.h"
+
+#define LOCTEXT_NAMESPACE "ClippingVerticalBox"
 
 void SClippingVerticalBox::OnArrangeChildren( const FGeometry& AllottedGeometry, FArrangedChildren& ArrangedChildren ) const
 {
@@ -24,8 +30,9 @@ void SClippingVerticalBox::OnArrangeChildren( const FGeometry& AllottedGeometry,
 	// Remove children that are clipped by the allotted geometry
 	const int32 NumChildren = ArrangedChildren.Num();
 	const int32 OverflowButtonIndex = NumChildren - 1;
-	const int32 LastToolBarButtonIndex = OverflowButtonIndex - 1;
-	constexpr int32 OverflowButtonSize = 22;
+	LastToolBarButtonIndex = OverflowButtonIndex - 1;
+	
+	constexpr int32 OverflowButtonSize = 48;
 	
 	for (int32 ChildIdx = LastToolBarButtonIndex; ChildIdx >= 0; --ChildIdx)
 	{
@@ -80,11 +87,37 @@ void SClippingVerticalBox::Construct( const FArguments& InArgs )
 	StyleSet = InArgs._StyleSet;
 	StyleName = InArgs._StyleName;
 	bIsFocusable = InArgs._IsFocusable;
+	SelectedIndex = InArgs._SelectedIndex;
 
 	LastClippedIdx = ClippedIdx = INDEX_NONE;
 }
 
 void SClippingVerticalBox::AddWrapButton()
+{
+
+	InitializeWrapButton( WrapButton, false );
+	InitializeWrapButton( SelectedWrapButton, true );
+
+	AddSlot()
+	.Padding( 0.f )
+	[
+		SNew( SVerticalBox )
+		+ SVerticalBox::Slot()
+		.Padding(0.f)
+		[
+			WrapButton.ToSharedRef()
+		]
+		+ SVerticalBox::Slot()
+		.Padding(0.f)
+		[
+			SelectedWrapButton.ToSharedRef()
+		]
+
+	];
+ 
+}
+
+void SClippingVerticalBox::InitializeWrapButton( TSharedPtr<SComboButton>& Button, bool bCreateSelectedAppearance  )
 {
 	const FToolBarStyle& ToolBarStyle = StyleSet->GetWidgetStyle<FToolBarStyle>(StyleName);
 	const TSharedRef<SImage> IconWidget =
@@ -92,16 +125,39 @@ void SClippingVerticalBox::AddWrapButton()
 		.Visibility(EVisibility::HitTestInvisible)
 		.Image( &ToolBarStyle.ExpandBrush );
 
+
+	Style = ToolBarStyle.ButtonStyle;
+	Style.SetNormalPadding(0 );
+	Style.SetPressedPadding( 0 );
 	
-	// Construct the wrap button used in toolbars and menubars
+	SelectedStyle = ToolBarStyle.ButtonStyle;
+	SelectedStyle.SetNormalPadding(0 );
+	SelectedStyle.SetPressedPadding( 0 );
+	SelectedStyle.SetNormal( FSlateRoundedBoxBrush( FStyleColors::Primary, 4.f,
+		FLinearColor(0, 0, 0, .8), 0.5) );
+	
+	// Construct the wrap button used in toolbars and menu bars
 	// Always allow this to be focusable to prevent the menu from collapsing during interaction
-	WrapButton =
+	Button =
 		SNew( SComboButton )
 		.HasDownArrow( false )
-		.ButtonStyle(&ToolBarStyle.ButtonStyle)
-		.ContentPadding( FMargin(2.f, 4.f) )
+		.Visibility_Lambda( [this, bCreateSelectedAppearance] ()
+		{
+			bool bShowingSelected = SelectedIndex.Get() > LastClippedIdx - 1;
+			if ( bCreateSelectedAppearance )
+			{
+				return bShowingSelected ? EVisibility::Visible : EVisibility::Collapsed;				
+			}
+			return !bShowingSelected ? EVisibility::Visible : EVisibility::Collapsed;
+		})
+		.ButtonStyle( bCreateSelectedAppearance ? &SelectedStyle : &Style )
+		.ContentPadding( FMargin(2.f, 8.f) )	
 		.ToolTipText( NSLOCTEXT("Slate", "ExpandToolbar", "Click to expand toolbar") )
-		.OnGetMenuContent( OnWrapButtonClicked )
+		.OnGetMenuContent_Lambda( [ this ] ()
+		{
+			TSharedRef<SWidget> Widget = OnWrapButtonClicked.Execute();
+			return Widget;
+		})
 		.Cursor( EMouseCursor::Default )
 		.OnMenuOpenChanged(this, &SClippingVerticalBox::OnWrapButtonOpenChanged)
 		.IsFocusable(true)
@@ -109,19 +165,22 @@ void SClippingVerticalBox::AddWrapButton()
 		[
 		SNew(SVerticalBox)
 				+ SVerticalBox::Slot()
+				.Padding(0)
 				.AutoHeight()
 				.HAlign(HAlign_Center)	// Center the icon horizontally, so that large labels don't stretch out the artwork
 				[
 					IconWidget
 				]
+				+ SVerticalBox::Slot()
+				.AutoHeight()
+				.HAlign(HAlign_Center)
+				[
+					SNew(STextBlock)
+					.Text( LOCTEXT("ClippingVerticalBox.Icon.More", "More") )
+					.TextStyle(&FCoreStyle::Get().GetWidgetStyle<FTextBlockStyle>("SmallText"))
+				]
 			
 		];
-
-	AddSlot()
-	.Padding( 0.f )
-	[
-		WrapButton.ToSharedRef()
-	];
 }
 
 void SClippingVerticalBox::OnWrapButtonOpenChanged(bool bIsOpen)
@@ -149,3 +208,4 @@ EActiveTimerReturnType SClippingVerticalBox::UpdateWrapButtonStatus(double Curre
 	return EActiveTimerReturnType::Continue;
 }
 
+#undef LOCTEXT_NAMESPACE
\ No newline at end of file
diff --git a/Engine/Source/Runtime/Slate/Private/Framework/MultiBox/SClippingVerticalBox.h b/Engine/Source/Runtime/Slate/Private/Framework/MultiBox/SClippingVerticalBox.h
--- a/Engine/Source/Runtime/Slate/Private/Framework/MultiBox/SClippingVerticalBox.h
+++ b/Engine/Source/Runtime/Slate/Private/Framework/MultiBox/SClippingVerticalBox.h
@@ -7,6 +7,7 @@
 #include "Styling/CoreStyle.h"
 #include "Framework/SlateDelegates.h"
 #include "Widgets/SBoxPanel.h"
+#include "Styling/SlateTypes.h"
 
 class FArrangedChildren;
 class FPaintArgs;
@@ -21,12 +22,14 @@ public:
 		: _StyleSet(&FCoreStyle::Get())
 		, _StyleName(NAME_None)
 		, _IsFocusable(true)
+		, _SelectedIndex(INDEX_NONE )
 		{ }
 
 		SLATE_ARGUMENT(FOnGetContent, OnWrapButtonClicked)
 		SLATE_ARGUMENT(const ISlateStyle*, StyleSet)
 		SLATE_ARGUMENT(FName, StyleName)
 		SLATE_ARGUMENT(bool, IsFocusable)
+		SLATE_ATTRIBUTE( int32, SelectedIndex )
 	SLATE_END_ARGS()
 
 	/** SWidget interface */
@@ -48,10 +51,22 @@ public:
 private:
 	void OnWrapButtonOpenChanged(bool bIsOpen);
 	EActiveTimerReturnType UpdateWrapButtonStatus(double CurrentTime, float DeltaTime);
+
+	/**
+	 * Initializes a wrap button that can handle clipped content
+	 * 
+	 * @param Button the SComboButton that will make up the clipped overflow button
+	 * @param bCreateSelectedAppearance if true, the button will have a "selected" appearance to denote that something within it is selected
+	 */
+	void InitializeWrapButton( TSharedPtr<SComboButton>& Button, bool bCreateSelectedAppearance );
+	
 private:
-	/** The button that is displayed when a toolbar or menubar is clipped */
+	/** The button that is displayed when a toolbar or menubar is clipped and something within it is not selected */
 	TSharedPtr<SComboButton> WrapButton;
 
+	/** The button that is displayed when a toolbar or menubar is clipped and something within the clipped content is selected */
+	TSharedPtr<SComboButton> SelectedWrapButton;
+
 	/** Callback for when the wrap button is clicked */
 	FOnGetContent OnWrapButtonClicked;
 
@@ -71,5 +86,17 @@ private:
 	/** The style to use */
 	const ISlateStyle* StyleSet;
 
+	/** the button style for the clipped content button when nothing within the clipped content is selected */
+	FButtonStyle Style;
+	
+	/** the button style for the clipped content button when something within the clipped content is selected */
+	FButtonStyle SelectedStyle;
+
+	/** the index in the toolbar that is currently selected */
+	TAttribute<int32> SelectedIndex = INDEX_NONE;
+
+	/** the last index in the toolbar which is not clipped by the clipped content button */
+	mutable int32 LastToolBarButtonIndex = INDEX_NONE;
+
 	FName StyleName;
 };
diff --git a/Engine/Source/Runtime/Slate/Private/Framework/MultiBox/SToolBarButtonBlock.cpp b/Engine/Source/Runtime/Slate/Private/Framework/MultiBox/SToolBarButtonBlock.cpp
--- a/Engine/Source/Runtime/Slate/Private/Framework/MultiBox/SToolBarButtonBlock.cpp
+++ b/Engine/Source/Runtime/Slate/Private/Framework/MultiBox/SToolBarButtonBlock.cpp
@@ -18,12 +18,12 @@ FToolBarButtonBlock::FToolBarButtonBlock( FButtonArgs ButtonArgs )
 	, LabelOverride( ButtonArgs.LabelOverride )
 	, ToolTipOverride( ButtonArgs.ToolTipOverride )
 	, IconOverride( ButtonArgs.IconOverride )
-   // , BorderBrushName(ButtonArgs.BorderBrushName)
 	, LabelVisibility()
 	, UserInterfaceActionType(ButtonArgs.UserInterfaceActionType != EUserInterfaceActionType::None ?
 		ButtonArgs.UserInterfaceActionType : EUserInterfaceActionType::Button)
 	, bIsFocusable(false)
 	, bForceSmallIcons(false)
+	, GetDecoratedButtonDelegate( ButtonArgs.GetDecoratedButtonDelegate )
 {
 }
 
@@ -61,6 +61,11 @@ void FToolBarButtonBlock::SetOnGetMenuContent(const FOnGetContent& InOnGetMenuCo
 	OnGetMenuContent = InOnGetMenuContent;
 }
 
+void FToolBarButtonBlock::SetGetDecoratedButtonDelegate( const FGetDecoratedButtonDelegate& InGetDecoratedButtonDelegate )
+{
+	GetDecoratedButtonDelegate = InGetDecoratedButtonDelegate;
+}
+
 void FToolBarButtonBlock::CreateMenuEntry(FMenuBuilder& MenuBuilder) const
 {
 	// Setup Command Context 
@@ -376,6 +381,11 @@ void SToolBarButtonBlock::BuildMultiBlockWidget(const ISlateStyle* StyleSet, con
 					CheckBox.ToSharedRef()
 				];
 		}
+		
+		if ( ToolBarButtonBlock->GetDecoratedButtonDelegate.IsBound() )
+		{
+			CheckBoxWidget = ToolBarButtonBlock->GetDecoratedButtonDelegate.Execute( CheckBoxWidget );
+		}
 
 		ChildSlot
 		[
diff --git a/Engine/Source/Runtime/Slate/Public/Framework/MultiBox/MultiBox.h b/Engine/Source/Runtime/Slate/Public/Framework/MultiBox/MultiBox.h
--- a/Engine/Source/Runtime/Slate/Public/Framework/MultiBox/MultiBox.h
+++ b/Engine/Source/Runtime/Slate/Public/Framework/MultiBox/MultiBox.h
@@ -332,6 +332,18 @@ public:
 		return bShouldCloseWindowAfterMenuSelection;
 	}
 
+	/**
+	 * returns the index of the item which was last selected in this box
+	 */
+	int32 GetLastSelectedCommandIndex() const;
+	
+	/**
+	 * sets the index of the item which was last selected in this box
+	 *
+	 * @param LastSelectedCommandIndex the index of the last command that was selected for this box
+	 */
+	void SetLastSelectedCommandIndex( int32 LastSelectedCommandIndex );
+
 	/**
 	 * Adds a MultiBlock to this MultiBox, to the end of the list
 	 */
@@ -480,6 +492,10 @@ private:
 
 	/** True if window that owns any widgets created from this multibox should be closed automatically after the user commits to a menu choice */
 	bool bShouldCloseWindowAfterMenuSelection;
+
+	/** the Last selected command index. If unset this has a value of INDEX_NONE */
+	int32 LastSelectedCommandIndex;
+
 };
 
 
diff --git a/Engine/Source/Runtime/Slate/Public/Framework/MultiBox/MultiBoxBuilder.h b/Engine/Source/Runtime/Slate/Public/Framework/MultiBox/MultiBoxBuilder.h
--- a/Engine/Source/Runtime/Slate/Public/Framework/MultiBox/MultiBoxBuilder.h
+++ b/Engine/Source/Runtime/Slate/Public/Framework/MultiBox/MultiBoxBuilder.h
@@ -97,6 +97,8 @@ public:
 	 */
 	SLATE_API void PushCommandList( const TSharedRef< const FUICommandList > CommandList );
 
+	SLATE_API void SetLastSelectedCommandIndex( int32 InLastSelectedCommandIndex );
+
 
 	/**
 	 * Pops the current command list.
diff --git a/Engine/Source/Runtime/Slate/Public/Framework/MultiBox/SToolBarButtonBlock.h b/Engine/Source/Runtime/Slate/Public/Framework/MultiBox/SToolBarButtonBlock.h
--- a/Engine/Source/Runtime/Slate/Public/Framework/MultiBox/SToolBarButtonBlock.h
+++ b/Engine/Source/Runtime/Slate/Public/Framework/MultiBox/SToolBarButtonBlock.h
@@ -18,6 +18,9 @@ class SBorder;
 /** Delegate used by multi-box to call a user function to populate a new menu.  Used for spawning sub-menus and pull-down menus. */
 DECLARE_DELEGATE_OneParam( FNewMenuDelegate, class FMenuBuilder& );
 
+/**  Delegate that takes the default toolbar button created by this as a parameter, adds any decorator needed and returns the new widget*/
+DECLARE_DELEGATE_RetVal_OneParam( TSharedRef<SWidget>, FGetDecoratedButtonDelegate, TSharedRef<SWidget> );
+
 struct FButtonArgs : public TSharedFromThis<FButtonArgs>
 {
 	TSharedPtr< const FUICommandInfo > Command;
@@ -33,6 +36,7 @@ struct FButtonArgs : public TSharedFromThis<FButtonArgs>
 	
 	FNewMenuDelegate CustomMenuDelegate;
 	FOnGetContent OnGetMenuContent;
+	FGetDecoratedButtonDelegate GetDecoratedButtonDelegate;
 
 	explicit FButtonArgs() {}
 };
@@ -91,6 +95,13 @@ public:
 
 	SLATE_API void SetOnGetMenuContent(const FOnGetContent& OnGetMenuContent);
 
+	/**
+	 *  Delegate that takes the default toolbar button created by this as a parameter, adds any decorator needed and returns the new widget
+	 *
+	 * @param InGetDecoratedButtonDelegate the delegate that handles decorating the button
+	 */
+	void SetGetDecoratedButtonDelegate( const FGetDecoratedButtonDelegate& InGetDecoratedButtonDelegate );
+
 protected:
 	
 	SLATE_API bool GetIsFocusable() const;
@@ -134,6 +145,9 @@ private:
 
 	/** Delegate to execute to get the menu content of this button */
 	FOnGetContent OnGetMenuContent;
+
+	/**  Delegate that takes the default toolbar button created by this as a parameter, adds any decorator needed and returns the new widget*/
+	FGetDecoratedButtonDelegate GetDecoratedButtonDelegate;
 };
 
 
diff --git a/Engine/Source/Runtime/SlateCore/Private/Styling/StarshipCoreStyle.cpp b/Engine/Source/Runtime/SlateCore/Private/Styling/StarshipCoreStyle.cpp
--- a/Engine/Source/Runtime/SlateCore/Private/Styling/StarshipCoreStyle.cpp
+++ b/Engine/Source/Runtime/SlateCore/Private/Styling/StarshipCoreStyle.cpp
@@ -345,6 +345,7 @@ TSharedRef<ISlateStyle> FStarshipCoreStyle::Create()
 		Style->Set("Icons.Delete", new IMAGE_BRUSH_SVG("Starship/Common/delete-outline", Icon16x16));
 		Style->Set("Icons.Save", new IMAGE_BRUSH_SVG("Starship/Common/save", Icon16x16));
 		Style->Set("Icons.SaveModified", new IMAGE_BRUSH_SVG("Starship/Common/save-modified", Icon16x16));
+		Style->Set("Icons.Favorites.Small", new IMAGE_BRUSH_SVG("Starship/Common/Favorite", Icon20x20));
 
 		Style->Set("Icons.Import", new IMAGE_BRUSH_SVG("Starship/Common/import", Icon16x16));
 		Style->Set("Icons.Filter", new IMAGE_BRUSH_SVG("Starship/Common/filter", Icon16x16));
