SUniformWrapPanel has jitter\scrollbar issues

SUniformWrapPanel, when made too small, can lose its panel layout when used with a scrollbar, or, alternatively, can develop a heavy "jitter". The reason for this is that while laying the children out, there is an edge case that is missed in displaying children, where none of the best fit cases work. We need to provide the most reasonable default view for when this happens. I will add some examples.

diff --git a/Engine/Source/Runtime/Slate/Private/Widgets/Layout/SUniformWrapPanel.cpp b/Engine/Source/Runtime/Slate/Private/Widgets/Layout/SUniformWrapPanel.cpp
--- a/Engine/Source/Runtime/Slate/Private/Widgets/Layout/SUniformWrapPanel.cpp
+++ b/Engine/Source/Runtime/Slate/Private/Widgets/Layout/SUniformWrapPanel.cpp
@@ -191,6 +191,12 @@ FVector2D SUniformWrapPanel::ComputeDesiredSize( float ) const
 			NumRows = FMath::CeilToInt((float)NumVisibleChildren / (float)NumColumns);
 			return FVector2D(NumColumns * MaxChildDesiredSize.X, NumRows * MaxChildDesiredSize.Y);
 		}
+		
+		/* If we have gotten this far we may be small enough that the best fit formulas fail yet we still have visible children to show ~
+		 * default to one column the width of the child desired size and the number of children rows times their height */
+		NumRows = NumVisibleChildren;
+		NumColumns = 1;
+		return FVector2D(MaxChildDesiredSize.X, NumRows * MaxChildDesiredSize.Y);
 	}
 
 	return FVector2D::ZeroVector;
