Add tooltips for window controls and improve tooltips for tab controls.

diff --git a/Engine/Source/Editor/MainFrame/Private/MainFrameModule.cpp b/Engine/Source/Editor/MainFrame/Private/MainFrameModule.cpp
--- a/Engine/Source/Editor/MainFrame/Private/MainFrameModule.cpp
+++ b/Engine/Source/Editor/MainFrame/Private/MainFrameModule.cpp
@@ -199,6 +199,7 @@ void FMainFrameModule::CreateDefaultMainFrameAuxiliary(const bool bStartImmersiv
 		TSharedRef<SWindow> RootWindow = SNew(SWindow)
 			.AutoCenter(WindowConfig.CenterRules)
 			.Title( WindowConfig.WindowTitle )
+			.CloseButtonToolTipText(LOCTEXT("MainFrameModuleCloseButtonToolTip", "Close Unreal Editor"))
 			.IsInitiallyMaximized( DefaultWindowLocation.InitiallyMaximized )
 			.ScreenPosition( DefaultWindowLocation.ScreenPosition )
 			.ClientSize( DefaultWindowLocation.WindowSize )
diff --git a/Engine/Source/Runtime/Slate/Private/Framework/Application/SWindowTitleBar.cpp b/Engine/Source/Runtime/Slate/Private/Framework/Application/SWindowTitleBar.cpp
--- a/Engine/Source/Runtime/Slate/Private/Framework/Application/SWindowTitleBar.cpp
+++ b/Engine/Source/Runtime/Slate/Private/Framework/Application/SWindowTitleBar.cpp
@@ -10,6 +10,7 @@ void SWindowTitleBar::Construct( const FArguments& InArgs, const TSharedRef<SWin
 	Style = InArgs._Style;
 	ShowAppIcon = InArgs._ShowAppIcon;
 	Title = InArgs._Title;
+	CloseButtonToolTipText = InArgs._CloseButtonToolTipText;
 
 	WindowMenuSlot = nullptr;
 
@@ -145,6 +146,7 @@ void SWindowTitleBar::MakeTitleBarContentWidgets( TSharedPtr< SWidget >& OutLeft
 						.Image(this, &SWindowTitleBar::GetMinimizeImage)
 						.ColorAndOpacity(this, &SWindowTitleBar::GetWindowTitleContentColor)
 						.AccessibleText(NSLOCTEXT("WindowTitleBar", "Minimize", "Minimize"))
+						.ToolTipText(NSLOCTEXT("WindowTitleBar", "Minimize", "Minimize"))
 				];
 
 		MaximizeRestoreButton = SNew(SButton)
@@ -159,7 +161,8 @@ void SWindowTitleBar::MakeTitleBarContentWidgets( TSharedPtr< SWidget >& OutLeft
 					SNew(SImage)
 						.Image(this, &SWindowTitleBar::GetMaximizeRestoreImage)
 						.ColorAndOpacity(this, &SWindowTitleBar::GetWindowTitleContentColor)
-						.AccessibleText(NSLOCTEXT("WindowTitleBar", "Maximize", "Maximize"))
+						.AccessibleText(this, &SWindowTitleBar::GetMaximizeButtonToolTip)
+						.ToolTipText(this, &SWindowTitleBar::GetMaximizeButtonToolTip)
 				];
 
 		CloseButton = SNew(SButton)
@@ -174,7 +177,8 @@ void SWindowTitleBar::MakeTitleBarContentWidgets( TSharedPtr< SWidget >& OutLeft
 					SNew(SImage)
 						.Image(this, &SWindowTitleBar::GetCloseImage)
 						.ColorAndOpacity(this, &SWindowTitleBar::GetWindowTitleContentColor)
-						.AccessibleText(NSLOCTEXT("WindowTitleBar", "Close", "Close"))
+						.AccessibleText(CloseButtonToolTipText)
+						.ToolTipText(CloseButtonToolTipText)
 				];
 	}
 #endif //!PLATFORM_MAC
@@ -579,3 +583,15 @@ FText SWindowTitleBar::HandleWindowTitleText( ) const
 
 	return OwnerWindow->GetTitle();
 }
+
+FText SWindowTitleBar::GetMaximizeButtonToolTip() const
+{
+	TSharedPtr<SWindow> OwnerWindow = OwnerWindowPtr.Pin();
+
+	if (OwnerWindow.IsValid() && OwnerWindow->IsWindowMaximized())
+	{
+		return NSLOCTEXT("WindowTitleBar", "MaximizeRestoreDown", "Restore Down");
+	}
+
+	return NSLOCTEXT("WindowTitleBar", "Maximize", "Maximize");
+}
\ No newline at end of file
diff --git a/Engine/Source/Runtime/Slate/Private/Framework/Application/SlateApplication.cpp b/Engine/Source/Runtime/Slate/Private/Framework/Application/SlateApplication.cpp
--- a/Engine/Source/Runtime/Slate/Private/Framework/Application/SlateApplication.cpp
+++ b/Engine/Source/Runtime/Slate/Private/Framework/Application/SlateApplication.cpp
@@ -4597,7 +4597,8 @@ TSharedRef<SImage> FSlateApplication::MakeImage( const TAttribute<const FSlateBr
 TSharedRef<SWidget> FSlateApplication::MakeWindowTitleBar(const FWindowTitleBarArgs& InArgs, TSharedPtr<IWindowTitleBar>& OutTitleBar) const
 {
 	TSharedRef<SWindowTitleBar> TitleBar = SNew(SWindowTitleBar, InArgs.Window, InArgs.CenterContent, InArgs.CenterContentAlignment)
-		.Visibility(EVisibility::SelfHitTestInvisible);
+		.Visibility(EVisibility::SelfHitTestInvisible)
+		.CloseButtonToolTipText(InArgs.CloseButtonToolTipText);
 
 	OutTitleBar = TitleBar;
 
diff --git a/Engine/Source/Runtime/Slate/Private/Framework/Docking/SDockingArea.cpp b/Engine/Source/Runtime/Slate/Private/Framework/Docking/SDockingArea.cpp
--- a/Engine/Source/Runtime/Slate/Private/Framework/Docking/SDockingArea.cpp
+++ b/Engine/Source/Runtime/Slate/Private/Framework/Docking/SDockingArea.cpp
@@ -319,6 +319,7 @@ void SDockingArea::SetParentWindow( const TSharedRef<SWindow>& NewParentWindow )
 		FWindowTitleBarArgs Args(NewParentWindow);
 		Args.CenterContent = SNullWidget::NullWidget;
 		Args.CenterContentAlignment = HAlign_Fill;
+		Args.CloseButtonToolTipText = NewParentWindow->GetWindowCloseButtonToolTipText();
 
 		TSharedRef<SWidget> TitleBarWidget = FSlateApplication::Get().MakeWindowTitleBar(Args, TitleBar);
 		(*WindowControlsArea)
diff --git a/Engine/Source/Runtime/Slate/Private/Widgets/Docking/SDockTab.cpp b/Engine/Source/Runtime/Slate/Private/Widgets/Docking/SDockTab.cpp
--- a/Engine/Source/Runtime/Slate/Private/Widgets/Docking/SDockTab.cpp
+++ b/Engine/Source/Runtime/Slate/Private/Widgets/Docking/SDockTab.cpp
@@ -952,7 +952,7 @@ FText SDockTab::GetCloseButtonToolTipText() const
 		GetVisualTabRole() == ETabRole::MajorTab ? FTabCommands::Get().CloseMajorTab : FTabCommands::Get().CloseMinorTab;
 
 	FFormatNamedArguments Arguments;
-	Arguments.Add(TEXT("Label"), CloseCommand->GetLabel());
+	Arguments.Add(TEXT("Label"), NSLOCTEXT("DockTab", "CloseButtonToolTipLabel", "Close Tab"));
 
 	FText InputText = CloseCommand->GetInputText();
 	if (InputText.IsEmptyOrWhitespace())
@@ -961,10 +961,10 @@ FText SDockTab::GetCloseButtonToolTipText() const
 	}
 	else
 	{
-		Arguments.Add(TEXT("InputText"), FText::Format(NSLOCTEXT("DockTab", "CloseButtonInputText", " ({0})"), InputText));
+		Arguments.Add(TEXT("InputText"), FText::Format(NSLOCTEXT("DockTab", "CloseButtonInputText", "({0})"), InputText));
 	}
 
-	return FText::Format(NSLOCTEXT("DockTab", "CloseButtonToolTip", "{Label}{InputText}"), Arguments);
+	return FText::Format(NSLOCTEXT("DockTab", "CloseButtonToolTipFormat", "{Label} {InputText}"), Arguments);
 }
 
 EVisibility SDockTab::HandleIsCloseButtonVisible() const
diff --git a/Engine/Source/Runtime/Slate/Public/Framework/Application/SWindowTitleBar.h b/Engine/Source/Runtime/Slate/Public/Framework/Application/SWindowTitleBar.h
--- a/Engine/Source/Runtime/Slate/Public/Framework/Application/SWindowTitleBar.h
+++ b/Engine/Source/Runtime/Slate/Public/Framework/Application/SWindowTitleBar.h
@@ -102,6 +102,7 @@ public:
 		SLATE_STYLE_ARGUMENT(FWindowStyle, Style)
 		SLATE_ARGUMENT(bool, ShowAppIcon)
 		SLATE_ATTRIBUTE(FText, Title)
+		SLATE_ATTRIBUTE(FText, CloseButtonToolTipText)
 	SLATE_END_ARGS()
 	
 public:
@@ -201,6 +202,8 @@ private:
 
 	FText HandleWindowTitleText( ) const;
 
+	FText GetMaximizeButtonToolTip() const;
+
 protected:
 
 	// Holds a weak pointer to the owner window.
@@ -227,6 +230,7 @@ private:
 	TSharedPtr<SButton> CloseButton;
 
 	TAttribute<FText> Title;
+	TAttribute<FText> CloseButtonToolTipText;
 
 	SHorizontalBox::FSlot* WindowMenuSlot;
 	SHorizontalBox::FSlot* RightSideContentSlot;
diff --git a/Engine/Source/Runtime/SlateCore/Private/Widgets/SWindow.cpp b/Engine/Source/Runtime/SlateCore/Private/Widgets/SWindow.cpp
--- a/Engine/Source/Runtime/SlateCore/Private/Widgets/SWindow.cpp
+++ b/Engine/Source/Runtime/SlateCore/Private/Widgets/SWindow.cpp
@@ -240,6 +240,7 @@ void SWindow::Construct(const FArguments& InArgs)
 	this->WindowBackground = &InArgs._Style->BackgroundBrush;
 
 	this->Title = InArgs._Title;
+	this->CloseButtonToolTipText = InArgs._CloseButtonToolTipText;
 	this->bDragAnywhere = InArgs._bDragAnywhere;
 	this->TransparencySupport = InArgs._SupportsTransparency.Value;
 	this->Opacity = InArgs._InitialOpacity;
@@ -506,6 +507,7 @@ TSharedRef<SWidget> SWindow::MakeWindowTitleBar(const TSharedRef<SWindow>& Windo
 	FWindowTitleBarArgs Args(Window);
 	Args.CenterContent = CenterContent;
 	Args.CenterContentAlignment = TitleContentAlignment;
+	Args.CloseButtonToolTipText = CloseButtonToolTipText;
 
 	return FSlateApplicationBase::Get().MakeWindowTitleBar(Args, TitleBar);
 }
@@ -543,6 +545,7 @@ void SWindow::ConstructWindowInternals()
 	FWindowTitleBarArgs Args(SharedThis(this));
 	Args.CenterContent = nullptr;
 	Args.CenterContentAlignment = GetTitleAlignment();
+	Args.CloseButtonToolTipText = CloseButtonToolTipText;
 
 	TSharedRef<SWidget> TitleBarWidget = FSlateApplicationBase::Get().MakeWindowTitleBar(Args, TitleBar);
 
@@ -1140,6 +1143,11 @@ EVisibility SWindow::GetWindowVisibility() const
 		: EVisibility::HitTestInvisible;
 }
 
+TAttribute<FText> SWindow::GetWindowCloseButtonToolTipText() const
+{
+	return CloseButtonToolTipText;
+}
+
 void SWindow::UpdateMorphTargetShape( const FSlateRect& TargetShape )
 {
 	Morpher.TargetMorphShape = TargetShape;
diff --git a/Engine/Source/Runtime/SlateCore/Public/Application/SlateApplicationBase.h b/Engine/Source/Runtime/SlateCore/Public/Application/SlateApplicationBase.h
--- a/Engine/Source/Runtime/SlateCore/Public/Application/SlateApplicationBase.h
+++ b/Engine/Source/Runtime/SlateCore/Public/Application/SlateApplicationBase.h
@@ -44,6 +44,9 @@ struct FWindowTitleBarArgs
 
 	/** The horizontal alignment of the center content. */
 	EHorizontalAlignment CenterContentAlignment;
+
+	/** The text for the tooltip that shows when the user hovers over the Close button. */
+	TAttribute<FText> CloseButtonToolTipText;
 };
 
 /**
diff --git a/Engine/Source/Runtime/SlateCore/Public/Widgets/SWindow.h b/Engine/Source/Runtime/SlateCore/Public/Widgets/SWindow.h
--- a/Engine/Source/Runtime/SlateCore/Public/Widgets/SWindow.h
+++ b/Engine/Source/Runtime/SlateCore/Public/Widgets/SWindow.h
@@ -272,7 +272,7 @@ public:
 		, _LayoutBorder( _Style->BorderPadding )
 		, _UserResizeBorder(FMargin(5, 5, 5, 5))
 		, _bManualManageDPI( false )
-
+		, _CloseButtonToolTipText(NSLOCTEXT("SWindow", "Window_DefaultCloseButtonToolTip", "Close"))
 	{
 	}
 
@@ -375,7 +375,10 @@ public:
 
 		SLATE_DEFAULT_SLOT( FArguments, Content )
 
-	SLATE_END_ARGS()
+		/** Optional tooltip for the close button, when HasCloseButton is true. */
+		SLATE_ATTRIBUTE(FText, CloseButtonToolTipText)
+
+		SLATE_END_ARGS()
 
 	/**
 	 * Default constructor. Use SNew(SWindow) instead.
@@ -1070,6 +1073,8 @@ public:
 	/** Windows that are not hittestable should not show up in the hittest grid. */
 	SLATECORE_API EVisibility GetWindowVisibility() const;
 
+	/** Return the tooltip text that should be used for the window's Close button. */
+	SLATECORE_API TAttribute<FText> GetWindowCloseButtonToolTipText() const;
 protected:
 	/**Returns swindow title bar widgets. */
 	SLATECORE_API virtual TSharedRef<SWidget> MakeWindowTitleBar(const TSharedRef<SWindow>& Window, const TSharedPtr<SWidget>& CenterContent, EHorizontalAlignment CenterContentAlignment);
@@ -1088,6 +1093,9 @@ protected:
 	/** Title of the window, displayed in the title bar as well as potentially in the task bar (Windows platform) */
 	TAttribute<FText> Title;
 
+	/** Tooltip text displayed over the window's Close button (on Windows only). */
+	TAttribute<FText> CloseButtonToolTipText;
+
 	/** When true, grabbing anywhere on the window will allow it to be dragged. */
 	bool bDragAnywhere;
 
