Slate: Deprecate SListView::ItemHeight and STreeViewItemHeight. ItemHeight and ItemWidth are only used when the panel is is as a tile or both value are assigned. With TreeView and ListView, you can't assigned both. The function are protected, a subclass of those could but it's up to them to expose the ItemHeight and ItemWidth. This behavior changed was introduced with CL 7634094.

diff --git a/Engine/Source/Runtime/Slate/Private/Widgets/Views/SListPanel.cpp b/Engine/Source/Runtime/Slate/Private/Widgets/Views/SListPanel.cpp
--- a/Engine/Source/Runtime/Slate/Private/Widgets/Views/SListPanel.cpp
+++ b/Engine/Source/Runtime/Slate/Private/Widgets/Views/SListPanel.cpp
@@ -7,17 +7,30 @@
 // Used to subtract a tiny amount from the available dimension to avoid floating point precision problems when arranging children
 static const float FloatingPointPrecisionOffset = 0.001f;
 
+SLATE_IMPLEMENT_WIDGET(SListPanel)
+void SListPanel::PrivateRegisterAttributes(FSlateAttributeInitializer& AttributeInitializer)
+{
+	SLATE_ADD_MEMBER_ATTRIBUTE_DEFINITION_WITH_NAME(AttributeInitializer, "ItemWidth", ItemWidthAttribute, EInvalidateWidgetReason::Layout);
+	SLATE_ADD_MEMBER_ATTRIBUTE_DEFINITION_WITH_NAME(AttributeInitializer, "ItemHeight", ItemHeightAttribute, EInvalidateWidgetReason::Layout);
+	SLATE_ADD_MEMBER_ATTRIBUTE_DEFINITION_WITH_NAME(AttributeInitializer, "NumDesiredItems", NumDesiredItemsAttribute, EInvalidateWidgetReason::Layout);
+	SLATE_ADD_MEMBER_ATTRIBUTE_DEFINITION_WITH_NAME(AttributeInitializer, "ItemAlignment", ItemAlignmentAttribute, EInvalidateWidgetReason::Layout);
+}
+
 SListPanel::SListPanel()
 	: Children(this)
+	, ItemAlignmentAttribute(*this)
+	, ItemWidthAttribute(*this, 16.0f)
+	, ItemHeightAttribute(*this, 0.0f)
+	, NumDesiredItemsAttribute(*this, 0)
 {
 }
 
 void SListPanel::Construct( const FArguments& InArgs )
 {
-	ItemWidth = InArgs._ItemWidth;
-	ItemHeight = InArgs._ItemHeight;
-	NumDesiredItems = InArgs._NumDesiredItems;
-	ItemAlignment = InArgs._ItemAlignment;
+	ItemWidthAttribute.Assign(*this, InArgs._ItemWidth);
+	ItemHeightAttribute.Assign(*this, InArgs._ItemHeight);
+	NumDesiredItemsAttribute.Assign(*this, InArgs._NumDesiredItems);
+	ItemAlignmentAttribute.Assign(*this, InArgs._ItemAlignment);
 	Orientation = InArgs._ListOrientation;
 	Children.AddSlots(MoveTemp(const_cast<TArray<FSlot::FSlotArguments>&>(InArgs._Slots)));
 }
@@ -42,7 +55,7 @@ void SListPanel::OnArrangeChildren( const FGeometry& AllottedGeometry, FArranged
 		if (ShouldArrangeAsTiles())
 		{
 			// This is a tile view list - arrange items side by side along the line axis until there is no more room, then create a new line along the scroll axis.
-			const EListItemAlignment ListItemAlignment = ItemAlignment.Get();
+			const EListItemAlignment ListItemAlignment = ItemAlignmentAttribute.Get();
 			const float ItemPadding = GetItemPadding(AllottedGeometry, ListItemAlignment);
 			const float HalfItemPadding = ItemPadding * 0.5f;
 
@@ -130,12 +143,12 @@ void SListPanel::Tick( const FGeometry& AllottedGeometry, const double InCurrent
 	{
 		const FTableViewDimensions AllottedDimensions(Orientation, AllottedGeometry.GetLocalSize());
 
-		const EListItemAlignment ListItemAlignment = ItemAlignment.Get();
+		const EListItemAlignment ListItemAlignment = ItemAlignmentAttribute.Get();
 		const float ItemPadding = GetItemPadding(AllottedGeometry, ListItemAlignment);
 		FTableViewDimensions ItemDimensions = GetItemSize(AllottedGeometry, ListItemAlignment);
 		ItemDimensions.LineAxis += ItemPadding;
 
-		const int32 NumChildren = NumDesiredItems.Get();
+		const int32 NumChildren = NumDesiredItemsAttribute.Get();
 
 		if (ItemDimensions.LineAxis > 0.0f && NumChildren > 0)
 		{
@@ -149,7 +162,7 @@ void SListPanel::Tick( const FGeometry& AllottedGeometry, const double InCurrent
 	}
 	else
 	{
-		PreferredNumLines = NumDesiredItems.Get();
+		PreferredNumLines = NumDesiredItemsAttribute.Get();
 	}
 }
 
@@ -220,12 +233,12 @@ void SListPanel::ClearItems()
 
 FTableViewDimensions SListPanel::GetDesiredItemDimensions() const
 {
-	return FTableViewDimensions(Orientation, ItemWidth.Get(), ItemHeight.Get());
+	return FTableViewDimensions(Orientation, ItemWidthAttribute.Get(), ItemHeightAttribute.Get());
 }
 
 float SListPanel::GetItemPadding(const FGeometry& AllottedGeometry) const
 {
-	return GetItemPadding(AllottedGeometry, ItemAlignment.Get());
+	return GetItemPadding(AllottedGeometry, ItemAlignmentAttribute.Get());
 }
 
 float SListPanel::GetItemPadding(const FGeometry& AllottedGeometry, const EListItemAlignment ListItemAlignment) const
@@ -250,7 +263,7 @@ float SListPanel::GetItemPadding(const FGeometry& AllottedGeometry, const EListI
 
 FTableViewDimensions SListPanel::GetItemSize(const FGeometry& AllottedGeometry) const
 {
-	return GetItemSize(AllottedGeometry, ItemAlignment.Get());
+	return GetItemSize(AllottedGeometry, ItemAlignmentAttribute.Get());
 }
 
 FTableViewDimensions SListPanel::GetItemSize(const FGeometry& AllottedGeometry, const EListItemAlignment ListItemAlignment) const
@@ -316,10 +329,10 @@ bool SListPanel::ShouldArrangeAsTiles() const
 
 void SListPanel::SetItemHeight(TAttribute<float> Height)
 {
-	ItemHeight = Height;
+	ItemHeightAttribute.Assign(*this, MoveTemp(Height));
 }
 
 void SListPanel::SetItemWidth(TAttribute<float> Width)
 {
-	ItemWidth = Width;
+	ItemWidthAttribute.Assign(*this, MoveTemp(Width));
 }
diff --git a/Engine/Source/Runtime/Slate/Private/Widgets/Views/SListPanel.h b/Engine/Source/Runtime/Slate/Private/Widgets/Views/SListPanel.h
--- a/Engine/Source/Runtime/Slate/Private/Widgets/Views/SListPanel.h
+++ b/Engine/Source/Runtime/Slate/Private/Widgets/Views/SListPanel.h
@@ -20,6 +20,8 @@ class FArrangedChildren;
  */
 class SListPanel : public SPanel
 {
+	SLATE_DECLARE_WIDGET_API(SListPanel, SPanel, SLATE_API)
+
 public:
 	/** A ListPanel slot is very simple - it just stores a widget. */
 	class FSlot : public TSlotBase<FSlot>
@@ -121,18 +123,23 @@ protected:
 	SLATE_API bool ShouldArrangeAsTiles() const;
 	
 protected:
+#if WITH_EDITORONLY_DATA
+	UE_DEPRECATED(5.5, "Direct access to ItemWidth is now deprecated. Use the setter or getter.")
+	FInvalidatableBrushAttribute ItemWidth;
+	
+	UE_DEPRECATED(5.5, "Direct access to ItemHeight is now deprecated. Use the setter or getter.")
+	FInvalidatableBrushAttribute ItemHeight;
+	
+	UE_DEPRECATED(5.5, "Direct access to NumDesiredItems is now deprecated. Use the setter or getter.")
+	FInvalidatableBrushAttribute NumDesiredItems;
+	
+	UE_DEPRECATED(5.5, "Direct access to ItemAlignment is now deprecated. Use the setter or getter.")
+	FInvalidatableBrushAttribute ItemAlignment;
+#endif
 
+protected:
 	/** The children being arranged by this panel */
 	TPanelChildren<FSlot> Children;
-
-	/** The uniform item width used to arrange the children. Only relevant for tile views. */
-	TAttribute<float> ItemWidth;
-	
-	/** The uniform item height used to arrange the children */
-	TAttribute<float> ItemHeight;
-
-	/** Total number of items that the tree wants to visualize */
-	TAttribute<int32> NumDesiredItems;
 	
 	/**
 	 * The offset of the view area from the top of the list in item heights.
@@ -143,18 +150,28 @@ protected:
 	/** Amount scrolled past beginning/end of list in Slate Units. */
 	float OverscrollAmount = 0.f;
 
+	/** The preferred number of lines that this widget should have orthogonal to the scroll axis. Only relevant for tile views. */
+	int32 PreferredNumLines = 1;
+
 	/**
 	 * When true, a refresh of the table view control that is using this panel is pending.
 	 * Some of the widgets in this panel are associated with items that may no longer be sound data.
 	 */
 	bool bIsRefreshPending = false;
 
-	/** How should be horizontally aligned? Only relevant for tile views. */
-	TAttribute<EListItemAlignment> ItemAlignment;
-
 	/** Overall orientation of the list for layout and scrolling. Only relevant for tile views. */
 	EOrientation Orientation;
 
-	/** The preferred number of lines that this widget should have orthogonal to the scroll axis. Only relevant for tile views. */
-	int32 PreferredNumLines = 1;
+private:
+	/** How should be horizontally aligned? Only relevant for tile views. */
+	TSlateAttribute<EListItemAlignment> ItemAlignmentAttribute;
+
+	/** The uniform item width used to arrange the children. Only relevant for tile views. */
+	TSlateAttribute<float> ItemWidthAttribute;
+
+	/** The uniform item height used to arrange the children */
+	TSlateAttribute<float> ItemHeightAttribute;
+
+	/** Total number of items that the tree wants to visualize */
+	TSlateAttribute<int32> NumDesiredItemsAttribute;
 };
diff --git a/Engine/Source/Runtime/Slate/Public/Widgets/Views/SListView.h b/Engine/Source/Runtime/Slate/Public/Widgets/Views/SListView.h
--- a/Engine/Source/Runtime/Slate/Public/Widgets/Views/SListView.h
+++ b/Engine/Source/Runtime/Slate/Public/Widgets/Views/SListView.h
@@ -43,11 +43,10 @@
  *   Given: TArray< TSharedPtr<FText> > Items;
  *
  *   SNew( SListView< TSharedPtr<FText> > )
- *     .ItemHeight(24)
  *     .ListItemsSource( &Items )
  *     .OnGenerateRow(this, &MyClass::GenerateItemRow)
  *
- * In the example we make all our widgets be 24 screen units tall. The ListView will create widgets based on data items
+ * The ListView will create widgets based on data items
  * in the Items TArray. When the ListView needs to generate an item, it will do so using the specified OnGenerateRow method.
  *
  * A sample implementation of MyClass::GenerateItemRow has to return a STableRow with optional content:
@@ -92,7 +91,6 @@ public:
 		, _OnGeneratePinnedRow()
 		, _OnEntryInitialized()
 		, _OnRowReleased()
-		, _ItemHeight(16)
 		, _MaxPinnedItems(6)
 		, _OnContextMenuOpening()
 		, _OnMouseButtonClick()
@@ -140,7 +138,7 @@ public:
 
 		SLATE_ITEMS_SOURCE_ARGUMENT( ItemType, ListItemsSource )
 
-		SLATE_ATTRIBUTE( float, ItemHeight )
+		SLATE_ATTRIBUTE_DEPRECATED( float, ItemHeight, 5.5, "The ItemHeight is only used for Tile. See ShouldArrangeAsTiles")
 
 		SLATE_ATTRIBUTE(int32, MaxPinnedItems)
 
@@ -284,8 +282,8 @@ public:
 		}
 		else
 		{
-			// Make the TableView
-			ConstructChildren( 0, InArgs._ItemHeight, EListItemAlignment::LeftAligned, InArgs._HeaderRow, InArgs._ExternalScrollbar, InArgs._Orientation, InArgs._OnListViewScrolled, InArgs._ScrollBarStyle, InArgs._PreventThrottling );
+			// Make the ListView
+			ConstructChildren( 0.0f, 0.0f, EListItemAlignment::LeftAligned, InArgs._HeaderRow, InArgs._ExternalScrollbar, InArgs._Orientation, InArgs._OnListViewScrolled, InArgs._ScrollBarStyle, InArgs._PreventThrottling );
 			if(this->ScrollBar.IsValid())
 			{
 				this->ScrollBar->SetDragFocusCause(InArgs._ScrollbarDragFocusCause);
diff --git a/Engine/Source/Runtime/Slate/Public/Widgets/Views/STreeView.h b/Engine/Source/Runtime/Slate/Public/Widgets/Views/STreeView.h
--- a/Engine/Source/Runtime/Slate/Public/Widgets/Views/STreeView.h
+++ b/Engine/Source/Runtime/Slate/Public/Widgets/Views/STreeView.h
@@ -119,7 +119,6 @@ public:
 		, _OnGeneratePinnedRow()
 		, _OnGetChildren()
 		, _OnSetExpansionRecursive()
-		, _ItemHeight(16)
 		, _MaxPinnedItems(6) // Having more than the max amount of items leads to the extra items in the middle being collapsed into ellipses, and the last item is fully shown
 		, _OnContextMenuOpening()
 		, _OnMouseButtonClick()
@@ -167,7 +166,7 @@ public:
 
 		SLATE_ITEMS_SOURCE_ARGUMENT( ItemType, TreeItemsSource )
 
-		SLATE_ATTRIBUTE( float, ItemHeight )
+		SLATE_ATTRIBUTE_DEPRECATED(float, ItemHeight, 5.5, "The ItemHeight is only used for Tile. See ShouldArrangeAsTiles")
 
 		SLATE_ATTRIBUTE( int32, MaxPinnedItems );
 
@@ -319,7 +318,7 @@ public:
 		else
 		{
 			// Make the TableView
-			this->ConstructChildren( 0, InArgs._ItemHeight, EListItemAlignment::LeftAligned, InArgs._HeaderRow, InArgs._ExternalScrollbar, Orient_Vertical, InArgs._OnTreeViewScrolled, InArgs._ScrollBarStyle, InArgs._PreventThrottling );
+			this->ConstructChildren( 0.0f, 0.0f, EListItemAlignment::LeftAligned, InArgs._HeaderRow, InArgs._ExternalScrollbar, Orient_Vertical, InArgs._OnTreeViewScrolled, InArgs._ScrollBarStyle, InArgs._PreventThrottling );
 			if (this->ScrollBar.IsValid())
 			{
 				this->ScrollBar->SetDragFocusCause(InArgs._ScrollbarDragFocusCause);
