[Backout] - CL33394076
Original CL Desc
-----------------------------------------------------------------
Slate: Deprecate SListView::ItemHeight and STreeViewItemHeight. ItemHeight and ItemWidth are only used when the panel is is as a tile or both value are assigned. With TreeView and ListView, you can't assigned both. The function are protected, a subclass of those could but it's up to them to expose the ItemHeight and ItemWidth. This behavior changed was introduced with CL 7634094.

diff --git a/Engine/Source/Runtime/Slate/Public/Widgets/Views/SListView.h b/Engine/Source/Runtime/Slate/Public/Widgets/Views/SListView.h
--- a/Engine/Source/Runtime/Slate/Public/Widgets/Views/SListView.h
+++ b/Engine/Source/Runtime/Slate/Public/Widgets/Views/SListView.h
@@ -43,10 +43,11 @@
  *   Given: TArray< TSharedPtr<FText> > Items;
  *
  *   SNew( SListView< TSharedPtr<FText> > )
+ *     .ItemHeight(24)
  *     .ListItemsSource( &Items )
  *     .OnGenerateRow(this, &MyClass::GenerateItemRow)
  *
- * The ListView will create widgets based on data items
+ * In the example we make all our widgets be 24 screen units tall. The ListView will create widgets based on data items
  * in the Items TArray. When the ListView needs to generate an item, it will do so using the specified OnGenerateRow method.
  *
  * A sample implementation of MyClass::GenerateItemRow has to return a STableRow with optional content:
@@ -91,6 +92,7 @@ public:
 		, _OnGeneratePinnedRow()
 		, _OnEntryInitialized()
 		, _OnRowReleased()
+		, _ItemHeight(16)
 		, _MaxPinnedItems(6)
 		, _OnContextMenuOpening()
 		, _OnMouseButtonClick()
@@ -138,7 +140,7 @@ public:
 
 		SLATE_ITEMS_SOURCE_ARGUMENT( ItemType, ListItemsSource )
 
-		SLATE_ATTRIBUTE_DEPRECATED( float, ItemHeight, 5.5, "The ItemHeight is only used for Tile. See ShouldArrangeAsTiles")
+		SLATE_ATTRIBUTE( float, ItemHeight )
 
 		SLATE_ATTRIBUTE(int32, MaxPinnedItems)
 
@@ -282,8 +284,8 @@ public:
 		}
 		else
 		{
-			// Make the ListView
-			ConstructChildren( 0.0f, 0.0f, EListItemAlignment::LeftAligned, InArgs._HeaderRow, InArgs._ExternalScrollbar, InArgs._Orientation, InArgs._OnListViewScrolled, InArgs._ScrollBarStyle, InArgs._PreventThrottling );
+			// Make the TableView
+			ConstructChildren( 0, InArgs._ItemHeight, EListItemAlignment::LeftAligned, InArgs._HeaderRow, InArgs._ExternalScrollbar, InArgs._Orientation, InArgs._OnListViewScrolled, InArgs._ScrollBarStyle, InArgs._PreventThrottling );
 			if(this->ScrollBar.IsValid())
 			{
 				this->ScrollBar->SetDragFocusCause(InArgs._ScrollbarDragFocusCause);
diff --git a/Engine/Source/Runtime/Slate/Public/Widgets/Views/STreeView.h b/Engine/Source/Runtime/Slate/Public/Widgets/Views/STreeView.h
--- a/Engine/Source/Runtime/Slate/Public/Widgets/Views/STreeView.h
+++ b/Engine/Source/Runtime/Slate/Public/Widgets/Views/STreeView.h
@@ -119,6 +119,7 @@ public:
 		, _OnGeneratePinnedRow()
 		, _OnGetChildren()
 		, _OnSetExpansionRecursive()
+		, _ItemHeight(16)
 		, _MaxPinnedItems(6) // Having more than the max amount of items leads to the extra items in the middle being collapsed into ellipses, and the last item is fully shown
 		, _OnContextMenuOpening()
 		, _OnMouseButtonClick()
@@ -166,7 +167,7 @@ public:
 
 		SLATE_ITEMS_SOURCE_ARGUMENT( ItemType, TreeItemsSource )
 
-		SLATE_ATTRIBUTE_DEPRECATED(float, ItemHeight, 5.5, "The ItemHeight is only used for Tile. See ShouldArrangeAsTiles")
+		SLATE_ATTRIBUTE( float, ItemHeight )
 
 		SLATE_ATTRIBUTE( int32, MaxPinnedItems );
 
@@ -318,7 +319,7 @@ public:
 		else
 		{
 			// Make the TableView
-			this->ConstructChildren( 0.0f, 0.0f, EListItemAlignment::LeftAligned, InArgs._HeaderRow, InArgs._ExternalScrollbar, Orient_Vertical, InArgs._OnTreeViewScrolled, InArgs._ScrollBarStyle, InArgs._PreventThrottling );
+			this->ConstructChildren( 0, InArgs._ItemHeight, EListItemAlignment::LeftAligned, InArgs._HeaderRow, InArgs._ExternalScrollbar, Orient_Vertical, InArgs._OnTreeViewScrolled, InArgs._ScrollBarStyle, InArgs._PreventThrottling );
 			if (this->ScrollBar.IsValid())
 			{
 				this->ScrollBar->SetDragFocusCause(InArgs._ScrollbarDragFocusCause);
