In SDockingTabWell::StartDraggingTab clear any pending layout save so that it doesn't occur during the drag operation.

A new save will be requested in FTabManager::OnTabRelocated when the drag ends.  This fixes an issue where a pending save that was scheduled by a previous call to OnTabRelocated occurs during a drag operation leading to a corrupt layout.

diff --git a/Engine/Source/Runtime/Slate/Private/Framework/Docking/SDockingTabWell.cpp b/Engine/Source/Runtime/Slate/Private/Framework/Docking/SDockingTabWell.cpp
--- a/Engine/Source/Runtime/Slate/Private/Framework/Docking/SDockingTabWell.cpp
+++ b/Engine/Source/Runtime/Slate/Private/Framework/Docking/SDockingTabWell.cpp
@@ -322,6 +322,10 @@ FReply SDockingTabWell::StartDraggingTab( TSharedRef<SDockTab> TabToStartDraggin
 
 	if (bCanLeaveTabWell)
 	{
+		// Clear any pending layout save so that it doesn't occur during the drag operation.
+		// A new save will be requested in FTabManager::OnTabRelocated when the drag ends.
+		TabManager->ClearPendingLayoutSave();
+
 		// We just removed the foreground tab.
 		ForegroundTabIndex = INDEX_NONE;
 		ParentTabStackPtr.Pin()->OnTabRemoved(TabToStartDragging->GetLayoutIdentifier());
