Slate: Remove deprecated code up to 5.0

diff --git a/Engine/Plugins/Developer/Concert/ConcertApp/MultiUserClient/Source/MultiUserClient/Private/Widgets/SConcertSandboxPersistWidget.cpp b/Engine/Plugins/Developer/Concert/ConcertApp/MultiUserClient/Source/MultiUserClient/Private/Widgets/SConcertSandboxPersistWidget.cpp
--- a/Engine/Plugins/Developer/Concert/ConcertApp/MultiUserClient/Source/MultiUserClient/Private/Widgets/SConcertSandboxPersistWidget.cpp
+++ b/Engine/Plugins/Developer/Concert/ConcertApp/MultiUserClient/Source/MultiUserClient/Private/Widgets/SConcertSandboxPersistWidget.cpp
@@ -183,7 +183,7 @@ void SConcertSandboxPersistWidget::Construct(const FArguments& InArgs)
 			.Padding(5)
 			[
 				SNew(SWrapBox)
-				.UseAllottedWidth(true)
+				.UseAllottedSize(true)
 				+SWrapBox::Slot()
 				.Padding(0.0f, 0.0f, 16.0f, 0.0f)
 				[
diff --git a/Engine/Plugins/Developer/TraceDataFiltering/Source/TraceDataFiltering/Private/SFilterPresetList.cpp b/Engine/Plugins/Developer/TraceDataFiltering/Source/TraceDataFiltering/Private/SFilterPresetList.cpp
--- a/Engine/Plugins/Developer/TraceDataFiltering/Source/TraceDataFiltering/Private/SFilterPresetList.cpp
+++ b/Engine/Plugins/Developer/TraceDataFiltering/Source/TraceDataFiltering/Private/SFilterPresetList.cpp
@@ -25,7 +25,7 @@ void SFilterPresetList::Construct( const FArguments& InArgs )
 	OnHighlightPreset = InArgs._OnHighlightPreset;
 
 	FilterBox = SNew(SWrapBox)
-		.UseAllottedWidth(true);
+		.UseAllottedSize(true);
 
 	LoadSettings(FTraceFilteringModule::TraceFiltersIni);
 	
diff --git a/Engine/Plugins/Developer/TraceSourceFiltering/Source/SourceFilteringEditor/Private/SClassTraceFilteringWidget.cpp b/Engine/Plugins/Developer/TraceSourceFiltering/Source/SourceFilteringEditor/Private/SClassTraceFilteringWidget.cpp
--- a/Engine/Plugins/Developer/TraceSourceFiltering/Source/SourceFilteringEditor/Private/SClassTraceFilteringWidget.cpp
+++ b/Engine/Plugins/Developer/TraceSourceFiltering/Source/SourceFilteringEditor/Private/SClassTraceFilteringWidget.cpp
@@ -26,7 +26,7 @@ void SClassTraceFilteringWidget::Construct(const FArguments& InArgs)
 		.Padding(10.0f, 0.f, 2.f, 0.f)
 		[
 			SAssignNew(ClassFiltersWrapBox, SWrapBox)
-			.UseAllottedWidth(true)
+			.UseAllottedSize(true)
 		]
 							
 		+ SVerticalBox::Slot()
diff --git a/Engine/Plugins/Developer/TraceSourceFiltering/Source/SourceFilteringEditor/Private/SFilterObjectWidget.cpp b/Engine/Plugins/Developer/TraceSourceFiltering/Source/SourceFilteringEditor/Private/SFilterObjectWidget.cpp
--- a/Engine/Plugins/Developer/TraceSourceFiltering/Source/SourceFilteringEditor/Private/SFilterObjectWidget.cpp
+++ b/Engine/Plugins/Developer/TraceSourceFiltering/Source/SourceFilteringEditor/Private/SFilterObjectWidget.cpp
@@ -10,7 +10,7 @@ void SFilterObjectRowWidget::Construct(const FArguments& InArgs, const TSharedRe
 	Object = InObject;
 
 	TSharedRef<SWrapBox> WrapBox = SNew(SWrapBox)
-		.UseAllottedWidth(true);
+		.UseAllottedSize(true);
 
 	TAttribute<bool> FilterEnabled = TAttribute<bool>::Create(TAttribute<bool>::FGetter::CreateSP(Object.Get(), &IFilterObject::IsFilterEnabled));
 	WrapBox->SetEnabled(FilterEnabled);
diff --git a/Engine/Plugins/Enterprise/DatasmithContent/Source/DatasmithContentEditor/Private/DatasmithSceneActorDetailsPanel.cpp b/Engine/Plugins/Enterprise/DatasmithContent/Source/DatasmithContentEditor/Private/DatasmithSceneActorDetailsPanel.cpp
--- a/Engine/Plugins/Enterprise/DatasmithContent/Source/DatasmithContentEditor/Private/DatasmithSceneActorDetailsPanel.cpp
+++ b/Engine/Plugins/Enterprise/DatasmithContent/Source/DatasmithContentEditor/Private/DatasmithSceneActorDetailsPanel.cpp
@@ -38,7 +38,7 @@ void FDatasmithSceneActorDetailsPanel::CustomizeDetails(IDetailLayoutBuilder& De
 {
 	SelectedObjectsList = DetailLayoutBuilder.GetSelectedObjects();
 
-	TSharedRef<SWrapBox> WrapBox = SNew(SWrapBox).UseAllottedWidth(true);
+	TSharedRef<SWrapBox> WrapBox = SNew(SWrapBox).UseAllottedSize(true);
 
 	FString CategoryName = TEXT("Datasmith");
 	IDetailCategoryBuilder& ActionsCategory = DetailLayoutBuilder.EditCategory(*CategoryName);
diff --git a/Engine/Plugins/Experimental/AnimNext/Source/AnimNextEditor/Private/Param/SLinkParametersDialog.cpp b/Engine/Plugins/Experimental/AnimNext/Source/AnimNextEditor/Private/Param/SLinkParametersDialog.cpp
--- a/Engine/Plugins/Experimental/AnimNext/Source/AnimNextEditor/Private/Param/SLinkParametersDialog.cpp
+++ b/Engine/Plugins/Experimental/AnimNext/Source/AnimNextEditor/Private/Param/SLinkParametersDialog.cpp
@@ -77,7 +77,7 @@ void SLinkParametersDialog::Construct(const FArguments& InArgs)
 					.Padding(5.0f)
 					[
 						SAssignNew(QueuedParametersBox, SWrapBox)
-						.UseAllottedWidth(true)
+						.UseAllottedSize(true)
 					]
 				]
 				+SVerticalBox::Slot()
diff --git a/Engine/Plugins/Experimental/DynamicMaterial/Source/DynamicMaterialEditor/Private/Slate/SDMSlot.cpp b/Engine/Plugins/Experimental/DynamicMaterial/Source/DynamicMaterialEditor/Private/Slate/SDMSlot.cpp
--- a/Engine/Plugins/Experimental/DynamicMaterial/Source/DynamicMaterialEditor/Private/Slate/SDMSlot.cpp
+++ b/Engine/Plugins/Experimental/DynamicMaterial/Source/DynamicMaterialEditor/Private/Slate/SDMSlot.cpp
@@ -477,8 +477,7 @@ TSharedRef<SWidget> SDMSlot::CreateHeaderPropertyListWidget()
 	TSharedRef<SWrapBox> PropertyContainer = SNew(SWrapBox)
 		.InnerSlotPadding(FVector2D(2.f, 2.f))
 		.Orientation(EOrientation::Orient_Horizontal)
-		.UseAllottedSize(true)
-		.UseAllottedWidth(true);
+		.UseAllottedSize(true);
 
 	const TArray<EDMMaterialPropertyType> Properties = ModelEditorOnlyData->GetMaterialPropertiesForSlot(Slot);
 
diff --git a/Engine/Plugins/FX/Niagara/Source/NiagaraEditor/Private/Widgets/SNiagaraHierarchy.cpp b/Engine/Plugins/FX/Niagara/Source/NiagaraEditor/Private/Widgets/SNiagaraHierarchy.cpp
--- a/Engine/Plugins/FX/Niagara/Source/NiagaraEditor/Private/Widgets/SNiagaraHierarchy.cpp
+++ b/Engine/Plugins/FX/Niagara/Source/NiagaraEditor/Private/Widgets/SNiagaraHierarchy.cpp
@@ -643,7 +643,7 @@ void SNiagaraHierarchy::Construct(const FArguments& InArgs, TObjectPtr<UNiagaraH
 				.Padding(2.f)
 				[
 					SAssignNew(SourceSectionBox, SWrapBox)
-					.UseAllottedWidth(true)
+					.UseAllottedSize(true)
 				]
 				+ SVerticalBox::Slot()
 				.Padding(1.f, 2.f)
@@ -681,7 +681,7 @@ void SNiagaraHierarchy::Construct(const FArguments& InArgs, TObjectPtr<UNiagaraH
 					.Padding(1.f)
 					[
 						SAssignNew(HierarchySectionBox, SWrapBox)
-						.UseAllottedWidth(true)
+						.UseAllottedSize(true)
 					]				
 				]
 				+ SVerticalBox::Slot()
diff --git a/Engine/Plugins/FX/Niagara/Source/NiagaraEditorWidgets/Private/Stack/SNiagaraStackRowPerfWidget.cpp b/Engine/Plugins/FX/Niagara/Source/NiagaraEditorWidgets/Private/Stack/SNiagaraStackRowPerfWidget.cpp
--- a/Engine/Plugins/FX/Niagara/Source/NiagaraEditorWidgets/Private/Stack/SNiagaraStackRowPerfWidget.cpp
+++ b/Engine/Plugins/FX/Niagara/Source/NiagaraEditorWidgets/Private/Stack/SNiagaraStackRowPerfWidget.cpp
@@ -72,7 +72,7 @@ void SNiagaraStackRowPerfWidget::Construct(const FArguments& InArgs, UNiagaraSta
             .VAlign(VAlign_Bottom)
             [
 				SNew(SWrapBox)
-				.UseAllottedWidth(true)
+				.UseAllottedSize(true)
 			    +SWrapBox::Slot()
                 [
                     // Placeholder brush to fill the remaining space
diff --git a/Engine/Plugins/VirtualProduction/DMX/DMXEngine/Source/DMXEditor/Private/Customizations/DMXLibraryPortReferencesCustomization.cpp b/Engine/Plugins/VirtualProduction/DMX/DMXEngine/Source/DMXEditor/Private/Customizations/DMXLibraryPortReferencesCustomization.cpp
--- a/Engine/Plugins/VirtualProduction/DMX/DMXEngine/Source/DMXEditor/Private/Customizations/DMXLibraryPortReferencesCustomization.cpp
+++ b/Engine/Plugins/VirtualProduction/DMX/DMXEngine/Source/DMXEditor/Private/Customizations/DMXLibraryPortReferencesCustomization.cpp
@@ -379,7 +379,7 @@ TSharedRef<SWidget> FDMXLibraryPortReferencesCustomization::GeneratePortInfoWidg
 			return
 				SNew(SWrapBox)
 				.InnerSlotPadding(FVector2D(4.f, 4.f))
-				.UseAllottedWidth(true)
+				.UseAllottedSize(true)
 
 				// Protocol Name
 				+ SWrapBox::Slot()
diff --git a/Engine/Plugins/VirtualProduction/DMX/DMXEngine/Source/DMXEditor/Private/Widgets/FixturePatch/SDMXMVRFixtureListToolbar.cpp b/Engine/Plugins/VirtualProduction/DMX/DMXEngine/Source/DMXEditor/Private/Widgets/FixturePatch/SDMXMVRFixtureListToolbar.cpp
--- a/Engine/Plugins/VirtualProduction/DMX/DMXEngine/Source/DMXEditor/Private/Widgets/FixturePatch/SDMXMVRFixtureListToolbar.cpp
+++ b/Engine/Plugins/VirtualProduction/DMX/DMXEngine/Source/DMXEditor/Private/Widgets/FixturePatch/SDMXMVRFixtureListToolbar.cpp
@@ -41,7 +41,7 @@ void SDMXMVRFixtureListToolbar::Construct(const FArguments& InArgs, TWeakPtr<FDM
 			[
 				SNew(SWrapBox)
 				.InnerSlotPadding(FVector2D(14.f, 8.f))
-				.UseAllottedWidth(true)
+				.UseAllottedSize(true)
 
 				// Add Fixture Button
 				+ SWrapBox::Slot()
diff --git a/Engine/Plugins/VirtualProduction/DMX/DMXEngine/Source/DMXEditor/Private/Widgets/Monitors/SDMXActivityMonitor.cpp b/Engine/Plugins/VirtualProduction/DMX/DMXEngine/Source/DMXEditor/Private/Widgets/Monitors/SDMXActivityMonitor.cpp
--- a/Engine/Plugins/VirtualProduction/DMX/DMXEngine/Source/DMXEditor/Private/Widgets/Monitors/SDMXActivityMonitor.cpp
+++ b/Engine/Plugins/VirtualProduction/DMX/DMXEngine/Source/DMXEditor/Private/Widgets/Monitors/SDMXActivityMonitor.cpp
@@ -65,7 +65,7 @@ void SDMXActivityMonitor::Construct(const FArguments& InArgs)
 			[
 				SNew(SWrapBox)
 				.InnerSlotPadding(FVector2D(35.0f, 10.0f))
-				.UseAllottedWidth(true)
+				.UseAllottedSize(true)
 
 				// Source Selector
 				+ SWrapBox::Slot()
diff --git a/Engine/Plugins/VirtualProduction/DMX/DMXEngine/Source/DMXEditor/Private/Widgets/Monitors/SDMXChannelsMonitor.cpp b/Engine/Plugins/VirtualProduction/DMX/DMXEngine/Source/DMXEditor/Private/Widgets/Monitors/SDMXChannelsMonitor.cpp
--- a/Engine/Plugins/VirtualProduction/DMX/DMXEngine/Source/DMXEditor/Private/Widgets/Monitors/SDMXChannelsMonitor.cpp
+++ b/Engine/Plugins/VirtualProduction/DMX/DMXEngine/Source/DMXEditor/Private/Widgets/Monitors/SDMXChannelsMonitor.cpp
@@ -44,7 +44,7 @@ void SDMXChannelsMonitor::Construct(const FArguments& InArgs)
 			[
 				SNew(SWrapBox)
 				.InnerSlotPadding(FVector2D(36.0f, 10.0f))
-				.UseAllottedWidth(true)
+				.UseAllottedSize(true)
 		
 				// Source Selector
 				+ SWrapBox::Slot()
@@ -121,7 +121,7 @@ void SDMXChannelsMonitor::Construct(const FArguments& InArgs)
 				.HAlign(HAlign_Fill)
 				[
 					SAssignNew(ChannelValuesBox, SWrapBox)
-					.UseAllottedWidth(true)
+					.UseAllottedSize(true)
 					.InnerSlotPadding(FVector2D(1.0f, 1.0f))
 				]
 			]
diff --git a/Engine/Plugins/VirtualProduction/DMX/DMXPixelMapping/Source/DMXPixelMappingEditor/Private/Views/SDMXPixelMappingDMXLibraryView.cpp b/Engine/Plugins/VirtualProduction/DMX/DMXPixelMapping/Source/DMXPixelMappingEditor/Private/Views/SDMXPixelMappingDMXLibraryView.cpp
--- a/Engine/Plugins/VirtualProduction/DMX/DMXPixelMapping/Source/DMXPixelMappingEditor/Private/Views/SDMXPixelMappingDMXLibraryView.cpp
+++ b/Engine/Plugins/VirtualProduction/DMX/DMXPixelMapping/Source/DMXPixelMappingEditor/Private/Views/SDMXPixelMappingDMXLibraryView.cpp
@@ -140,7 +140,7 @@ void SDMXPixelMappingDMXLibraryView::Construct(const FArguments& InArgs, const T
 			.AutoHeight()
 			[
 				SAssignNew(AddPatchesWrapBox, SWrapBox)
-				.UseAllottedWidth(true)
+				.UseAllottedSize(true)
 
 				+ SWrapBox::Slot()
 				.Padding(4.f)
diff --git a/Engine/Plugins/VirtualProduction/RemoteControl/Source/RemoteControlUI/Private/UI/Customizations/NetworkAddressCustomization.cpp b/Engine/Plugins/VirtualProduction/RemoteControl/Source/RemoteControlUI/Private/UI/Customizations/NetworkAddressCustomization.cpp
--- a/Engine/Plugins/VirtualProduction/RemoteControl/Source/RemoteControlUI/Private/UI/Customizations/NetworkAddressCustomization.cpp
+++ b/Engine/Plugins/VirtualProduction/RemoteControl/Source/RemoteControlUI/Private/UI/Customizations/NetworkAddressCustomization.cpp
@@ -55,7 +55,7 @@ void FNetworkAddressCustomization::CustomizeHeader(TSharedRef<IPropertyHandle> S
 		.VAlign(VAlign_Center)
 		[
 			SNew(SWrapBox)
-			.UseAllottedWidth(bUseAllottedWidth)
+			.UseAllottedSize(bUseAllottedWidth)
 
 			+ SWrapBox::Slot()
 			.Padding(2.f, 0.f)
diff --git a/Engine/Source/Editor/EnvironmentLightingViewer/Private/SEnvironmentLightingViewer.cpp b/Engine/Source/Editor/EnvironmentLightingViewer/Private/SEnvironmentLightingViewer.cpp
--- a/Engine/Source/Editor/EnvironmentLightingViewer/Private/SEnvironmentLightingViewer.cpp
+++ b/Engine/Source/Editor/EnvironmentLightingViewer/Private/SEnvironmentLightingViewer.cpp
@@ -183,8 +183,8 @@ void SEnvironmentLightingViewer::Construct(const FArguments& InArgs)
 			+ SScrollBox::Slot()
 			[
 				SNew(SWrapBox)
-				.UseAllottedWidth(true)
-				.PreferredWidth(384)
+				.UseAllottedSize(true)
+				.PreferredSize(384.0f)
 				+SWrapBox::Slot()
 				.FillLineWhenSizeLessThan(384)
 				.Padding(5.0f, 20.0f, 5.0f, 20.0f)
diff --git a/Engine/Source/Editor/NNEEditor/Private/NNEEditorRuntimesWidget.cpp b/Engine/Source/Editor/NNEEditor/Private/NNEEditorRuntimesWidget.cpp
--- a/Engine/Source/Editor/NNEEditor/Private/NNEEditorRuntimesWidget.cpp
+++ b/Engine/Source/Editor/NNEEditor/Private/NNEEditorRuntimesWidget.cpp
@@ -48,7 +48,7 @@ void SRuntimesWidget::Construct(const FArguments& InArgs)
 			.Padding(4.0)
 			[
 				SAssignNew(WrapBox, SWrapBox)
-				.PreferredWidth(300.0f)
+				.PreferredSize(300.0f)
 			]
 		]
 	];
diff --git a/Engine/Source/Editor/Sequencer/Private/SSequencerDebugVisualizer.cpp b/Engine/Source/Editor/Sequencer/Private/SSequencerDebugVisualizer.cpp
--- a/Engine/Source/Editor/Sequencer/Private/SSequencerDebugVisualizer.cpp
+++ b/Engine/Source/Editor/Sequencer/Private/SSequencerDebugVisualizer.cpp
@@ -342,7 +342,7 @@ void SSequencerEntityComponentSystemDebugSlot::Construct(const FArguments& InArg
 	SetBorderImage(SectionBackgroundBrush);
 
 	TSharedRef<SWrapBox> Container = SNew(SWrapBox)
-		.UseAllottedWidth(true)
+		.UseAllottedSize(true)
 		+SWrapBox::Slot()
 		[
 			SNew(SBorder)
diff --git a/Engine/Source/Runtime/Slate/Private/Widgets/Layout/SRadialBox.cpp b/Engine/Source/Runtime/Slate/Private/Widgets/Layout/SRadialBox.cpp
--- a/Engine/Source/Runtime/Slate/Private/Widgets/Layout/SRadialBox.cpp
+++ b/Engine/Source/Runtime/Slate/Private/Widgets/Layout/SRadialBox.cpp
@@ -160,12 +160,6 @@ void SRadialBox::SetUseAllottedWidth(bool bInUseAllottedWidth)
 	bUseAllottedWidth = bInUseAllottedWidth;
 }
 
-int32 SRadialBox::NormalizeAngle(int32 Angle) const
-{
-	int32 NormalizedAngle = Angle % 360;
-	return NormalizedAngle < 0 ? NormalizedAngle + 360 : NormalizedAngle;
-}
-
 float SRadialBox::InternalNormalizeAngle(float Angle) const
 {
 	float NormalizedAngle = FMath::Fmod(Angle, 360.f);
diff --git a/Engine/Source/Runtime/Slate/Private/Widgets/Layout/SSplitter.cpp b/Engine/Source/Runtime/Slate/Private/Widgets/Layout/SSplitter.cpp
--- a/Engine/Source/Runtime/Slate/Private/Widgets/Layout/SSplitter.cpp
+++ b/Engine/Source/Runtime/Slate/Private/Widgets/Layout/SSplitter.cpp
@@ -5,7 +5,6 @@
 #include "Rendering/DrawElements.h"
 #include "Framework/Application/SlateApplication.h"
 
-PRAGMA_DISABLE_DEPRECATION_WARNINGS
 void SSplitter::FSlot::Construct(const FChildren& SlotOwner, FSlotArguments&& InArgs)
 {
 	TSlotBase<FSlot>::Construct(SlotOwner, MoveTemp(InArgs));
@@ -131,14 +130,14 @@ TArray<FLayoutGeometry> SSplitter::ArrangeChildrenForLayout(const FGeometry& All
 
 		++NumNonCollapsedChildren;
 
-		if (CurSlot.SizingRule.Get() == SSplitter::SizeToContent)
+		if (CurSlot.GetSizingRule() == SSplitter::SizeToContent)
 		{
 			NonResizableSpace += CurSlot.GetWidget()->GetDesiredSize()[AxisIndex];
 		}
 		else // SizingRule == SSplitter::FractionOfParent
 		{
 			MinResizableSpace += FMath::Max(MinSplitterChildLength, CurSlot.GetMinSize());
-			CoefficientTotal += CurSlot.SizeValue.Get();
+			CoefficientTotal += CurSlot.GetSizeValue();
 		}
 	}
 
@@ -165,13 +164,13 @@ TArray<FLayoutGeometry> SSplitter::ArrangeChildrenForLayout(const FGeometry& All
 
 		float ChildSpace;
 		
-		if (CurSlot.SizingRule.Get() == SSplitter::SizeToContent)
+		if (CurSlot.GetSizingRule() == SSplitter::SizeToContent)
 		{
 			ChildSpace = CurSlot.GetWidget()->GetDesiredSize()[AxisIndex];
 		}
 		else
 		{
-			ChildSpace = ResizableSpace * CurSlot.SizeValue.Get() / CoefficientTotal  - ExtraRequiredSpace;
+			ChildSpace = ResizableSpace * CurSlot.GetSizeValue() / CoefficientTotal  - ExtraRequiredSpace;
 			ExtraRequiredSpace = 0;
 		}
 
@@ -184,7 +183,7 @@ TArray<FLayoutGeometry> SSplitter::ArrangeChildrenForLayout(const FGeometry& All
 
 			for (int32 PrevIndex = ChildIndex - 1; PrevIndex >= 0 && CurrentRequiredSpace > 0; --PrevIndex)
 			{
-				if (Children[PrevIndex].SizingRule.Get() == SSplitter::FractionOfParent && !IsSlotCollapsed(CurSlot))
+				if (Children[PrevIndex].GetSizingRule() == SSplitter::FractionOfParent && !IsSlotCollapsed(CurSlot))
 				{
 					const float MinChildSize = FMath::Max(MinSplitterChildLength, Children[PrevIndex].GetMinSize());
 					const float AvailableSpace = SlotSizes[PrevIndex] - MinChildSize;
@@ -280,8 +279,8 @@ int32 SSplitter::OnPaint( const FPaintArgs& Args, const FGeometry& AllottedGeome
 		const float HalfHitDetectionSplitterHandleSize = FMath::RoundToFloat(HitDetectionSplitterHandleSize / 2.0f);
 		const float HalfPhysicalSplitterHandleSize = PhysicalSplitterHandleSize;// FMath::RoundToInt(PhysicalSplitterHandleSize / 2.0f);
 
-		FVector2D HandleSize;
-		FVector2D HandlePosition;
+		FVector2f HandleSize;
+		FVector2f HandlePosition;
 		if (Orientation == Orient_Horizontal)
 		{
 			HandleSize.Set( PhysicalSplitterHandleSize, GeometryAfterSplitter.Size.Y );
@@ -299,7 +298,7 @@ int32 SSplitter::OnPaint( const FPaintArgs& Args, const FGeometry& AllottedGeome
 			FSlateDrawElement::MakeBox(
 				OutDrawElements,
 				MaxLayerId,
-				GeometryAfterSplitter.ToPaintGeometry( HandlePosition, HandleSize, 1.0f ),
+				GeometryAfterSplitter.ToPaintGeometry(HandleSize, FSlateLayoutTransform(1.0f, TransformPoint(1.0f, HandlePosition))),
 				NormalHandleBrush,
 				ShouldBeEnabled(bParentEnabled) ? ESlateDrawEffect::None : ESlateDrawEffect::DisabledEffect,
 				InWidgetStyle.GetColorAndOpacityTint() * Style->HandleHighlightBrush.TintColor.GetSpecifiedColor()
@@ -310,7 +309,7 @@ int32 SSplitter::OnPaint( const FPaintArgs& Args, const FGeometry& AllottedGeome
 			FSlateDrawElement::MakeBox(
 				OutDrawElements,
 				MaxLayerId,
-				GeometryAfterSplitter.ToPaintGeometry( HandlePosition, HandleSize, 1.0f ),
+				GeometryAfterSplitter.ToPaintGeometry(HandleSize, FSlateLayoutTransform(1.0f, TransformPoint(1.0f, HandlePosition))),
 				&Style->HandleHighlightBrush,
 				ShouldBeEnabled(bParentEnabled) ? ESlateDrawEffect::None : ESlateDrawEffect::DisabledEffect,
 				InWidgetStyle.GetColorAndOpacityTint() * NormalHandleBrush->TintColor.GetSpecifiedColor()
@@ -724,41 +723,41 @@ void SSplitter::HandleResizingDelta(EOrientation SplitterOrientation, const floa
 			float TotalStretchLength = 0.f;
 			float TotalStretchCoefficients = 0.f;
 			
-			if (PrevChild.SizingRule.Get() == ESizeRule::FractionOfParent)
+			if (PrevChild.GetSizingRule() == ESizeRule::FractionOfParent)
 			{
 				TotalStretchLength = NewPrevChildLength;
-				TotalStretchCoefficients = PrevChild.SizeValue.Get();
+				TotalStretchCoefficients = PrevChild.GetSizeValue();
 			}
 
 			for (int32 SlotIndex = 0; SlotIndex < NumSlotsAfterDragHandle; SlotIndex++)
 			{
 				const FSlotInfo& SlotInfo = SlotsAfterDragHandle[SlotIndex];
 
-				if (SlotInfo.Slot->SizingRule.Get() == ESizeRule::FractionOfParent)
+				if (SlotInfo.Slot->GetSizingRule() == ESizeRule::FractionOfParent)
 				{
 					TotalStretchLength += SlotInfo.NewSize;
-					TotalStretchCoefficients += SlotInfo.Slot->SizeValue.Get();
+					TotalStretchCoefficients += SlotInfo.Slot->GetSizeValue();
 				}
 			}
 
 			auto SetSize = [](FSlot& Slot, float TotalStretchCoefficients, float NewLength, float TotalLength)
 			{
-				if (Slot.SizingRule.Get() == ESizeRule::FractionOfParent)
+				if (Slot.GetSizingRule() == ESizeRule::FractionOfParent)
 				{
 					float NewFillSize = TotalLength > 0.f ? (TotalStretchCoefficients * (NewLength / TotalLength)) : TotalStretchCoefficients;
-					if (Slot.OnSlotResized_Handler.IsBound())
+					if (Slot.OnSlotResized().IsBound())
 					{
-						Slot.OnSlotResized_Handler.Execute(NewFillSize);
+						Slot.OnSlotResized().Execute(NewFillSize);
 					}
 					else
 					{
-						Slot.SizeValue = NewFillSize;
+						Slot.SetSizeValue(NewFillSize);
 					}
 				}
 				else
 				{
-					ensure(Slot.OnSlotResized_Handler.IsBound());
-					Slot.OnSlotResized_Handler.ExecuteIfBound(NewLength);
+					ensure(Slot.OnSlotResized().IsBound());
+					Slot.OnSlotResized().ExecuteIfBound(NewLength);
 				}
 			};
 
@@ -828,7 +827,6 @@ int32 SSplitter::GetHandleBeingResizedFromMousePosition( float InPhysicalSplitte
 
 	return INDEX_NONE;
 }
-PRAGMA_ENABLE_DEPRECATION_WARNINGS
 
 /********************************************************************
 * SSplitter2x2														*
diff --git a/Engine/Source/Runtime/Slate/Private/Widgets/Layout/SWrapBox.cpp b/Engine/Source/Runtime/Slate/Private/Widgets/Layout/SWrapBox.cpp
--- a/Engine/Source/Runtime/Slate/Private/Widgets/Layout/SWrapBox.cpp
+++ b/Engine/Source/Runtime/Slate/Private/Widgets/Layout/SWrapBox.cpp
@@ -30,19 +30,12 @@ int32 SWrapBox::RemoveSlot( const TSharedRef<SWidget>& SlotWidget )
 	return Slots.Remove(SlotWidget);
 }
 
-PRAGMA_DISABLE_DEPRECATION_WARNINGS
 void SWrapBox::Construct( const FArguments& InArgs )
 {
 	PreferredSize.Assign(*this, InArgs._PreferredSize);
 
-	// Handle deprecation of PreferredWidth
-	if (!InArgs._PreferredSize.IsSet() && !InArgs._PreferredSize.IsBound())
-	{
-		PreferredSize.Assign(*this, InArgs._PreferredWidth);
-	}
-
 	InnerSlotPadding = InArgs._InnerSlotPadding;
-	bUseAllottedSize = InArgs._UseAllottedSize || InArgs._UseAllottedWidth;
+	bUseAllottedSize = InArgs._UseAllottedSize;
 	Orientation = InArgs._Orientation;
 	HAlign.Assign(*this, InArgs._HAlign);
 
@@ -51,7 +44,6 @@ void SWrapBox::Construct( const FArguments& InArgs )
 
 	SetCanTick(bUseAllottedSize);
 }
-PRAGMA_ENABLE_DEPRECATION_WARNINGS
 
 void SWrapBox::Tick( const FGeometry& AllottedGeometry, const double InCurrentTime, const float InDeltaTime )
 {
@@ -442,21 +434,11 @@ void SWrapBox::SetInnerSlotPadding(FVector2D InInnerSlotPadding)
 	}
 }
 
-void SWrapBox::SetWrapWidth(TAttribute<float> InWrapWidth)
-{
-	SetWrapSize(MoveTemp(InWrapWidth));
-}
-
 void SWrapBox::SetWrapSize(TAttribute<float> InWrapSize)
 {
 	PreferredSize.Assign(*this, MoveTemp(InWrapSize));
 }
 
-void SWrapBox::SetUseAllottedWidth(bool bInUseAllottedWidth)
-{
-	SetUseAllottedSize(bInUseAllottedWidth);
-}
-
 void SWrapBox::SetUseAllottedSize(bool bInUseAllottedSize)
 {
 	if (bUseAllottedSize != bInUseAllottedSize)
diff --git a/Engine/Source/Runtime/Slate/Private/Widgets/Views/SListPanel.h b/Engine/Source/Runtime/Slate/Private/Widgets/Views/SListPanel.h
--- a/Engine/Source/Runtime/Slate/Private/Widgets/Views/SListPanel.h
+++ b/Engine/Source/Runtime/Slate/Private/Widgets/Views/SListPanel.h
@@ -121,21 +121,6 @@ protected:
 
 	/** @return true if this panel should arrange items as tiles placed alongside one another in each line */
 	SLATE_API bool ShouldArrangeAsTiles() const;
-	
-protected:
-#if WITH_EDITORONLY_DATA
-	UE_DEPRECATED(5.5, "Direct access to ItemWidth is now deprecated. Use the setter or getter.")
-	FInvalidatableBrushAttribute ItemWidth;
-	
-	UE_DEPRECATED(5.5, "Direct access to ItemHeight is now deprecated. Use the setter or getter.")
-	FInvalidatableBrushAttribute ItemHeight;
-	
-	UE_DEPRECATED(5.5, "Direct access to NumDesiredItems is now deprecated. Use the setter or getter.")
-	FInvalidatableBrushAttribute NumDesiredItems;
-	
-	UE_DEPRECATED(5.5, "Direct access to ItemAlignment is now deprecated. Use the setter or getter.")
-	FInvalidatableBrushAttribute ItemAlignment;
-#endif
 
 protected:
 	/** The children being arranged by this panel */
diff --git a/Engine/Source/Runtime/Slate/Public/Widgets/Colors/SColorBlock.h b/Engine/Source/Runtime/Slate/Public/Widgets/Colors/SColorBlock.h
--- a/Engine/Source/Runtime/Slate/Public/Widgets/Colors/SColorBlock.h
+++ b/Engine/Source/Runtime/Slate/Public/Widgets/Colors/SColorBlock.h
@@ -67,16 +67,6 @@ public:
 		/** A handler to activate when the mouse is pressed. */
 		SLATE_EVENT(FPointerEventHandler, OnMouseButtonDown)
 
-		UE_DEPRECATED(5.0, "IgnoreAlpha is deprecated. Set AlphaDisplayMode to EColorBlockAlphaDisplayMode::Ignore instead")
-		FArguments& IgnoreAlpha(bool bInIgnoreAlpha)
-		{
-			if (bInIgnoreAlpha)
-			{
-				_AlphaDisplayMode = EColorBlockAlphaDisplayMode::Ignore;
-			}
-			return Me();
-		}
-
 	SLATE_END_ARGS()
 
 public:
diff --git a/Engine/Source/Runtime/Slate/Public/Widgets/Colors/SColorGradingWheel.h b/Engine/Source/Runtime/Slate/Public/Widgets/Colors/SColorGradingWheel.h
--- a/Engine/Source/Runtime/Slate/Public/Widgets/Colors/SColorGradingWheel.h
+++ b/Engine/Source/Runtime/Slate/Public/Widgets/Colors/SColorGradingWheel.h
@@ -119,15 +119,6 @@ protected:
 	/** Invoked when a new value is selected on the color wheel. */
 	FOnColorGradingWheelValueChanged OnValueChanged;
 
-#if WITH_EDITORONLY_DATA
-	UE_DEPRECATED(5.0, "Direct access to SelectedColor is now deprecated. Use the setter or getter.")
-	TSlateDeprecatedTAttribute<FLinearColor> SelectedColor;
-	UE_DEPRECATED(5.0, "Direct access to DesiredWheelSize is now deprecated. Use the setter or getter.")
-	TSlateDeprecatedTAttribute<int32> DesiredWheelSize;
-	UE_DEPRECATED(5.0, "Direct access to DesiredSizeOverride is now deprecated. Use the setter or getter.")
-	TSlateDeprecatedTAttribute<float> ExponentDisplacement;
-#endif
-
 private:
 	/** The current color selected by the user. */
 	TSlateAttribute<FLinearColor> SelectedColorAttribute;
diff --git a/Engine/Source/Runtime/Slate/Public/Widgets/Colors/SSimpleGradient.h b/Engine/Source/Runtime/Slate/Public/Widgets/Colors/SSimpleGradient.h
--- a/Engine/Source/Runtime/Slate/Public/Widgets/Colors/SSimpleGradient.h
+++ b/Engine/Source/Runtime/Slate/Public/Widgets/Colors/SSimpleGradient.h
@@ -37,12 +37,6 @@ public:
 		/** Horizontal or vertical gradient */
 		SLATE_ATTRIBUTE(EOrientation, Orientation)
 
-		/** Whether to display sRGB color */
-		PRAGMA_DISABLE_DEPRECATION_WARNINGS
-		UE_DEPRECATED(5.0, "UseSRGB is unused")
-		SLATE_ARGUMENT_DEFAULT(bool, UseSRGB) = true;
-		PRAGMA_ENABLE_DEPRECATION_WARNINGS
-
 	SLATE_END_ARGS()
 
 public:
diff --git a/Engine/Source/Runtime/Slate/Public/Widgets/Input/SButton.h b/Engine/Source/Runtime/Slate/Public/Widgets/Input/SButton.h
--- a/Engine/Source/Runtime/Slate/Public/Widgets/Input/SButton.h
+++ b/Engine/Source/Runtime/Slate/Public/Widgets/Input/SButton.h
@@ -124,15 +124,6 @@ protected:
 		SLATE_API SButton();
 
 public:
-	/**
-	 * @return An image that represents this button's border
-	 */
-	UE_DEPRECATED(5.0, "GetBorder is deprecated. Use SetBorderImage or GetBorderImage")
-	virtual const FSlateBrush* GetBorder() const
-	{
-		return GetBorderImage();
-	}
-
 	/** @return the Foreground color that this widget sets; unset options if the widget does not set a foreground color */
 	virtual FSlateColor GetForegroundColor() const final
 	{
@@ -273,30 +264,6 @@ private:
 	SLATE_API void UpdateForegroundColor();
 	SLATE_API void UpdateDisabledForegroundColor();
 
-protected:
-#if WITH_EDITORONLY_DATA
-	UE_DEPRECATED(5.0, "Direct access to ContentPadding is now deprecated. Use the setter or getter.")
-	TSlateDeprecatedTAttribute<FMargin> ContentPadding;
-	/** Brush resource that represents a button */
-	UE_DEPRECATED(5.0, "NormalImage is now deprecated. Use the ButtonStyle.")
-	const FSlateBrush* NormalImage;
-	/** Brush resource that represents a button when it is hovered */
-	UE_DEPRECATED(5.0, "HoverImage is now deprecated. Use the ButtonStyle.")
-	const FSlateBrush* HoverImage;
-	/** Brush resource that represents a button when it is pressed */
-	UE_DEPRECATED(5.0, "PressedImage is now deprecated. Use the ButtonStyle.")
-	const FSlateBrush* PressedImage;
-	/** Brush resource that represents a button when it is disabled */
-	UE_DEPRECATED(5.0, "DisabledImage is now deprecated. Use the ButtonStyle.")
-	const FSlateBrush* DisabledImage;
-	/** Padding that accounts for the button border */
-	UE_DEPRECATED(5.0, "BorderPadding is now deprecated. Use the ButtonStyle.")
-	FMargin BorderPadding;
-	/** Padding that accounts for the button border when pressed */
-	UE_DEPRECATED(5.0, "PressedBorderPadding is now deprecated. Use the ButtonStyle.")
-	FMargin PressedBorderPadding;
-#endif
-
 private:
 	/** The location in screenspace the button was pressed */
 	FVector2D PressedScreenSpacePosition;
diff --git a/Engine/Source/Runtime/Slate/Public/Widgets/Input/SRotatorInputBox.h b/Engine/Source/Runtime/Slate/Public/Widgets/Input/SRotatorInputBox.h
--- a/Engine/Source/Runtime/Slate/Public/Widgets/Input/SRotatorInputBox.h
+++ b/Engine/Source/Runtime/Slate/Public/Widgets/Input/SRotatorInputBox.h
@@ -53,13 +53,7 @@ public:
 		SLATE_ATTRIBUTE( TOptional<NumericType>, Yaw )
 
 		/** Should the axis labels be colored */
-		SLATE_ARGUMENT( bool, bColorAxisLabels )		
-
-		UE_DEPRECATED(5.0, "AllowResponsiveLayout unused as it is no longer necessary.")
-		FArguments& AllowResponsiveLayout(bool bAllow)
-		{
-			return TSlateBaseNamedArgs<SNumericRotatorInputBox<NumericType>>::Me();
-		}
+		SLATE_ARGUMENT( bool, bColorAxisLabels )
 
 		/** Font to use for the text in this box */
 		SLATE_ATTRIBUTE( FSlateFontInfo, Font )
diff --git a/Engine/Source/Runtime/Slate/Public/Widgets/Input/SSlider.h b/Engine/Source/Runtime/Slate/Public/Widgets/Input/SSlider.h
--- a/Engine/Source/Runtime/Slate/Public/Widgets/Input/SSlider.h
+++ b/Engine/Source/Runtime/Slate/Public/Widgets/Input/SSlider.h
@@ -206,19 +206,6 @@ protected:
 	SLATE_API const FSlateBrush* GetThumbImage() const;
 
 protected:
-#if WITH_EDITORONLY_DATA
-	UE_DEPRECATED(5.1, "Direct access to Value is now deprecated. Use the setter or getter.")
-	TSlateDeprecatedTAttribute<float> ValueAttribute;
-	UE_DEPRECATED(5.1, "Direct access to IndentHandle is now deprecated. Use the setter or getter.")
-	TSlateDeprecatedTAttribute<bool> IndentHandle;
-	UE_DEPRECATED(5.1, "Direct access to LockedAttribute is now deprecated. Use the setter or getter.")
-	TSlateDeprecatedTAttribute<bool> LockedAttribute;
-	UE_DEPRECATED(5.1, "Direct access to SliderBarColor is now deprecated. Use the setter or getter.")
-	TSlateDeprecatedTAttribute<FSlateColor> SliderBarColor;
-	UE_DEPRECATED(5.1, "Direct access to SliderHandleColor is now deprecated. Use the setter or getter.")
-	TSlateDeprecatedTAttribute<FSlateColor> SliderHandleColor;
-#endif
-
 	/** @return an attribute reference of IndentHandle */
 	TSlateAttributeRef<float> GetValueAttribute() const
 	{
diff --git a/Engine/Source/Runtime/Slate/Public/Widgets/Input/SVectorInputBox.h b/Engine/Source/Runtime/Slate/Public/Widgets/Input/SVectorInputBox.h
--- a/Engine/Source/Runtime/Slate/Public/Widgets/Input/SVectorInputBox.h
+++ b/Engine/Source/Runtime/Slate/Public/Widgets/Input/SVectorInputBox.h
@@ -246,14 +246,7 @@ public:
 		SLATE_ATTRIBUTE( int32, LinearDeltaSensitivity )
 
 		/** Should the axis labels be colored */
-		SLATE_ARGUMENT( bool, bColorAxisLabels )		
-
-		/** Allow responsive layout to crush the label and margins when there is not a lot of room */
-		UE_DEPRECATED(5.0, "AllowResponsiveLayout unused as it is no longer necessary.")
-		FArguments& AllowResponsiveLayout(bool bAllow)
-		{
-			return TSlateBaseNamedArgs<ThisClass>::Me();
-		}
+		SLATE_ARGUMENT( bool, bColorAxisLabels )
 
 		/** Called right before the slider begins to move for any of the vector components */
 		SLATE_EVENT( FSimpleDelegate, OnBeginSliderMovement )
diff --git a/Engine/Source/Runtime/Slate/Public/Widgets/Layout/SBorder.h b/Engine/Source/Runtime/Slate/Public/Widgets/Layout/SBorder.h
--- a/Engine/Source/Runtime/Slate/Public/Widgets/Layout/SBorder.h
+++ b/Engine/Source/Runtime/Slate/Public/Widgets/Layout/SBorder.h
@@ -172,17 +172,6 @@ protected:
 	TSlateAttributeRef<FVector2D> GetDesiredSizeScaleAttribute() const { return TSlateAttributeRef<FVector2D>(SharedThis(this), DesiredSizeScaleAttribute); }
 	TSlateAttributeRef<bool> GetShowDisabledEffectAttribute() const { return TSlateAttributeRef<bool>(SharedThis(this), ShowDisabledEffectAttribute); }
 
-#if WITH_EDITORONLY_DATA
-	UE_DEPRECATED(5.0, "Direct access to BorderImage is now deprecated. Use the setter or getter.")
-	FInvalidatableBrushAttribute BorderImage;
-	UE_DEPRECATED(5.0, "Direct access to BorderBackgroundColor is now deprecated. Use the setter or getter.")
-	TSlateDeprecatedTAttribute<FSlateColor> BorderBackgroundColor;
-	UE_DEPRECATED(5.0, "Direct access to DesiredSizeScale is now deprecated. Use the setter or getter.")
-	TSlateDeprecatedTAttribute<FVector2D> DesiredSizeScale;
-	UE_DEPRECATED(5.0, "Direct access to ShowDisabledEffect is now deprecated. Use the setter or getter.")
-	TSlateDeprecatedTAttribute<bool> ShowDisabledEffect;
-#endif
-
  private:
 	TSlateAttribute<const FSlateBrush*> BorderImageAttribute;
 	TSlateAttribute<FSlateColor> BorderBackgroundColorAttribute;
diff --git a/Engine/Source/Runtime/Slate/Public/Widgets/Layout/SBox.h b/Engine/Source/Runtime/Slate/Public/Widgets/Layout/SBox.h
--- a/Engine/Source/Runtime/Slate/Public/Widgets/Layout/SBox.h
+++ b/Engine/Source/Runtime/Slate/Public/Widgets/Layout/SBox.h
@@ -24,12 +24,6 @@ class SBox : public SPanel
 	SLATE_DECLARE_WIDGET_API(SBox, SPanel, SLATE_API)
 
 public:
-	class UE_DEPRECATED(5.0, "FBoxSlot is deprecated. Use FSingleWidgetChildrenWithBasicLayoutSlot or FOneSimpleMemberChild")
-	FBoxSlot : public FSingleWidgetChildrenWithBasicLayoutSlot
-	{
-		using FSingleWidgetChildrenWithBasicLayoutSlot::FSingleWidgetChildrenWithBasicLayoutSlot;
-	};
-
 	SLATE_BEGIN_ARGS(SBox)
 		: _HAlign(HAlign_Fill)
 		, _VAlign(VAlign_Fill)
diff --git a/Engine/Source/Runtime/Slate/Public/Widgets/Layout/SConstraintCanvas.h b/Engine/Source/Runtime/Slate/Public/Widgets/Layout/SConstraintCanvas.h
--- a/Engine/Source/Runtime/Slate/Public/Widgets/Layout/SConstraintCanvas.h
+++ b/Engine/Source/Runtime/Slate/Public/Widgets/Layout/SConstraintCanvas.h
@@ -29,7 +29,6 @@ public:
 	/**
 	 * ConstraintCanvas slots allow child widgets to be positioned and sized
 	 */
-	PRAGMA_DISABLE_DEPRECATION_WARNINGS
 	class SLATE_API FSlot : public TSlotBase<FSlot>
 	{
 	public:
@@ -100,33 +99,19 @@ public:
 			return ZOrder;
 		}
 
-	public:
+	private:
 		/** Offset */
-		UE_DEPRECATED(5.0, "Direct access to OffsetAttr is now deprecated. Use the getter or setter.")
 		TAttribute<FMargin> OffsetAttr;
-
 		/** Anchors */
-		UE_DEPRECATED(5.0, "Direct access to AnchorsAttr is now deprecated. Use the getter or setter.")
 		TAttribute<FAnchors> AnchorsAttr;
-
 		/** Size */
-		UE_DEPRECATED(5.0, "Direct access to AlignmentAttr is now deprecated. Use the getter or setter.")
 		TAttribute<FVector2D> AlignmentAttr;
-
 		/** Auto-Size */
-		UE_DEPRECATED(5.0, "Direct access to AutoSizeAttr is now deprecated. Use the getter or setter.")
 		TAttribute<bool> AutoSizeAttr;
-
-#if WITH_EDITORONLY_DATA
-		/** Z-Order */
-		UE_DEPRECATED(5.0, "Direct access to ZOrderAttr is now deprecated. Use the getter or setter.")
-		TAttribute<float> ZOrderAttr;
-#endif
 	private:
 		/** Z-Order */
 		float ZOrder;
 	};
-	PRAGMA_ENABLE_DEPRECATION_WARNINGS
 
 	SLATE_BEGIN_ARGS( SConstraintCanvas )
 		{
diff --git a/Engine/Source/Runtime/Slate/Public/Widgets/Layout/SDPIScaler.h b/Engine/Source/Runtime/Slate/Public/Widgets/Layout/SDPIScaler.h
--- a/Engine/Source/Runtime/Slate/Public/Widgets/Layout/SDPIScaler.h
+++ b/Engine/Source/Runtime/Slate/Public/Widgets/Layout/SDPIScaler.h
@@ -60,11 +60,6 @@ protected:
 
 	TSlateAttributeRef<float> GetDPIScaleAttribute() const { return TSlateAttributeRef<float>(SharedThis(this), DPIScaleAttribute); }
 
-#if WITH_EDITORONLY_DATA
-	UE_DEPRECATED(5.0, "Direct access to DPIScale is now deprecated. Use the setter or getter.")
-	TSlateDeprecatedTAttribute<float> DPIScale;
-#endif
-
 	struct FDPIScalerOneChildSlot : ::TSingleWidgetChildrenWithBasicLayoutSlot<EInvalidateWidgetReason::None> // we want to add it to the Attribute descriptor
 	{
 		friend SDPIScaler;
diff --git a/Engine/Source/Runtime/Slate/Public/Widgets/Layout/SRadialBox.h b/Engine/Source/Runtime/Slate/Public/Widgets/Layout/SRadialBox.h
--- a/Engine/Source/Runtime/Slate/Public/Widgets/Layout/SRadialBox.h
+++ b/Engine/Source/Runtime/Slate/Public/Widgets/Layout/SRadialBox.h
@@ -100,10 +100,6 @@ public:
 
 	SLATE_API void SetUseAllottedWidth(bool bInUseAllottedWidth);
 
-	/** Mods the angle so it's between 0-360 */
-	UE_DEPRECATED(5.0, "NormalizeAngle is deprecated. You should use the FRotator::NormalizeAxis.")
-	SLATE_API int32 NormalizeAngle(int32 Angle) const;
-
 private:
 
 	/** Mods the angle so it's between 0-360 */
diff --git a/Engine/Source/Runtime/Slate/Public/Widgets/Layout/SScaleBox.h b/Engine/Source/Runtime/Slate/Public/Widgets/Layout/SScaleBox.h
--- a/Engine/Source/Runtime/Slate/Public/Widgets/Layout/SScaleBox.h
+++ b/Engine/Source/Runtime/Slate/Public/Widgets/Layout/SScaleBox.h
@@ -172,18 +172,6 @@ protected:
 	SLATE_API void DebugSafeAreaUpdated(const FMargin& NewSafeZone, bool bShouldRecacheMetrics);
 #endif
 
-protected:
-#if WITH_EDITORONLY_DATA
-	UE_DEPRECATED(5.0, "Direct access to StretchDirection is now deprecated. Use the setter or getter.")
-	TSlateDeprecatedTAttribute<EStretchDirection::Type> StretchDirection;
-	UE_DEPRECATED(5.0, "Direct access to Stretch is now deprecated. Use the setter or getter.")
-	TSlateDeprecatedTAttribute<EStretch::Type> Stretch;
-	UE_DEPRECATED(5.0, "Direct access to UserSpecifiedScale is now deprecated. Use the setter or getter.")
-	TSlateDeprecatedTAttribute<float> UserSpecifiedScale;
-	UE_DEPRECATED(5.0, "Direct access to IgnoreInheritedScale is now deprecated. Use the setter or getter.")
-	TSlateDeprecatedTAttribute<bool> IgnoreInheritedScale;
-#endif
-
 private:
 	/** The allowed direction of stretching of the content */
 	TSlateAttribute<EStretchDirection::Type> StretchDirectionAttribute;
diff --git a/Engine/Source/Runtime/Slate/Public/Widgets/Layout/SSplitter.h b/Engine/Source/Runtime/Slate/Public/Widgets/Layout/SSplitter.h
--- a/Engine/Source/Runtime/Slate/Public/Widgets/Layout/SSplitter.h
+++ b/Engine/Source/Runtime/Slate/Public/Widgets/Layout/SSplitter.h
@@ -60,7 +60,6 @@ public:
 	DECLARE_DELEGATE_RetVal_OneParam(FVector2D, FOnGetMaxSlotSize, int32);
 
 public:
-	PRAGMA_DISABLE_DEPRECATION_WARNINGS
 	class SLATE_API FSlot : public TSlotBase<FSlot>
 	{
 	public:		
@@ -147,22 +146,16 @@ public:
 		/** A slot can be resize if bIsResizable and the SizeRule is a FractionOfParent or the OnSlotResized delegate is set. */
 		bool CanBeResized() const;
 
-	public:
-		UE_DEPRECATED(5.0, "Direct access to SizingRule is now deprecated. Use the getter.")
+	private:
 		TAttribute<ESizeRule> SizingRule;
-		UE_DEPRECATED(5.0, "Direct access to SizeValue is now deprecated. Use the getter.")
 		TAttribute<float> SizeValue;
-		UE_DEPRECATED(5.0, "Direct access to MinSizeValue is now deprecated. Use the getter.")
 		TAttribute<float> MinSizeValue;
-		UE_DEPRECATED(5.0, "Direct access to OnSlotResized_Handler is now deprecated. Use the getter.")
 		FOnSlotResized OnSlotResized_Handler;
-		UE_DEPRECATED(5.0, "Direct access to bIsResizable is now deprecated. Use the getter.")
 		TOptional<bool> bIsResizable;
 	};
-	SLATE_API PRAGMA_ENABLE_DEPRECATION_WARNINGS
 
 	/** @return Add a new FSlot() */
-	static FSlot::FSlotArguments Slot();
+	SLATE_API static FSlot::FSlotArguments Slot();
 	
 	using FScopedWidgetSlotArguments = TPanelChildren<FSlot>::FScopedWidgetSlotArguments;
 	/**
diff --git a/Engine/Source/Runtime/Slate/Public/Widgets/Layout/SUniformGridPanel.h b/Engine/Source/Runtime/Slate/Public/Widgets/Layout/SUniformGridPanel.h
--- a/Engine/Source/Runtime/Slate/Public/Widgets/Layout/SUniformGridPanel.h
+++ b/Engine/Source/Runtime/Slate/Public/Widgets/Layout/SUniformGridPanel.h
@@ -19,7 +19,6 @@ class SUniformGridPanel : public SPanel
 {
 public:
 	/** Stores the per-child info for this panel type */
-	PRAGMA_DISABLE_DEPRECATION_WARNINGS
 	struct FSlot : public TSlotBase<FSlot>, public TAlignmentWidgetSlotMixin<FSlot>
 	{
 		SLATE_SLOT_BEGIN_ARGS_OneMixin(FSlot, TSlotBase<FSlot>, TAlignmentWidgetSlotMixin<FSlot>)
@@ -67,15 +66,12 @@ public:
 			}
 		}
 
-	public:
-		UE_DEPRECATED(5.0, "Direct access to Column is now deprecated. Use the getter or setter.")
+	private:
 		int32 Column;
-		UE_DEPRECATED(5.0, "Direct access to Row is now deprecated. Use the getter or setter.")
 		int32 Row;
 	};
-	SLATE_API PRAGMA_ENABLE_DEPRECATION_WARNINGS
 
-	SUniformGridPanel();
+	SLATE_API SUniformGridPanel();
 
 	/**
 	 * Used by declarative syntax to create a Slot in the specified Column, Row.
diff --git a/Engine/Source/Runtime/Slate/Public/Widgets/Layout/SWrapBox.h b/Engine/Source/Runtime/Slate/Public/Widgets/Layout/SWrapBox.h
--- a/Engine/Source/Runtime/Slate/Public/Widgets/Layout/SWrapBox.h
+++ b/Engine/Source/Runtime/Slate/Public/Widgets/Layout/SWrapBox.h
@@ -53,7 +53,6 @@ class SWrapBox : public SPanel
 public:
 
 	/** A slot that support alignment of content and padding */
-	PRAGMA_DISABLE_DEPRECATION_WARNINGS
 	class FSlot : public TBasicLayoutWidgetSlot<FSlot>
 	{
 	public:
@@ -126,23 +125,17 @@ public:
 			return bSlotForceNewLine;
 		}
 
-	public:
-		UE_DEPRECATED(5.0, "Direct access to SlotFillLineWhenSizeLessThan is now deprecated. Use the getter or setter.")
+	private:
 		TOptional<float> SlotFillLineWhenSizeLessThan;
-		UE_DEPRECATED(5.0, "Direct access to bSlotFillEmptySpace is now deprecated. Use the getter or setter.")
 		bool bSlotFillEmptySpace;
-		UE_DEPRECATED(5.0, "Direct access to bSlotForceNewLine is now deprecated. Use the getter or setter.")
 		bool bSlotForceNewLine;
 	};
-	PRAGMA_ENABLE_DEPRECATION_WARNINGS
 
 
-	PRAGMA_DISABLE_DEPRECATION_WARNINGS
 	SLATE_BEGIN_ARGS(SWrapBox)
 		: _PreferredSize(100.f)
 		, _HAlign(HAlign_Left)
 		, _InnerSlotPadding(FVector2D::ZeroVector)
-		, _UseAllottedWidth(false)
 		, _UseAllottedSize(false)
 		, _Orientation(EOrientation::Orient_Horizontal)
 		{
@@ -152,10 +145,6 @@ public:
 		/** The slot supported by this panel */
 		SLATE_SLOT_ARGUMENT( FSlot, Slots )
 
-		/** The preferred width, if not set will fill the space */
-		UE_DEPRECATED(5.0, "PreferredWidth is deprecated. Use PreferredSize instead.")
-		SLATE_ATTRIBUTE( float, PreferredWidth )
-
 		/** The preferred size, if not set will fill the space */
 		SLATE_ATTRIBUTE( float, PreferredSize )
 
@@ -165,19 +154,14 @@ public:
 		/** The inner slot padding goes between slots sharing borders */
 		SLATE_ARGUMENT( FVector2D, InnerSlotPadding )
 
-		/** if true, the PreferredWidth will always match the room available to the SWrapBox  */
-		UE_DEPRECATED(5.0, "UseAllottedWidth is deprecated. Use UseAllottedSize instead.")
-		SLATE_ARGUMENT( bool, UseAllottedWidth )
-
 		/** if true, the PreferredSize will always match the room available to the SWrapBox  */
 		SLATE_ARGUMENT( bool, UseAllottedSize )
 
 		/** Determines if the wrap box needs to arrange the slots left-to-right or top-to-bottom.*/
 		SLATE_ARGUMENT(EOrientation, Orientation);
 	SLATE_END_ARGS()
-	SLATE_API PRAGMA_ENABLE_DEPRECATION_WARNINGS
 
-	SWrapBox();
+	SLATE_API SWrapBox();
 
 	static FSlot::FSlotArguments Slot()
 	{
@@ -209,17 +193,9 @@ public:
 	/** See InnerSlotPadding Attribute */
 	SLATE_API void SetInnerSlotPadding(FVector2D InInnerSlotPadding);
 
-	/** Set the width at which the wrap panel should wrap its content. */
-	UE_DEPRECATED(4.26, "Deprecated, please use SetWrapSize() instead")
-	SLATE_API void SetWrapWidth(TAttribute<float> InWrapWidth);
-
 	/** Set the size at which the wrap panel should wrap its content. */
 	SLATE_API void SetWrapSize(TAttribute<float> InWrapSize );
 
-	/** When true, use the WrapWidth property to determine where to wrap to the next line. */
-	UE_DEPRECATED(4.26, "Deprecated, please use SetUseAllottedSize() instead")
-	SLATE_API void SetUseAllottedWidth(bool bInUseAllottedWidth);
-
 	/** When true, use the WrapSize property to determine where to wrap to the next line. */
 	SLATE_API void SetUseAllottedSize(bool bInUseAllottedSize);
 
diff --git a/Engine/Source/Runtime/Slate/Public/Widgets/SCanvas.h b/Engine/Source/Runtime/Slate/Public/Widgets/SCanvas.h
--- a/Engine/Source/Runtime/Slate/Public/Widgets/SCanvas.h
+++ b/Engine/Source/Runtime/Slate/Public/Widgets/SCanvas.h
@@ -76,14 +76,6 @@ public:
 			return Size.Get();
 		}
 
-	public:
-#if WITH_EDITORONLY_DATA
-		UE_DEPRECATED(5.0, "Direct access to PositionAttr is now deprecated. Use the getter or setter.")
-		TSlateDeprecatedTAttribute<FVector2D> PositionAttr;
-		UE_DEPRECATED(5.0, "Direct access to SizeAttr is now deprecated. Use the getter or setter.")
-		TSlateDeprecatedTAttribute<FVector2D> SizeAttr;
-#endif
-
 	public:
 		/** Default values for a slot. */
 		FSlot()
diff --git a/Engine/Source/Runtime/Slate/Public/Widgets/Text/STextBlock.h b/Engine/Source/Runtime/Slate/Public/Widgets/Text/STextBlock.h
--- a/Engine/Source/Runtime/Slate/Public/Widgets/Text/STextBlock.h
+++ b/Engine/Source/Runtime/Slate/Public/Widgets/Text/STextBlock.h
@@ -219,10 +219,6 @@ public:
 	/** Set TransformPolicy attribute */
 	SLATE_API void SetTransformPolicy(TAttribute<ETextTransformPolicy> InTransformPolicy);
 
-	/** Get TransformPolicy attribute */
-	UE_DEPRECATED(5.0, "GetTransformPolicy is not accessible anymore since it's attribute value may not have been updated yet.")
-	ETextTransformPolicy GetTransformPolicy() const { return GetTransformPolicyImpl(); }
-
 	/** Sets the overflow policy for this text block */
 	SLATE_API void SetOverflowPolicy(TOptional<ETextOverflowPolicy> InOverflowPolicy);
 
diff --git a/Engine/Source/Runtime/UMG/Private/Slate/SlateVectorArtData.cpp b/Engine/Source/Runtime/UMG/Private/Slate/SlateVectorArtData.cpp
--- a/Engine/Source/Runtime/UMG/Private/Slate/SlateVectorArtData.cpp
+++ b/Engine/Source/Runtime/UMG/Private/Slate/SlateVectorArtData.cpp
@@ -166,13 +166,6 @@ void USlateVectorArtData::EnsureValidData()
 #endif
 }
 
-void USlateVectorArtData::PreSave(const class ITargetPlatform* TargetPlatform)
-{
-	PRAGMA_DISABLE_DEPRECATION_WARNINGS;
-	Super::PreSave(TargetPlatform);
-	PRAGMA_ENABLE_DEPRECATION_WARNINGS;
-}
-
 void USlateVectorArtData::PreSave(FObjectPreSaveContext ObjectSaveContext)
 {
 	Super::PreSave(ObjectSaveContext);
diff --git a/Engine/Source/Runtime/UMG/Public/Slate/SlateVectorArtData.h b/Engine/Source/Runtime/UMG/Public/Slate/SlateVectorArtData.h
--- a/Engine/Source/Runtime/UMG/Public/Slate/SlateVectorArtData.h
+++ b/Engine/Source/Runtime/UMG/Public/Slate/SlateVectorArtData.h
@@ -109,10 +109,6 @@ public:
 
 private:
 	// ~ UObject Interface
-	PRAGMA_DISABLE_DEPRECATION_WARNINGS // Suppress compiler warning on override of deprecated function
-	UE_DEPRECATED(5.0, "Use version that takes FObjectPreSaveContext instead.")
-	UMG_API virtual void PreSave(const class ITargetPlatform* TargetPlatform) override;
-	UMG_API PRAGMA_ENABLE_DEPRECATION_WARNINGS
 	virtual void PreSave(FObjectPreSaveContext ObjectSaveContext) override;
 	// ~ UObject Interface
 
