PR #11767: Fixed crash caused by chord overlapping

diff --git a/Engine/Source/Runtime/Slate/Private/Framework/Commands/InputBindingManager.cpp b/Engine/Source/Runtime/Slate/Private/Framework/Commands/InputBindingManager.cpp
--- a/Engine/Source/Runtime/Slate/Private/Framework/Commands/InputBindingManager.cpp
+++ b/Engine/Source/Runtime/Slate/Private/Framework/Commands/InputBindingManager.cpp
@@ -514,7 +514,7 @@ void FInputBindingManager::CreateInputCommand( const TSharedRef<FBindingContext>
 				if (*PreviousInfo->ActiveChords[i] != PreviousInfo->DefaultChords[i])
 				{
 					// two user defined chords are the same within a context.  If the keybinding editor was used this wont happen so this must have been directly a modified user setting file
-					UE_LOG(LogSlate, Error, TEXT("Duplicate user defined chords found: [%s,%s].  Chord for %s being removed"), *InCommandInfo->GetLabel().ToString(), *PreviousInfo->GetLabel().ToString(), *ExistingInfo->GetLabel().ToString());
+					UE_LOG(LogSlate, Error, TEXT("Duplicate user defined chords found: [%s,%s].  Chord for %s being removed"), *InCommandInfo->GetLabel().ToString(), *PreviousInfo->GetLabel().ToString(), (ExistingInfo  ? *ExistingInfo->GetLabel().ToString() : TEXT("")));
 				}
 				ContextEntry.ChordToCommandInfoMaps[i].Remove(*PreviousInfo->ActiveChords[i]);
 				// Remove the existing chord. 
