In SDockingTabStack::OpenPersistentTab consider a persistent opened tab as a closed tab if it is absent from the live tabs.

This fixes an issue where persistent opened tabs could accumulate over time.  This situation can occur when the editor process is terminated or if the layout is saved during the shutdown sequence before an editor mode is deactivated.

In FTabManager::RestoreArea_Helper allow the call to SpawnTab to return nullptr instead of creating an empty "unrecognized tab".

This case can happen after one of the shutdown sequences described above when we attempt to restore a persistent opened tab but there is no currently registered tab spawner for it.  Now, we no longer present the unhelpful "unrecognized tab" but we allow the persistent opened tab to be treated as a closed tab and thus opened successfully at a later point in time.

diff --git a/Engine/Source/Runtime/Slate/Private/Framework/Docking/SDockingTabStack.cpp b/Engine/Source/Runtime/Slate/Private/Framework/Docking/SDockingTabStack.cpp
--- a/Engine/Source/Runtime/Slate/Private/Framework/Docking/SDockingTabStack.cpp
+++ b/Engine/Source/Runtime/Slate/Private/Framework/Docking/SDockingTabStack.cpp
@@ -1233,7 +1233,35 @@ const FSlateBrush* SDockingTabStack::GetTabStackBorderImage() const
 
 int32 SDockingTabStack::OpenPersistentTab( const FTabId& TabId, int32 OpenLocationAmongActiveTabs )
 {
-	const int32 ExistingClosedTabIndex = Tabs.IndexOfByPredicate(FTabMatcher(TabId, static_cast<ETabState::Type>(ETabState::ClosedTab|ETabState::SidebarTab)));
+	int32 ExistingClosedTabIndex = Tabs.IndexOfByPredicate(FTabMatcher(TabId, static_cast<ETabState::Type>(ETabState::ClosedTab|ETabState::SidebarTab)));
+
+	if (ExistingClosedTabIndex == INDEX_NONE)
+	{
+		// Check for a persistent opened tab that isn't actually opened ( in the live tabs ).
+		// This situation can happen in some corner cases, e.g:
+		// - If the process is terminated after a periodic layout save.
+		// - If the layout save on editor shutdown occurs before an editor mode is deactivated.
+		// In any case we want to treat this opened tab as a closed tab instead of creating a new tab.
+		const bool bTreatIndexNoneAsWildcard = false;
+		const int32 ExistingOpenedTabIndex = Tabs.IndexOfByPredicate(FTabMatcher(TabId, static_cast<ETabState::Type>(ETabState::OpenedTab), bTreatIndexNoneAsWildcard));
+		if (ExistingOpenedTabIndex != INDEX_NONE)
+		{
+			bool bHasLiveTab = false;
+			const TArray< TSharedRef<SDockTab> > LiveTabs = this->GetTabs().AsArrayCopy();
+			for (int32 TabIndex = 0; TabIndex < LiveTabs.Num(); ++TabIndex)
+			{
+				if (TabId == LiveTabs[TabIndex]->GetLayoutIdentifier())
+				{
+					bHasLiveTab = true;
+					break;
+				}
+			}
+			if (!bHasLiveTab)
+			{
+				ExistingClosedTabIndex = ExistingOpenedTabIndex;
+			}
+		}
+	}
 
 	if (OpenLocationAmongActiveTabs == INDEX_NONE)
 	{						
diff --git a/Engine/Source/Runtime/Slate/Private/Framework/Docking/TabManager.cpp b/Engine/Source/Runtime/Slate/Private/Framework/Docking/TabManager.cpp
--- a/Engine/Source/Runtime/Slate/Private/Framework/Docking/TabManager.cpp
+++ b/Engine/Source/Runtime/Slate/Private/Framework/Docking/TabManager.cpp
@@ -1683,7 +1683,8 @@ TSharedPtr<SDockingNode> FTabManager::RestoreArea_Helper(const TSharedRef<FLayou
 		{
 			if ((SomeTab.TabState == ETabState::OpenedTab || SomeTab.TabState == ETabState::SidebarTab) && IsValidTabForSpawning(SomeTab))
 			{
-				const TSharedPtr<SDockTab> NewTabWidget = SpawnTab(SomeTab.TabId, ParentWindow, bCanOutputBeNullptr);
+				const bool bCanUnrecognizedTabBeNullptr = true;
+				const TSharedPtr<SDockTab> NewTabWidget = SpawnTab(SomeTab.TabId, ParentWindow, bCanUnrecognizedTabBeNullptr);
 
 				if (NewTabWidget.IsValid())
 				{
