TableViewBase scrollbar: take into account the padding for scrollbar when calculating the size of the box that contains the scrollbar.

diff --git a/Engine/Source/Runtime/Slate/Private/Widgets/Views/STableViewBase.cpp b/Engine/Source/Runtime/Slate/Private/Widgets/Views/STableViewBase.cpp
--- a/Engine/Source/Runtime/Slate/Private/Widgets/Views/STableViewBase.cpp
+++ b/Engine/Source/Runtime/Slate/Private/Widgets/Views/STableViewBase.cpp
@@ -97,7 +97,7 @@ void STableViewBase::ConstructChildren( const TAttribute<float>& InItemWidth, co
 			.Style(InScrollBarStyle ? InScrollBarStyle : &FAppStyle::Get().GetWidgetStyle<FScrollBarStyle>("ScrollBar"))
 			.PreventThrottling(bInPreventThrottling);
 
-		const FOptionalSize ScrollBarSize(InScrollBarStyle ? InScrollBarStyle->Thickness : 16.f);
+		const FOptionalSize ScrollBarSize(InScrollBarStyle ? InScrollBarStyle->Thickness + (SScrollBar::DefaultUniformPadding * 2.0f) : 16.f);
 
 		if (Orientation == Orient_Vertical)
 		{
diff --git a/Engine/Source/Runtime/Slate/Public/Widgets/Layout/SScrollBar.h b/Engine/Source/Runtime/Slate/Public/Widgets/Layout/SScrollBar.h
--- a/Engine/Source/Runtime/Slate/Public/Widgets/Layout/SScrollBar.h
+++ b/Engine/Source/Runtime/Slate/Public/Widgets/Layout/SScrollBar.h
@@ -29,6 +29,8 @@ class SScrollBar : public SBorder
 {
 public:
 
+	inline static const float DefaultUniformPadding = 2.0f;
+
 	SLATE_BEGIN_ARGS( SScrollBar )
 		: _Style( &FAppStyle::Get().GetWidgetStyle<FScrollBarStyle>("Scrollbar") )
 		, _OnUserScrolled()
@@ -43,7 +45,7 @@ public:
 		, _Orientation( Orient_Vertical )
 		, _DragFocusCause( EFocusCause::Mouse )
 		, _Thickness()
-		, _Padding( 2.0f )
+		, _Padding( DefaultUniformPadding )
 		{}
 
 		/** The style to use for this scrollbar */
