TabManager : Disable deferred layout saves for the entirety of an FDockingDragOperation.

diff --git a/Engine/Source/Runtime/Slate/Private/Framework/Docking/FDockingDragOperation.cpp b/Engine/Source/Runtime/Slate/Private/Framework/Docking/FDockingDragOperation.cpp
--- a/Engine/Source/Runtime/Slate/Private/Framework/Docking/FDockingDragOperation.cpp
+++ b/Engine/Source/Runtime/Slate/Private/Framework/Docking/FDockingDragOperation.cpp
@@ -44,6 +44,8 @@ void FDockingDragOperation::OnDrop( bool bDropWasHandled, const FPointerEvent& M
 		// The event was handled, so we HAVE to have some window that we dropped onto.
 		TSharedRef<SWindow> WindowDroppedInto = MouseEvent.GetWindow();
 
+		TabOwnerAreaOfOrigin->GetTabManager()->GetPrivateApi().SetCanDoDeferredLayoutSave(true);
+
 		TabOwnerAreaOfOrigin->GetTabManager()->GetPrivateApi().OnTabRelocated(TabBeingDragged.ToSharedRef(), WindowDroppedInto);
 	}
 
@@ -282,6 +284,8 @@ FDockingDragOperation::FDockingDragOperation( const TSharedRef<SDockTab>& InTabT
 		]
 	);
 
+	TabOwnerAreaOfOrigin->GetTabManager()->GetPrivateApi().SetCanDoDeferredLayoutSave(false);
+
 	if ( TabBeingDragged->IsActive() )
 	{
 		FGlobalTabmanager::Get()->SetActiveTab(NULL);
@@ -383,6 +387,8 @@ void FDockingDragOperation::DroppedOntoNothing()
 	// Do this after the window parenting so that the window title is set correctly
 	NewDockNode->OpenTab(TabBeingDragged.ToSharedRef());
 
+	TabOwnerAreaOfOrigin->GetTabManager()->GetPrivateApi().SetCanDoDeferredLayoutSave(true);
+
 	// Let every widget under this tab manager know that this tab has found a new home.
 	TabOwnerAreaOfOrigin->GetTabManager()->GetPrivateApi().OnTabRelocated(TabBeingDragged.ToSharedRef(), NewWindow);
 }
diff --git a/Engine/Source/Runtime/Slate/Private/Framework/Docking/SDockingTabWell.cpp b/Engine/Source/Runtime/Slate/Private/Framework/Docking/SDockingTabWell.cpp
--- a/Engine/Source/Runtime/Slate/Private/Framework/Docking/SDockingTabWell.cpp
+++ b/Engine/Source/Runtime/Slate/Private/Framework/Docking/SDockingTabWell.cpp
@@ -322,10 +322,6 @@ FReply SDockingTabWell::StartDraggingTab( TSharedRef<SDockTab> TabToStartDraggin
 
 	if (bCanLeaveTabWell)
 	{
-		// Clear any pending layout save so that it doesn't occur during the drag operation.
-		// A new save will be requested in FTabManager::OnTabRelocated when the drag ends.
-		TabManager->ClearPendingLayoutSave();
-
 		// We just removed the foreground tab.
 		ForegroundTabIndex = INDEX_NONE;
 		ParentTabStackPtr.Pin()->OnTabRemoved(TabToStartDragging->GetLayoutIdentifier());
diff --git a/Engine/Source/Runtime/Slate/Private/Framework/Docking/TabManager.cpp b/Engine/Source/Runtime/Slate/Private/Framework/Docking/TabManager.cpp
--- a/Engine/Source/Runtime/Slate/Private/Framework/Docking/TabManager.cpp
+++ b/Engine/Source/Runtime/Slate/Private/Framework/Docking/TabManager.cpp
@@ -751,6 +751,15 @@ void FTabManager::FPrivateApi::HideWindows()
 	SetWindowVisibility(TabManager.DockAreas, false);
 }
 
+void FTabManager::FPrivateApi::SetCanDoDeferredLayoutSave(bool bInCanDoDeferredLayoutSave)
+{
+	if (!bInCanDoDeferredLayoutSave)
+	{
+		TabManager.ClearPendingLayoutSave();
+	}
+	TabManager.bCanDoDeferredLayoutSave = bInCanDoDeferredLayoutSave;
+}
+
 FTabManager::FPrivateApi& FTabManager::GetPrivateApi()
 {
 	return *PrivateApi;
@@ -964,6 +973,11 @@ void FTabManager::RequestSavePersistentLayout()
 	// this is to avoid hitches when eg. resizing a docked tab
 	ClearPendingLayoutSave();
 
+	if (!bCanDoDeferredLayoutSave)
+	{
+		return;
+	}
+
 	auto OnTick = [ThisWeak = AsWeak()](float FrameTime)
 	{
 		if (TSharedPtr<FTabManager> This = ThisWeak.Pin())
diff --git a/Engine/Source/Runtime/Slate/Public/Framework/Docking/TabManager.h b/Engine/Source/Runtime/Slate/Public/Framework/Docking/TabManager.h
--- a/Engine/Source/Runtime/Slate/Public/Framework/Docking/TabManager.h
+++ b/Engine/Source/Runtime/Slate/Public/Framework/Docking/TabManager.h
@@ -776,6 +776,8 @@ class FTabManager : public TSharedFromThis<FTabManager>
 				void ShowWindows();
 				void HideWindows();
 
+				void SetCanDoDeferredLayoutSave(bool bInCanDoDeferredLayoutSave);
+
 			private:
 				FTabManager& TabManager;
 				
@@ -1248,6 +1250,9 @@ protected:
 		/* Prevent or allow Drag operation. */
 		bool bCanDoDragOperation = true;
 
+		/** Whether or not this tab manager can perform deferred layout saves at this time; used internally to disable layout saves during drag operations */
+		bool bCanDoDeferredLayoutSave = true;
+
 		/** Whether or not this tab manager puts any registered menus in the windows menu bar area */
 		bool bAllowPerWindowMenu = false;
 
