Slate- Added event for when menus are closed. This was necessary for some custom menu cleanup code for Sequencer.

diff --git a/Engine/Source/Runtime/Slate/Private/Framework/Application/MenuStack.cpp b/Engine/Source/Runtime/Slate/Private/Framework/Application/MenuStack.cpp
--- a/Engine/Source/Runtime/Slate/Private/Framework/Application/MenuStack.cpp
+++ b/Engine/Source/Runtime/Slate/Private/Framework/Application/MenuStack.cpp
@@ -700,6 +700,11 @@ void FMenuStack::OnMenuDestroyed(TSharedRef<IMenu> InMenu)
 			Stack[StackIndex]->Dismiss();	// this will cause OnMenuDestroyed() to re-enter
 		}
 
+		if (MenuDestroyedEvent.IsBound())
+		{
+			MenuDestroyedEvent.Broadcast(InMenu);
+		}
+
 		// Clean up the stack and content map arrays
 		for (int32 StackIndex = Stack.Num() - 1; StackIndex >= Index; --StackIndex)
 		{
diff --git a/Engine/Source/Runtime/Slate/Private/Framework/Application/SlateApplication.cpp b/Engine/Source/Runtime/Slate/Private/Framework/Application/SlateApplication.cpp
--- a/Engine/Source/Runtime/Slate/Private/Framework/Application/SlateApplication.cpp
+++ b/Engine/Source/Runtime/Slate/Private/Framework/Application/SlateApplication.cpp
@@ -865,6 +865,8 @@ FSlateApplication::FSlateApplication()
 #if WITH_EDITOR
 	FCoreDelegates::OnSafeFrameChangedEvent.AddRaw(this, &FSlateApplication::SwapSafeZoneTypes);
 	OnDebugSafeZoneChanged.AddRaw(this, &FSlateApplication::UpdateCustomSafeZone);
+
+	MenuStack.OnMenuDestroyedEvent().AddRaw(this, &FSlateApplication::OnMenuDestroyed);
 #endif
 
 	IConsoleVariable* CVarGlobalInvalidation = IConsoleManager::Get().FindConsoleVariable(TEXT("Slate.EnableGlobalInvalidation"));
@@ -885,6 +887,7 @@ FSlateApplication::~FSlateApplication()
 
 #if WITH_EDITOR
 	OnDebugSafeZoneChanged.RemoveAll(this);
+	MenuStack.OnMenuDestroyedEvent().RemoveAll(this);
 #endif
 
 	IConsoleVariable* CVarGlobalInvalidation = IConsoleManager::Get().FindConsoleVariable(TEXT("Slate.EnableGlobalInvalidation"));
@@ -6316,6 +6319,15 @@ bool FSlateApplication::ExecuteNavigation(const FWidgetPath& NavigationSource, T
 	return bHandled;
 }
 
+
+void FSlateApplication::OnMenuDestroyed(const TSharedRef<IMenu>& Menu)
+{
+	if (MenuBeingDestroyedEvent.IsBound())
+	{
+		MenuBeingDestroyedEvent.Broadcast(Menu);
+	}
+}
+
 bool FSlateApplication::OnControllerAnalog(FGamepadKeyNames::Type KeyName, FPlatformUserId PlatformUserId, FInputDeviceId InputDeviceId, float AnalogValue)
 {
 	FKey Key(KeyName);
diff --git a/Engine/Source/Runtime/Slate/Public/Framework/Application/MenuStack.h b/Engine/Source/Runtime/Slate/Public/Framework/Application/MenuStack.h
--- a/Engine/Source/Runtime/Slate/Public/Framework/Application/MenuStack.h
+++ b/Engine/Source/Runtime/Slate/Public/Framework/Application/MenuStack.h
@@ -10,6 +10,7 @@
 #include "Types/SlateStructs.h"
 #include "Widgets/SWindow.h"
 #include "Application/ThrottleManager.h"
+#include "Delegates/DelegateCombinations.h"
 
 class FWidgetPath;
 class SMenuPanel;
@@ -58,6 +59,8 @@ struct FPopupTransitionEffect
 	{ }
 };
 
+DECLARE_MULTICAST_DELEGATE_OneParam(FOnMenuDestroyed, const TSharedRef<IMenu>& InMenu);
+
 /**
  * Represents a stack of open menus. The last item in the stack is the top most menu
  * Menus are described as IMenus. Implementations of IMenu can control how the menus are created and presented (e.g. in their own window)
@@ -204,6 +207,9 @@ public:
 	 */
 	bool HasOpenSubMenus(TSharedPtr<IMenu> InMenu) const;
 
+	/** Delegate called when a menu is dismissed and destroyed */
+	FOnMenuDestroyed& OnMenuDestroyedEvent() { return MenuDestroyedEvent; }
+
 private:
 	/**
 	 * Queries the widgets in the path for a menu popup method.
@@ -366,4 +372,7 @@ private:
 
 	/** Guard to prevent the HostWindow and HostWindowPopupPanel being set reentrantly */
 	bool bHostWindowGuard;
+
+	/** Called when a menu is destroyed */
+	FOnMenuDestroyed MenuDestroyedEvent;
 };
diff --git a/Engine/Source/Runtime/Slate/Public/Framework/Application/SlateApplication.h b/Engine/Source/Runtime/Slate/Public/Framework/Application/SlateApplication.h
--- a/Engine/Source/Runtime/Slate/Public/Framework/Application/SlateApplication.h
+++ b/Engine/Source/Runtime/Slate/Public/Framework/Application/SlateApplication.h
@@ -565,6 +565,9 @@ public:
 	DECLARE_EVENT_OneParam(FSlateApplication, FOnWindowBeingDestroyed, const SWindow&);
 	FOnWindowBeingDestroyed& OnWindowBeingDestroyed() { return WindowBeingDestroyedEvent; }
 
+	/* Delegate called when a menu is about to be destroyed*/
+	FOnMenuDestroyed& OnMenuBeingDestroyed() { return MenuBeingDestroyedEvent; }
+
 	/** Delegate called just before possible focus change */
 	DECLARE_MULTICAST_DELEGATE_FiveParams(FOnFocusChanging, const FFocusEvent&, const FWeakWidgetPath&, const TSharedPtr<SWidget>&, const FWidgetPath&, const TSharedPtr<SWidget>&);
 	FOnFocusChanging& OnFocusChanging() { return FocusChangingDelegate; }
@@ -1842,6 +1845,10 @@ private:
 	 */
 	SLATE_API bool ExecuteNavigation(const FWidgetPath& NavigationSource, TSharedPtr<SWidget> DestinationWidget, const uint32 UserIndex, bool bAlwaysHandleNavigationAttempt);
 
+
+	/* Event receiver for when a menu in the menu stack is destroyed */
+	SLATE_API void OnMenuDestroyed(const TSharedRef<IMenu>& Menu);
+
 private:
 	SLATE_API FSlateApplication();
 	SLATE_API void SetLastUserInteractionTime(const double InCurrentTime);	
@@ -2095,6 +2102,9 @@ private:
 	/** Delegate for when a window is in the process of being destroyed */
 	FOnWindowBeingDestroyed WindowBeingDestroyedEvent;
 
+	/** Delegate for a when a menu has been dismissed */
+	FOnMenuDestroyed MenuBeingDestroyedEvent;
+
 	/** Delegate for slate Tick during modal dialogs */
 	FOnModalLoopTickEvent ModalLoopTickEvent;
 
