Optimized FDefaultRichTextMarkupParser

Create the FRegexMatcher for a given input string once rather than once per-loop, since they were always being constructed on the same input string, with SetLimits being used to restrict the match range.

#rnx

diff --git a/Engine/Source/Runtime/Slate/Private/Framework/Text/RichTextMarkupProcessing.cpp b/Engine/Source/Runtime/Slate/Private/Framework/Text/RichTextMarkupProcessing.cpp
--- a/Engine/Source/Runtime/Slate/Private/Framework/Text/RichTextMarkupProcessing.cpp
+++ b/Engine/Source/Runtime/Slate/Private/Framework/Text/RichTextMarkupProcessing.cpp
@@ -84,9 +84,17 @@ void FDefaultRichTextMarkupParser::Process(TArray<FTextLineParseResults>& Result
 
 void FDefaultRichTextMarkupParser::ParseLineRanges(const FString& Input, const TArray<FTextRange>& LineRanges, TArray<FTextLineParseResults>& LineParseResultsArray) const
 {
+	if (LineRanges.IsEmpty())
+	{
+		return;
+	}
+
 	// Special regular expression pattern for matching rich text markup elements. IE: <ElementName AttributeName="AttributeValue">Content</>
 	FRegexMatcher ElementRegexMatcher(ElementRegexPattern, Input);
 
+	// Lazily constructed attribute parser
+	TOptional<FRegexMatcher> LazyAttributeRegexMatcher;
+
 	// Parse line ranges, creating line parse results and run parse results.
 	for(int32 i = 0; i < LineRanges.Num(); ++i)
 	{
@@ -126,7 +134,11 @@ void FDefaultRichTextMarkupParser::ParseLineRanges(const FString& Input, const T
 
 			if(AttributeListBegin != INDEX_NONE && AttributeListEnd != INDEX_NONE)
 			{
-				FRegexMatcher AttributeRegexMatcher(AttributeRegexPattern, Input);
+				if (!LazyAttributeRegexMatcher.IsSet())
+				{
+					LazyAttributeRegexMatcher.Emplace(AttributeRegexPattern, Input);
+				}
+				FRegexMatcher& AttributeRegexMatcher = *LazyAttributeRegexMatcher;
 				AttributeRegexMatcher.SetLimits(AttributeListBegin, AttributeListEnd);
 
 				// Iterate through the attribute list, each time trying to find a match for the attribute regex.
@@ -184,6 +196,13 @@ void FDefaultRichTextMarkupParser::ParseLineRanges(const FString& Input, const T
 
 void FDefaultRichTextMarkupParser::HandleEscapeSequences(const FString& Input, TArray<FTextLineParseResults>& LineParseResultsArray, FString& ConcatenatedUnescapedLines) const
 {
+	if (LineParseResultsArray.IsEmpty())
+	{
+		return;
+	}
+
+	FRegexMatcher EscapeSequenceRegexMatcher(EscapeSequenceRegexPattern, Input);
+
 	// Modify original string to handle escape sequences that need to be replaced while updating run ranges.
 	for(int32 i = 0; i < LineParseResultsArray.Num(); ++i)
 	{
@@ -195,7 +214,6 @@ void FDefaultRichTextMarkupParser::HandleEscapeSequences(const FString& Input, T
 		for(int32 j = 0; j < LineParseResults.Runs.Num(); ++j)
 		{
 			FTextRunParseResults& RunParseResults = LineParseResults.Runs[j];
-			FRegexMatcher EscapeSequenceRegexMatcher(EscapeSequenceRegexPattern, Input);
 
 			TArray<int32*> IndicesToUpdate;
 			IndicesToUpdate.Add(&RunParseResults.OriginalRange.BeginIndex);
