- Make SDockingTabStack::OnMouseButtonDown _not_ consider a mouse button click for tab focus if the button is _not_ left/right or middle click (this allows for mouse forward/back buttons CommandBindings)
- Make SWidget::OnKeyDown _not_ capture any directional input in case a modifier key is down (this allows for Modifier+ Arrow keys CommandBindings)

diff --git a/Engine/Source/Runtime/Slate/Private/Framework/Application/NavigationConfig.cpp b/Engine/Source/Runtime/Slate/Private/Framework/Application/NavigationConfig.cpp
--- a/Engine/Source/Runtime/Slate/Private/Framework/Application/NavigationConfig.cpp
+++ b/Engine/Source/Runtime/Slate/Private/Framework/Application/NavigationConfig.cpp
@@ -58,21 +58,19 @@ void FNavigationConfig::OnUserRemoved(int32 UserIndex)
 
 EUINavigation FNavigationConfig::GetNavigationDirectionFromKey(const FKeyEvent& InKeyEvent) const
 {
-	if (const EUINavigation* Rule = KeyEventRules.Find(InKeyEvent.GetKey()))
+	const bool bModifierHeld = InKeyEvent.IsControlDown() || InKeyEvent.IsAltDown() || InKeyEvent.IsCommandDown() || InKeyEvent.IsShiftDown();
+	if (bIgnoreModifiersForNavigationActions || !bModifierHeld)
 	{
-		if (bKeyNavigation)
+		if (const EUINavigation* Rule = KeyEventRules.Find(InKeyEvent.GetKey()))
 		{
-			return *Rule;
+			if (bKeyNavigation)
+			{
+				return *Rule;
+			}
 		}
-	}
-	else if (bTabNavigation && InKeyEvent.GetKey() == EKeys::Tab )
-	{
-		//@TODO: Really these uses of input should be at a lower priority, only occurring if nothing else handled them
-		// For now this code prevents consuming them when some modifiers are held down, allowing some limited binding
-		const bool bAllowEatingKeyEvents = !InKeyEvent.IsControlDown() && !InKeyEvent.IsAltDown() && !InKeyEvent.IsCommandDown();
-
-		if ( bAllowEatingKeyEvents )
+		else if (bTabNavigation && InKeyEvent.GetKey() == EKeys::Tab )
 		{
+			//@TODO: Really these uses of input should be at a lower priority, only occurring if nothing else handled them
 			return ( InKeyEvent.IsShiftDown() ) ? EUINavigation::Previous : EUINavigation::Next;
 		}
 	}
diff --git a/Engine/Source/Runtime/Slate/Private/Framework/Docking/SDockingTabStack.cpp b/Engine/Source/Runtime/Slate/Private/Framework/Docking/SDockingTabStack.cpp
--- a/Engine/Source/Runtime/Slate/Private/Framework/Docking/SDockingTabStack.cpp
+++ b/Engine/Source/Runtime/Slate/Private/Framework/Docking/SDockingTabStack.cpp
@@ -497,7 +497,12 @@ void SDockingTabStack::OnFocusChanging( const FWeakWidgetPath& PreviousFocusPath
 FReply SDockingTabStack::OnMouseButtonDown( const FGeometry& MyGeometry, const FPointerEvent& MouseEvent )
 {
 	const TSharedPtr<SDockTab> ForegroundTab = TabWell->GetForegroundTab();
-	if ( ForegroundTab.IsValid() && !ForegroundTab->IsActive() )
+
+	const bool bIsRelevantButtonForTabFocus = MouseEvent.GetPressedButtons().Contains(EKeys::LeftMouseButton)
+			|| MouseEvent.GetPressedButtons().Contains(EKeys::RightMouseButton)
+			|| MouseEvent.GetPressedButtons().Contains(EKeys::MiddleMouseButton);
+
+	if (bIsRelevantButtonForTabFocus && ForegroundTab.IsValid() && !ForegroundTab->IsActive())
 	{
 		FGlobalTabmanager::Get()->SetActiveTab( ForegroundTab );
 #if PLATFORM_LINUX
