Fixed pre-SH bug with broken scrollbars

#rnx

diff --git a/Engine/Source/Runtime/Slate/Public/Widgets/Views/SListView.h b/Engine/Source/Runtime/Slate/Public/Widgets/Views/SListView.h
--- a/Engine/Source/Runtime/Slate/Public/Widgets/Views/SListView.h
+++ b/Engine/Source/Runtime/Slate/Public/Widgets/Views/SListView.h
@@ -114,6 +114,7 @@ public:
 		, _HandleSpacebarSelection(false)
 		, _IsFocusable(true)
 		, _ReturnFocusToSelection()
+		, _GenerateAllItems_Debug(false)
 		, _OnItemToString_Debug()
 		, _OnEnteredBadState()
 		{
