In FSlateApplication::RoutePointerUpEvent route a synthetic pointer move event to ensure drag events ( e.g. OnDragLeave ) are called if necessary on any widgets before OnDrop.

The console variable "Slate.Input.PointerUpFiresPointerMoveForDragDrop" ( on by default ) can be used to disable this behavior if necessary.

diff --git a/Engine/Source/Runtime/Slate/Private/Framework/Application/SlateApplication.cpp b/Engine/Source/Runtime/Slate/Private/Framework/Application/SlateApplication.cpp
--- a/Engine/Source/Runtime/Slate/Private/Framework/Application/SlateApplication.cpp
+++ b/Engine/Source/Runtime/Slate/Private/Framework/Application/SlateApplication.cpp
@@ -102,6 +102,12 @@ static FAutoConsoleVariableRef CVarSlateInputMotionFiresUserInteractionEvents(
 	TEXT("If this is false, LastUserInteractionTimeUpdateEvent events won't be fired based on motion input, and LastInteractionTime won't be updated\n")
 	TEXT("Some motion devices report small tiny changes constantly without filtering, so motion input is unhelpful for determining user activity"));
 
+static bool GSlateInputPointerUpFiresPointerMoveForDragDrop = true;
+static FAutoConsoleVariableRef CVarSlateInputPointerUpFiresPointerMoveForDragDrop(
+	TEXT("Slate.Input.PointerUpFiresPointerMoveForDragDrop"),
+	GSlateInputPointerUpFiresPointerMoveForDragDrop,
+	TEXT("When true, a synthetic pointer move event is fired from pointer up to ensure drag events are called if necessary on any widgets before OnDrop."));
+
 //////////////////////////////////////////////////////////////////////////
 
 bool GSlateEnableGamepadEditorNavigation = true;
@@ -5344,10 +5350,17 @@ FReply FSlateApplication::RoutePointerUpEvent(const FWidgetPath& WidgetsUnderPoi
 		// Switch worlds widgets in the current path
 		FScopedSwitchWorldHack SwitchWorld(LocalWidgetsUnderPointer);
 
-		// Cache the drag drop content and reset the pointer in case OnMouseButtonUpMessage re-enters as a result of OnDrop
-		// In such a case, we want the re-entrant call to skip any drag-drop stuff (otherwise we'd execute the drop action twice)
 		if (bIsDragDropping)
 		{
+			// Route a synthetic pointer move event to ensure drag events ( e.g. OnDragLeave ) are called if necessary on any widgets before OnDrop
+			if (GSlateInputPointerUpFiresPointerMoveForDragDrop)
+			{
+				const bool bIsSynthetic = true;
+				RoutePointerMoveEvent(LocalWidgetsUnderPointer, PointerEvent, bIsSynthetic);
+			}
+
+			// Cache the drag drop content and reset the pointer in case OnMouseButtonUpMessage re-enters as a result of OnDrop
+			// In such a case, we want the re-entrant call to skip any drag-drop stuff (otherwise we'd execute the drop action twice)
 			LocalDragDropContent = SlateUser->GetDragDropContent();
 			SlateUser->ResetDragDropContent();
 		}
