Remove unnecessary clear of FToolBarBuilder::bSectionShouldHaveSeparator

diff --git a/Engine/Source/Runtime/Slate/Private/Framework/MultiBox/MultiBoxBuilder.cpp b/Engine/Source/Runtime/Slate/Private/Framework/MultiBox/MultiBoxBuilder.cpp
--- a/Engine/Source/Runtime/Slate/Private/Framework/MultiBox/MultiBoxBuilder.cpp
+++ b/Engine/Source/Runtime/Slate/Private/Framework/MultiBox/MultiBoxBuilder.cpp
@@ -700,7 +700,6 @@ void FToolBarBuilder::ApplySectionBeginning()
 			NewSeparatorBlock->SetStyleNameOverride(CurrentStyleOverride);
 
 			MultiBox->AddMultiBlock(NewSeparatorBlock);
-			bSectionShouldHaveSeparator = false;
 		}
 		bSectionNeedsToBeApplied = false;
 	}
diff --git a/Engine/Source/Runtime/Slate/Public/Framework/MultiBox/MultiBoxBuilder.h b/Engine/Source/Runtime/Slate/Public/Framework/MultiBox/MultiBoxBuilder.h
--- a/Engine/Source/Runtime/Slate/Public/Framework/MultiBox/MultiBoxBuilder.h
+++ b/Engine/Source/Runtime/Slate/Public/Framework/MultiBox/MultiBoxBuilder.h
@@ -674,6 +674,7 @@ private:
 	/** True if there is a pending section that needs to be applied */
 	bool bSectionNeedsToBeApplied;
 
+	/** Whether the current section should be preceded by a separator */
 	bool bSectionShouldHaveSeparator;
 
 	/** Whether the buttons created can receive keyboard focus */
