Added a new listview flag ScrollIntoViewAlignment. This can be used to override the behavior of how an offscreen item is scrolled into view. The default flag matches prior behavior (CenterAligned) but a new condition (IntoView) will scroll the minimum distance required to bring the element into view (the bottom of the listview when scrolling down, or the top when scrolling up).

diff --git a/Engine/Source/Runtime/Slate/Public/Widgets/Views/SListView.h b/Engine/Source/Runtime/Slate/Public/Widgets/Views/SListView.h
--- a/Engine/Source/Runtime/Slate/Public/Widgets/Views/SListView.h
+++ b/Engine/Source/Runtime/Slate/Public/Widgets/Views/SListView.h
@@ -162,6 +162,8 @@ public:
 
 		SLATE_ARGUMENT( EOrientation, Orientation )
 
+		SLATE_ARGUMENT( EScrollIntoViewAlignment, ScrollIntoViewAlignment )
+
 		SLATE_ARGUMENT( bool, EnableAnimatedScrolling)
 
 		SLATE_ARGUMENT( TOptional<double>, FixedLineScrollOffset )
@@ -242,6 +244,7 @@ public:
 
 		this->bEnableAnimatedScrolling = InArgs._EnableAnimatedScrolling;
 		this->FixedLineScrollOffset = InArgs._FixedLineScrollOffset;
+		this->ScrollIntoViewAlignment = InArgs._ScrollIntoViewAlignment;
 
 		this->OnItemToString_Debug = InArgs._OnItemToString_Debug.IsBound()
 			? InArgs._OnItemToString_Debug
@@ -304,6 +307,7 @@ PRAGMA_DISABLE_DEPRECATION_WARNINGS
 		, ItemToScrollIntoView(TListTypeTraits<ItemType>::MakeNullPtr())
 		, UserRequestingScrollIntoView(0)
 		, ItemToNotifyWhenInView(TListTypeTraits<ItemType>::MakeNullPtr())
+		, ScrollIntoViewAlignment(EScrollIntoViewAlignment::CenterAligned)
 		, IsFocusable(true)
 	{ 
 #if WITH_ACCESSIBILITY
@@ -2121,8 +2125,23 @@ protected:
 					// Scroll the top of the listview to the item in question
 					double NewScrollOffset = IndexOfItem;
 
-					// Center the list view on the item in question.
-					NewScrollOffset -= (NumLiveWidgets / 2.0);
+					switch (ScrollIntoViewAlignment)
+					{
+					case EScrollIntoViewAlignment::CenterAligned:
+						// Center the list view on the item in question.
+						NewScrollOffset -= (NumLiveWidgets / 2.0);
+						break;
+
+					case EScrollIntoViewAlignment::IntoView:
+						if (IndexOfItem > MaxDisplayedIndex)
+						{
+							// Bring the new item in question to the bottom
+							NewScrollOffset -= (NumLiveWidgets - 1.0);
+						}
+
+						// The alternative is that IndexOfItem < MinDisplayedIndex, and NewScrollOffset is already correct for that case
+						break;
+					}
 
 					// Limit offset to top and bottom of the list.
 					const double MaxScrollOffset = FMath::Max(0.0, static_cast<double>(Items.Num()) - NumLiveWidgets);
@@ -2566,6 +2585,9 @@ protected:
 	/** When set, the list will notify this item when it has been scrolled into view */
 	NullableItemType ItemToNotifyWhenInView;
 
+	/** How to scroll an item into view */
+	EScrollIntoViewAlignment ScrollIntoViewAlignment;
+
 	/** Delegate to invoke when selection changes. */
 	FOnSelectionChanged OnSelectionChanged;
 
diff --git a/Engine/Source/Runtime/Slate/Public/Widgets/Views/STableViewBase.h b/Engine/Source/Runtime/Slate/Public/Widgets/Views/STableViewBase.h
--- a/Engine/Source/Runtime/Slate/Public/Widgets/Views/STableViewBase.h
+++ b/Engine/Source/Runtime/Slate/Public/Widgets/Views/STableViewBase.h
@@ -53,6 +53,17 @@ enum class EListItemAlignment : uint8
 	Fill,
 };
 
+/** How the list panel scrolls an offscreen item into view */
+UENUM(BlueprintType)
+enum class EScrollIntoViewAlignment : uint8
+{
+	/** Item will be aligned to the center of the view */
+	CenterAligned,
+
+	/** Item will only be scrolled enough to be brought into view */
+	IntoView,
+};
+
 DECLARE_DELEGATE_OneParam(
 	FOnTableViewScrolled,
 	double );	/** Scroll offset from the beginning of the list in items */
diff --git a/Engine/Source/Runtime/UMG/Public/Components/ListViewBase.h b/Engine/Source/Runtime/UMG/Public/Components/ListViewBase.h
--- a/Engine/Source/Runtime/UMG/Public/Components/ListViewBase.h
+++ b/Engine/Source/Runtime/UMG/Public/Components/ListViewBase.h
@@ -237,6 +237,7 @@ protected:
 		EConsumeMouseWheel ConsumeMouseWheel = EConsumeMouseWheel::WhenScrollingPossible;
 		bool bReturnFocusToSelection = false;
 		EOrientation Orientation = Orient_Vertical;
+		EScrollIntoViewAlignment ScrollIntoViewAlignment = EScrollIntoViewAlignment::CenterAligned;
 		const FTableViewStyle* ListViewStyle = &FUMGCoreStyle::Get().GetWidgetStyle<FTableViewStyle>("ListView");
 		const FScrollBarStyle* ScrollBarStyle = &FUMGCoreStyle::Get().GetWidgetStyle<FScrollBarStyle>("ScrollBar");
 		bool bPreventThrottling = false;
@@ -257,6 +258,7 @@ protected:
 			.SelectionMode(Args.SelectionMode)
 			.ReturnFocusToSelection(Args.bReturnFocusToSelection)
 			.Orientation(Args.Orientation)
+			.ScrollIntoViewAlignment(Args.ScrollIntoViewAlignment)
 			.ListViewStyle(Args.ListViewStyle)
 			.ScrollBarStyle(Args.ScrollBarStyle)
 			.PreventThrottling(Args.bPreventThrottling)
