Added EnableTouchScrolling property to ScrollBox widget.

diff --git a/Engine/Source/Runtime/Slate/Private/Widgets/Layout/SScrollBox.cpp b/Engine/Source/Runtime/Slate/Private/Widgets/Layout/SScrollBox.cpp
--- a/Engine/Source/Runtime/Slate/Private/Widgets/Layout/SScrollBox.cpp
+++ b/Engine/Source/Runtime/Slate/Private/Widgets/Layout/SScrollBox.cpp
@@ -125,6 +125,7 @@ void SScrollBox::Construct( const FArguments& InArgs )
 	FrontPadScrolling = InArgs._FrontPadScrolling;
 	bAnimateWheelScrolling = InArgs._AnimateWheelScrolling;
 	WheelScrollMultiplier = InArgs._WheelScrollMultiplier;
+	bEnableTouchScrolling = InArgs._EnableTouchScrolling;
 	NavigationScrollPadding = InArgs._NavigationScrollPadding;
 	NavigationDestination = InArgs._NavigationDestination;
 	ScrollWhenFocusChanges = InArgs._ScrollWhenFocusChanges;
@@ -681,7 +682,7 @@ bool SScrollBox::ComputeVolatility() const
 
 FReply SScrollBox::OnPreviewMouseButtonDown(const FGeometry& MyGeometry, const FPointerEvent& MouseEvent)
 {
-	if (MouseEvent.IsTouchEvent() && !bFingerOwningTouchInteraction.IsSet())
+	if (bEnableTouchScrolling && MouseEvent.IsTouchEvent() && !bFingerOwningTouchInteraction.IsSet())
 	{
 		// Clear any inertia 
 		InertialScrollManager.ClearScrollVelocity();
@@ -846,7 +847,7 @@ FReply SScrollBox::OnMouseMove( const FGeometry& MyGeometry, const FPointerEvent
 
 void SScrollBox::OnMouseEnter( const FGeometry& MyGeometry, const FPointerEvent& MouseEvent )
 {
-	if ( MouseEvent.IsTouchEvent() )
+	if (bEnableTouchScrolling && MouseEvent.IsTouchEvent())
 	{
 		if ( !bFingerOwningTouchInteraction.IsSet() )
 		{
@@ -1190,6 +1191,12 @@ void SScrollBox::SetWheelScrollMultiplier(float NewWheelScrollMultiplier)
 	WheelScrollMultiplier = NewWheelScrollMultiplier;
 }
 
+void SScrollBox::SetIsTouchScrollingEnabled(const bool bInEnableTouchScrolling)
+{
+	bEnableTouchScrolling = bInEnableTouchScrolling;
+	ensureMsgf(!bFingerOwningTouchInteraction.IsSet(), TEXT("TouchScrollingEnabled flag should not be changed while scrolling."));
+}
+
 void SScrollBox::SetScrollWhenFocusChanges(EScrollWhenFocusChanges NewScrollWhenFocusChanges)
 {
 	ScrollWhenFocusChanges = NewScrollWhenFocusChanges;
diff --git a/Engine/Source/Runtime/Slate/Public/Widgets/Layout/SScrollBox.h b/Engine/Source/Runtime/Slate/Public/Widgets/Layout/SScrollBox.h
--- a/Engine/Source/Runtime/Slate/Public/Widgets/Layout/SScrollBox.h
+++ b/Engine/Source/Runtime/Slate/Public/Widgets/Layout/SScrollBox.h
@@ -248,6 +248,7 @@ public:
 		, _AnimateWheelScrolling(false)
 		, _ScrollAnimationInterpSpeed(15.f)
 		, _WheelScrollMultiplier(1.f)
+		, _EnableTouchScrolling(true)
 		, _NavigationDestination(EDescendantScrollDestination::IntoView)
 		, _NavigationScrollPadding(0.0f)
 		, _ScrollWhenFocusChanges(EScrollWhenFocusChanges::NoScroll)
@@ -295,6 +296,8 @@ public:
 
 		SLATE_ARGUMENT(float, WheelScrollMultiplier);
 
+		SLATE_ARGUMENT(bool, EnableTouchScrolling);
+
 		SLATE_ARGUMENT(EDescendantScrollDestination, NavigationDestination);
 
 		/**
@@ -345,6 +348,9 @@ public:
 	SLATE_API void SetScrollingAnimationInterpolationSpeed(float NewScrollingAnimationInterpolationSpeed);
 
 	SLATE_API void SetWheelScrollMultiplier(float NewWheelScrollMultiplier);
+
+	/** Enables/disables being able to scroll using touch input. */
+	SLATE_API void SetIsTouchScrollingEnabled(const bool bInEnableTouchScrolling);
 	
 	SLATE_API void SetScrollWhenFocusChanges(EScrollWhenFocusChanges NewScrollWhenFocusChanges);
 
@@ -585,6 +591,9 @@ protected:
 	/** Multiplier applied to each click of the scroll wheel (applied alongside the global scroll amount) */
 	float WheelScrollMultiplier = 1.f;
 
+	/** True to allow scrolling by using touch input. */
+	bool bEnableTouchScrolling = true;
+
 	/** The speed of interpolation for the scrolling animation */
 	float ScrollingAnimationInterpolationSpeed = 15.f;
 
diff --git a/Engine/Source/Runtime/UMG/Private/Components/ScrollBox.cpp b/Engine/Source/Runtime/UMG/Private/Components/ScrollBox.cpp
--- a/Engine/Source/Runtime/UMG/Private/Components/ScrollBox.cpp
+++ b/Engine/Source/Runtime/UMG/Private/Components/ScrollBox.cpp
@@ -104,6 +104,7 @@ PRAGMA_DISABLE_DEPRECATION_WARNINGS
 		.AnimateWheelScrolling(bAnimateWheelScrolling)
 		.ScrollAnimationInterpSpeed(ScrollAnimationInterpolationSpeed)
 		.WheelScrollMultiplier(WheelScrollMultiplier)
+		.EnableTouchScrolling(bEnableTouchScrolling)
 		.OnUserScrolled(BIND_UOBJECT_DELEGATE(FOnUserScrolled, SlateHandleUserScrolled))
 		.OnScrollBarVisibilityChanged(BIND_UOBJECT_DELEGATE(FOnScrollBarVisibilityChanged, SlateHandleScrollBarVisibilityChanged));
 PRAGMA_ENABLE_DEPRECATION_WARNINGS
@@ -142,6 +143,7 @@ PRAGMA_DISABLE_DEPRECATION_WARNINGS
 	MyScrollBox->SetAnimateWheelScrolling(bAnimateWheelScrolling);
 	MyScrollBox->SetScrollingAnimationInterpolationSpeed(ScrollAnimationInterpolationSpeed);
 	MyScrollBox->SetWheelScrollMultiplier(WheelScrollMultiplier);
+	MyScrollBox->SetIsTouchScrollingEnabled(bEnableTouchScrolling);
 	MyScrollBox->SetStyle(&WidgetStyle);
 	MyScrollBox->InvalidateStyle();
 	MyScrollBox->SetScrollBarStyle(&WidgetBarStyle);
@@ -459,6 +461,20 @@ float UScrollBox::GetWheelScrollMultiplier() const
 	return WheelScrollMultiplier;
 }
 
+void UScrollBox::SetIsTouchScrollingEnabled(bool bInEnableTouchScrolling)
+{
+	bEnableTouchScrolling = bInEnableTouchScrolling;
+	if (MyScrollBox)
+	{
+		MyScrollBox->SetIsTouchScrollingEnabled(bInEnableTouchScrolling);
+	}
+}
+
+bool UScrollBox::GetIsTouchScrollingEnabled() const
+{
+	return bEnableTouchScrolling;
+}
+
 void UScrollBox::SetScrollWhenFocusChanges(EScrollWhenFocusChanges NewScrollWhenFocusChanges)
 {
 	ScrollWhenFocusChanges = NewScrollWhenFocusChanges;
diff --git a/Engine/Source/Runtime/UMG/Public/Components/ScrollBox.h b/Engine/Source/Runtime/UMG/Public/Components/ScrollBox.h
--- a/Engine/Source/Runtime/UMG/Public/Components/ScrollBox.h
+++ b/Engine/Source/Runtime/UMG/Public/Components/ScrollBox.h
@@ -27,6 +27,10 @@ private:
 	UPROPERTY(EditAnywhere, BlueprintReadWrite, Getter, Setter, Category = "Scroll", meta = (AllowPrivateAccess = "true", DisplayAfter = "bAnimateWheelScrolling"))
 	float ScrollAnimationInterpolationSpeed = 15.f;
 
+	/** True to allow scrolling using touch input. */
+	UPROPERTY(EditAnywhere, BlueprintReadWrite, Getter="GetIsTouchScrollingEnabled", Setter="SetIsTouchScrollingEnabled", Category = "Scroll", meta = (AllowPrivateAccess = "true", DisplayAfter = "WheelScrollMultiplier"))
+	bool bEnableTouchScrolling = true;
+
 public:
 
 	UE_DEPRECATED(5.2, "Direct access to WidgetStyle is deprecated. Please use the getter or setter.")
@@ -180,6 +184,11 @@ public:
 
 	UMG_API float GetWheelScrollMultiplier() const;
 
+	UFUNCTION(Category = "Scroll")
+	UMG_API void SetIsTouchScrollingEnabled(bool bInEnableTouchScrolling);
+
+	UMG_API bool GetIsTouchScrollingEnabled() const;
+
 	UFUNCTION(BlueprintCallable, Category = "Scroll")
 	UMG_API void SetScrollWhenFocusChanges(EScrollWhenFocusChanges NewScrollWhenFocusChanges);
 
