Exposed ScrollIntoViewAlignment of ListView to blueprints. Added support for TopOrLeft and BottomOrRight options.

diff --git a/Engine/Source/Runtime/Slate/Public/Widgets/Views/SListView.h b/Engine/Source/Runtime/Slate/Public/Widgets/Views/SListView.h
--- a/Engine/Source/Runtime/Slate/Public/Widgets/Views/SListView.h
+++ b/Engine/Source/Runtime/Slate/Public/Widgets/Views/SListView.h
@@ -2022,6 +2022,12 @@ public:
 		}
 	}
 
+	/* Sets ScrollIntoViewAlignment which allows to stick the selected item to either side or center */
+	void SetScrollIntoViewAlignment (EScrollIntoViewAlignment NewScrollIntoViewAlignment)
+	{
+		ScrollIntoViewAlignment = NewScrollIntoViewAlignment;
+	}
+
 	/**
 	 * Find a widget for this item if it has already been constructed.
 	 *
@@ -2120,7 +2126,9 @@ protected:
 				// When navigating, we don't want to scroll partially visible existing rows all the way to the center, so we count partially displayed indices in the displayed range
 				const double MinDisplayedIndex = bNavigateOnScrollIntoView ? FMath::FloorToDouble(CurrentScrollOffset) : FMath::CeilToDouble(CurrentScrollOffset);
 				const double MaxDisplayedIndex = bNavigateOnScrollIntoView ? FMath::CeilToDouble(CurrentScrollOffset + NumFullEntriesInView) : FMath::FloorToDouble(CurrentScrollOffset + NumFullEntriesInView);
-				if (IndexOfItem < MinDisplayedIndex || IndexOfItem > MaxDisplayedIndex)
+				if (IndexOfItem < MinDisplayedIndex || IndexOfItem > MaxDisplayedIndex || 
+					ScrollIntoViewAlignment == EScrollIntoViewAlignment::TopOrLeft || 
+					ScrollIntoViewAlignment == EScrollIntoViewAlignment::BottomOrRight)
 				{
 					// Scroll the top of the listview to the item in question
 					double NewScrollOffset = IndexOfItem;
@@ -2141,6 +2149,16 @@ protected:
 
 						// The alternative is that IndexOfItem < MinDisplayedIndex, and NewScrollOffset is already correct for that case
 						break;
+
+					case EScrollIntoViewAlignment::TopOrLeft:
+						// Set NewScrollOffset to exactly IndexOfItem to ensure it's at the top/left
+						NewScrollOffset = IndexOfItem;
+						break;
+
+					case EScrollIntoViewAlignment::BottomOrRight:
+						// Set NewScrollOffset to position the item at the bottom/right
+						NewScrollOffset = IndexOfItem - (NumLiveWidgets - 1.0);
+						break;
 					}
 
 					// Limit offset to top and bottom of the list.
diff --git a/Engine/Source/Runtime/Slate/Public/Widgets/Views/STableViewBase.h b/Engine/Source/Runtime/Slate/Public/Widgets/Views/STableViewBase.h
--- a/Engine/Source/Runtime/Slate/Public/Widgets/Views/STableViewBase.h
+++ b/Engine/Source/Runtime/Slate/Public/Widgets/Views/STableViewBase.h
@@ -57,11 +57,17 @@ enum class EListItemAlignment : uint8
 UENUM(BlueprintType)
 enum class EScrollIntoViewAlignment : uint8
 {
+	/** Item will only be scrolled enough to be brought into view */
+	IntoView,
+
+	/** Always scroll the widget so it appears at the top/Left of the scrollable area. */
+	TopOrLeft,
+
 	/** Item will be aligned to the center of the view */
 	CenterAligned,
 
-	/** Item will only be scrolled enough to be brought into view */
-	IntoView,
+	/** Always scroll the widget so it appears at the bottom/Right of the scrollable area. */
+	BottomOrRight
 };
 
 DECLARE_DELEGATE_OneParam(
diff --git a/Engine/Source/Runtime/UMG/Private/Components/ListView.cpp b/Engine/Source/Runtime/UMG/Private/Components/ListView.cpp
--- a/Engine/Source/Runtime/UMG/Private/Components/ListView.cpp
+++ b/Engine/Source/Runtime/UMG/Private/Components/ListView.cpp
@@ -116,6 +116,16 @@ void UListView::SetSelectionMode(TEnumAsByte<ESelectionMode::Type> InSelectionMo
 	}
 }
 
+void UListView::SetScrollIntoViewAlignment (EScrollIntoViewAlignment NewScrollIntoViewAlignment)
+{
+	ScrollIntoViewAlignment = NewScrollIntoViewAlignment;
+	if (MyListView)
+	{
+		MyListView->SetScrollIntoViewAlignment(NewScrollIntoViewAlignment);
+	}
+}
+
+
 int32 UListView::BP_GetNumItemsSelected() const
 {
 	return GetNumItemsSelected();
diff --git a/Engine/Source/Runtime/UMG/Public/Components/ListView.h b/Engine/Source/Runtime/UMG/Public/Components/ListView.h
--- a/Engine/Source/Runtime/UMG/Public/Components/ListView.h
+++ b/Engine/Source/Runtime/UMG/Public/Components/ListView.h
@@ -106,6 +106,10 @@ public:
 	/** Sets the new selection mode, preserving the current selection where possible. */
 	UFUNCTION(BlueprintCallable, Category = ListView)
 	UMG_API void SetSelectionMode(TEnumAsByte<ESelectionMode::Type> SelectionMode);
+	
+	/* Sets ScrollIntoViewAlignment which allows to stick the selected item to either side or center */
+	UFUNCTION(BlueprintCallable, Category = ListView)
+	UMG_API void SetScrollIntoViewAlignment (EScrollIntoViewAlignment NewScrollIntoViewAlignment);
 
 	/** Returns true if a refresh is pending and the list will be rebuilt on the next tick */
 	UFUNCTION(BlueprintCallable, Category = ListView)
@@ -166,6 +170,7 @@ protected:
 		Args.Orientation = Orientation;
 		Args.ListViewStyle = &WidgetStyle;
 		Args.ScrollBarStyle = &ScrollBarStyle;
+		Args.ScrollIntoViewAlignment = ScrollIntoViewAlignment;
 		MyListView = ITypedUMGListView<UObject*>::ConstructListView<ListViewT>(this, ListItems, Args);
 		
 		MyListView->SetOnEntryInitialized(SListView<UObject*>::FOnEntryInitialized::CreateUObject(this, &UListView::HandleOnEntryInitializedInternal));
@@ -203,6 +208,10 @@ protected:
 	UPROPERTY(EditAnywhere, BlueprintReadOnly, Category = ListView)
 	bool bReturnFocusToSelection = false;
 
+	/** Sets where to scroll a widget to when using explicit navigation */
+	UPROPERTY(EditAnywhere, BlueprintReadOnly, Category = Scrolling)
+	EScrollIntoViewAlignment ScrollIntoViewAlignment = EScrollIntoViewAlignment::CenterAligned;
+
 	UPROPERTY(Transient)
 	TArray<TObjectPtr<UObject>> ListItems;
 
