SSafeZone: when not in editor update the safe margin regardless of the validity of the GameViewport

#rnx

diff --git a/Engine/Source/Runtime/Slate/Private/Widgets/Layout/SSafeZone.cpp b/Engine/Source/Runtime/Slate/Private/Widgets/Layout/SSafeZone.cpp
--- a/Engine/Source/Runtime/Slate/Private/Widgets/Layout/SSafeZone.cpp
+++ b/Engine/Source/Runtime/Slate/Private/Widgets/Layout/SSafeZone.cpp
@@ -101,7 +101,6 @@ void SSafeZone::UpdateSafeMargin()
 		FSlateApplication::Get().GetSafeZoneSize(SafeMargin, OverrideScreenSize.GetValue());
 	}
 	else
-#endif
 	{
 		// Need to get owning viewport not display 
 		// use pixel values (same as custom safe zone above)
@@ -124,6 +123,10 @@ void SSafeZone::UpdateSafeMargin()
 			return;
 		}
 	}
+#else
+	// The second parameter of GetSafeZoneSize is only used #if WITH_EDITOR
+	FSlateApplication::Get().GetSafeZoneSize(SafeMargin, FDeprecateSlateVector2D());
+#endif
 
 	TOptional<float> SafeZoneScale = GetGlobalSafeZoneScale();
 	if (SafeZoneScale.IsSet())
