Added Scrollbar Padding property to the ListView for parity with the ScrollBox. WheelScrollMultiplier = 0 completely disables scrolling with MouseWheel.

[RN] Added Scrollbar Padding property to the ListView for parity with the ScrollBox. WheelScrollMultiplier = 0 completely disables scrolling with MouseWheel..

diff --git a/Engine/Source/Runtime/Slate/Private/Widgets/Views/STableViewBase.cpp b/Engine/Source/Runtime/Slate/Private/Widgets/Views/STableViewBase.cpp
--- a/Engine/Source/Runtime/Slate/Private/Widgets/Views/STableViewBase.cpp
+++ b/Engine/Source/Runtime/Slate/Private/Widgets/Views/STableViewBase.cpp
@@ -101,6 +101,7 @@ void STableViewBase::ConstructChildren( const TAttribute<float>& InItemWidth, co
 
 		if (Orientation == Orient_Vertical)
 		{
+			VerticalScrollBarSlot = nullptr;
 			ListAndScrollbar = SNew(SHorizontalBox)
 				+SHorizontalBox::Slot()
 				.FillWidth(1)
@@ -119,6 +120,8 @@ void STableViewBase::ConstructChildren( const TAttribute<float>& InItemWidth, co
 				]
 				+SHorizontalBox::Slot()
 				.AutoWidth()
+				.Padding(ScrollBarSlotPadding)
+				.Expose(VerticalScrollBarSlot)
 				[
 					SNew(SBox)
 					.WidthOverride(ScrollBarSize)
@@ -129,6 +132,7 @@ void STableViewBase::ConstructChildren( const TAttribute<float>& InItemWidth, co
 		}
 		else
 		{
+			HorizontalScrollBarSlot = nullptr;
 			ListAndScrollbar = SNew(SVerticalBox)
 				+SVerticalBox::Slot()
 				.FillHeight(1)
@@ -147,6 +151,8 @@ void STableViewBase::ConstructChildren( const TAttribute<float>& InItemWidth, co
 				]
 				+SVerticalBox::Slot()
 				.AutoHeight()
+				.Padding(ScrollBarSlotPadding)
+				.Expose(HorizontalScrollBarSlot)
 				[
 					SNew(SBox)
 					.HeightOverride(ScrollBarSize)
@@ -582,6 +588,11 @@ void STableViewBase::OnMouseLeave( const FPointerEvent& MouseEvent )
 
 FReply STableViewBase::OnMouseWheel( const FGeometry& MyGeometry, const FPointerEvent& MouseEvent )
 {
+	if (WheelScrollMultiplier == 0.0f)
+    {
+        return FReply::Unhandled();
+    }
+
 	if( bIsPointerScrollingEnabled && !MouseEvent.IsControlDown() )
 	{
 		// Make sure scroll velocity is cleared so it doesn't fight with the mouse wheel input
@@ -890,6 +901,26 @@ void STableViewBase::SetScrollbarVisibility(const EVisibility InVisibility)
 	}
 }
 
+void STableViewBase::SetScrollbarPadding(const FMargin& InScrollbarPadding)
+{
+	ScrollBarSlotPadding = InScrollbarPadding;
+
+	if (Orientation == Orient_Vertical)
+	{
+		if (VerticalScrollBarSlot)
+		{
+			VerticalScrollBarSlot->SetPadding(ScrollBarSlotPadding);
+		}
+	}
+	else
+	{
+		if (HorizontalScrollBarSlot)
+		{
+			HorizontalScrollBarSlot->SetPadding(ScrollBarSlotPadding);
+		}
+	}
+}
+
 EVisibility STableViewBase::GetScrollbarVisibility() const
 {
 	return ScrollBar ? ScrollBar->ShouldBeVisible() : EVisibility::Collapsed;
diff --git a/Engine/Source/Runtime/Slate/Public/Widgets/Views/SListView.h b/Engine/Source/Runtime/Slate/Public/Widgets/Views/SListView.h
--- a/Engine/Source/Runtime/Slate/Public/Widgets/Views/SListView.h
+++ b/Engine/Source/Runtime/Slate/Public/Widgets/Views/SListView.h
@@ -110,6 +110,7 @@ public:
 		, _ConsumeMouseWheel(EConsumeMouseWheel::WhenScrollingPossible)
 		, _WheelScrollMultiplier(GetGlobalScrollAmount())
 		, _NavigationScrollOffset(0.5f)
+		, _ScrollBarPadding(0.0f)
 		, _HandleGamepadEvents( true )
 		, _HandleDirectionalNavigation( true )
 		, _HandleSpacebarSelection(false)
@@ -187,6 +188,8 @@ public:
 
 		SLATE_ARGUMENT( float, NavigationScrollOffset );
 
+		SLATE_ARGUMENT( FMargin, ScrollBarPadding );
+
 		SLATE_ARGUMENT( bool, HandleGamepadEvents );
 
 		SLATE_ARGUMENT( bool, HandleDirectionalNavigation );
@@ -261,6 +264,7 @@ public:
 
 		this->MaxPinnedItems = InArgs._MaxPinnedItems;
 		this->DefaultMaxPinnedItems = InArgs._MaxPinnedItems;
+		this->ScrollBarSlotPadding = InArgs._ScrollBarPadding;
 
 		// Check for any parameters that the coder forgot to specify.
 		FString ErrorString;
diff --git a/Engine/Source/Runtime/Slate/Public/Widgets/Views/STableViewBase.h b/Engine/Source/Runtime/Slate/Public/Widgets/Views/STableViewBase.h
--- a/Engine/Source/Runtime/Slate/Public/Widgets/Views/STableViewBase.h
+++ b/Engine/Source/Runtime/Slate/Public/Widgets/Views/STableViewBase.h
@@ -14,6 +14,7 @@
 #include "Framework/Layout/InertialScrollManager.h"
 #include "Framework/Layout/Overscroll.h"
 #include "Styling/SlateTypes.h"
+#include "Widgets/SBoxPanel.h"
 
 #include "STableViewBase.generated.h"
 
@@ -181,6 +182,8 @@ public:
 
 	SLATE_API void SetScrollbarVisibility(const EVisibility InVisibility);
 
+	SLATE_API void SetScrollbarPadding(const FMargin& InScrollbarPadding);
+
 	/** Returns true if scrolling is possible; false if the view is big enough to fit all the content. */
 	SLATE_API bool IsScrollbarNeeded() const;
 
@@ -431,6 +434,9 @@ protected:
 	/** The scroll bar widget */
 	TSharedPtr< SScrollBar > ScrollBar;
 
+	/** Padding to the scrollbox */
+	FMargin ScrollBarSlotPadding;
+
 	/** Delegate to call when the table view is scrolled */
 	FOnTableViewScrolled OnTableViewScrolled;
 
@@ -553,6 +559,13 @@ protected:
 private:
 	/** When true, a refresh should occur the next tick */
 	bool bItemsNeedRefresh = false;
+
+	union
+	{
+		// vertical scroll bar is stored in horizontal box and vice versa
+		SHorizontalBox::FSlot* VerticalScrollBarSlot; // valid when Orientation == Orient_Vertical
+		SVerticalBox::FSlot* HorizontalScrollBarSlot; // valid when Orientation == Orient_Horizontal
+	};
 };
 
 
diff --git a/Engine/Source/Runtime/UMG/Private/Components/ListView.cpp b/Engine/Source/Runtime/UMG/Private/Components/ListView.cpp
--- a/Engine/Source/Runtime/UMG/Private/Components/ListView.cpp
+++ b/Engine/Source/Runtime/UMG/Private/Components/ListView.cpp
@@ -125,6 +125,14 @@ void UListView::SetScrollIntoViewAlignment (EScrollIntoViewAlignment NewScrollIn
 	}
 }
 
+void UListView::SetScrollBarPadding(const FMargin& InScrollBarPadding)
+{
+	ScrollBarPadding = InScrollBarPadding;
+	if (MyListView.IsValid())
+	{
+		MyListView->SetScrollbarPadding(ScrollBarPadding);
+	}
+}
 
 int32 UListView::BP_GetNumItemsSelected() const
 {
diff --git a/Engine/Source/Runtime/UMG/Public/Components/ListView.h b/Engine/Source/Runtime/UMG/Public/Components/ListView.h
--- a/Engine/Source/Runtime/UMG/Public/Components/ListView.h
+++ b/Engine/Source/Runtime/UMG/Public/Components/ListView.h
@@ -111,6 +111,10 @@ public:
 	UFUNCTION(BlueprintCallable, Category = ListView)
 	UMG_API void SetScrollIntoViewAlignment (EScrollIntoViewAlignment NewScrollIntoViewAlignment);
 
+	/** Sets padding around the scrollbar. */
+	UFUNCTION(Category = ListView)
+	UMG_API void SetScrollBarPadding(const FMargin& InScrollBarPadding);
+
 	/** Returns true if a refresh is pending and the list will be rebuilt on the next tick */
 	UFUNCTION(BlueprintCallable, Category = ListView)
 	UMG_API bool IsRefreshPending() const;
@@ -170,6 +174,7 @@ protected:
 		Args.Orientation = Orientation;
 		Args.ListViewStyle = &WidgetStyle;
 		Args.ScrollBarStyle = &ScrollBarStyle;
+		Args.ScrollBarPadding = ScrollBarPadding;
 		Args.ScrollIntoViewAlignment = ScrollIntoViewAlignment;
 		MyListView = ITypedUMGListView<UObject*>::ConstructListView<ListViewT>(this, ListItems, Args);
 		
@@ -232,6 +237,9 @@ private:
 
 	UPROPERTY(EditAnywhere, BlueprintReadOnly, Getter, Category = ListView, meta = (ClampMin = 0, AllowPrivateAccess = "true"))
 	float VerticalEntrySpacing = 0.f;
+	
+	UPROPERTY(EditAnywhere, BlueprintReadWrite, Getter, Setter, Category = ListView, meta = (AllowPrivateAccess = "true", DisplayAfter="ScrollBarStyle"))
+	FMargin ScrollBarPadding;
 
 public:
 	/** Get the horizontal spacing between entries. */
@@ -242,6 +250,10 @@ public:
 	UFUNCTION(BlueprintCallable, Category = ListView)
 	float GetVerticalEntrySpacing() const { return VerticalEntrySpacing; }
 
+	/** Gets padding around the scrollbar. */
+	UFUNCTION(Category = ListView)
+	FMargin GetScrollBarPadding() const { return ScrollBarPadding; }
+
 protected:
 	UMG_API void InitHorizontalEntrySpacing(float InHorizontalEntrySpacing);
 	UMG_API void InitVerticalEntrySpacing(float InVerticalEntrySpacing);
diff --git a/Engine/Source/Runtime/UMG/Public/Components/ListViewBase.h b/Engine/Source/Runtime/UMG/Public/Components/ListViewBase.h
--- a/Engine/Source/Runtime/UMG/Public/Components/ListViewBase.h
+++ b/Engine/Source/Runtime/UMG/Public/Components/ListViewBase.h
@@ -240,6 +240,7 @@ protected:
 		EScrollIntoViewAlignment ScrollIntoViewAlignment = EScrollIntoViewAlignment::CenterAligned;
 		const FTableViewStyle* ListViewStyle = &FUMGCoreStyle::Get().GetWidgetStyle<FTableViewStyle>("ListView");
 		const FScrollBarStyle* ScrollBarStyle = &FUMGCoreStyle::Get().GetWidgetStyle<FScrollBarStyle>("ScrollBar");
+		FMargin ScrollBarPadding = FMargin(0.0f);
 		bool bPreventThrottling = false;
 	};
 
@@ -261,6 +262,7 @@ protected:
 			.ScrollIntoViewAlignment(Args.ScrollIntoViewAlignment)
 			.ListViewStyle(Args.ListViewStyle)
 			.ScrollBarStyle(Args.ScrollBarStyle)
+			.ScrollBarPadding(Args.ScrollBarPadding)
 			.PreventThrottling(Args.bPreventThrottling)
 			.OnGenerateRow_UObject(Implementer, &UListViewBaseT::HandleGenerateRow)
 			.OnSelectionChanged_UObject(Implementer, &UListViewBaseT::HandleSelectionChanged)
