Fixing crash when removing while iterating on console

diff --git a/Engine/Source/Runtime/Slate/Private/Framework/Application/SlateApplication.cpp b/Engine/Source/Runtime/Slate/Private/Framework/Application/SlateApplication.cpp
--- a/Engine/Source/Runtime/Slate/Private/Framework/Application/SlateApplication.cpp
+++ b/Engine/Source/Runtime/Slate/Private/Framework/Application/SlateApplication.cpp
@@ -7517,14 +7517,10 @@ void FSlateApplication::InputPreProcessorsHelper::Remove(TSharedPtr<IInputProces
 		InputPreProcessorsIteratorList.Remove(InputProcessor);
 	}
 
-	for (const FInputPreprocessorRegistration& Registration : ProcessorsPendingAddition)
-	{
-		if (Registration.InputProcessor == InputProcessor)
+	ProcessorsPendingAddition.RemoveAllSwap([InputProcessor](const FInputPreprocessorRegistration& Registration)
 		{
-			ProcessorsPendingAddition.Remove(Registration);
-			break;
-		}
-	}
+			return Registration.InputProcessor == InputProcessor;
+		});
 }
 
 void FSlateApplication::InputPreProcessorsHelper::RemoveAll()
