Face-lift Search Box icons

diff --git a/Engine/Source/Runtime/Slate/Private/Widgets/Input/SSearchBox.cpp b/Engine/Source/Runtime/Slate/Private/Widgets/Input/SSearchBox.cpp
--- a/Engine/Source/Runtime/Slate/Private/Widgets/Input/SSearchBox.cpp
+++ b/Engine/Source/Runtime/Slate/Private/Widgets/Input/SSearchBox.cpp
@@ -93,6 +93,7 @@ void SSearchBox::Construct( const FArguments& InArgs )
 				SNew(SImage)
 				.Image( &InArgs._Style->UpArrowImage )
 				.ColorAndOpacity( FSlateColor::UseForeground() )
+				.DesiredSizeOverride( InArgs._Style->ImageSizeOverride )
 			]
 		];
 		// Next result button
@@ -115,6 +116,7 @@ void SSearchBox::Construct( const FArguments& InArgs )
 				SNew(SImage)
 				.Image( &InArgs._Style->DownArrowImage )
 				.ColorAndOpacity( FSlateColor::UseForeground() )
+				.DesiredSizeOverride( InArgs._Style->ImageSizeOverride )
 			]
 		];
 	}
@@ -140,10 +142,10 @@ void SSearchBox::Construct( const FArguments& InArgs )
 			.Visibility(this, &SSearchBox::GetSearchGlassVisibility)
 			.Image(&InArgs._Style->GlassImage)
 			.ColorAndOpacity(FSlateColor::UseForeground())
+			.DesiredSizeOverride(InArgs._Style->ImageSizeOverride)
 		];
 
 	// Add an X to clear the search whenever there is some text typed into it
-	FMargin ClearButtonMargin = FMargin(2.f, 0.f);
 	Box->InsertSlot(ClearButtonIndex)
 	.AutoWidth()
 	.HAlign(HAlign_Center)
@@ -151,8 +153,8 @@ void SSearchBox::Construct( const FArguments& InArgs )
 	[
 		SNew(SButton)
 		.Visibility(this, &SSearchBox::GetXVisibility)
-		.ButtonStyle(FAppStyle::Get(), "HoverOnlyButton")
-		.ContentPadding(ClearButtonMargin)
+		.ButtonStyle(FAppStyle::Get(), "ThinHoverOnlyButton")
+		.ContentPadding(InArgs._Style->ImagePadding)
 		.HAlign(HAlign_Center)
 		.VAlign(VAlign_Center)
 		.OnClicked(this, &SSearchBox::OnClearSearch)
@@ -163,6 +165,7 @@ void SSearchBox::Construct( const FArguments& InArgs )
 			SNew(SImage)
 			.Image(&InArgs._Style->ClearImage)
 			.ColorAndOpacity(FSlateColor::UseForeground())
+			.DesiredSizeOverride(InArgs._Style->ImageSizeOverride)
 		]
 	];
 
diff --git a/Engine/Source/Runtime/SlateCore/Private/Styling/StarshipCoreStyle.cpp b/Engine/Source/Runtime/SlateCore/Private/Styling/StarshipCoreStyle.cpp
--- a/Engine/Source/Runtime/SlateCore/Private/Styling/StarshipCoreStyle.cpp
+++ b/Engine/Source/Runtime/SlateCore/Private/Styling/StarshipCoreStyle.cpp
@@ -1311,7 +1311,8 @@ void FStarshipCoreStyle::SetupTextStyles(TSharedRef<FStyle>& Style)
 			.SetDownArrowImage(IMAGE_BRUSH_SVG("Starship/Common/chevron-down", Icon16x16, FStyleColors::Foreground))
 			.SetGlassImage(IMAGE_BRUSH_SVG("Starship/Common/search", Icon16x16))
 			.SetClearImage(IMAGE_BRUSH_SVG("Starship/Common/close", Icon16x16))
-			.SetImagePadding(FMargin(3.f, 0.f, 0.f, 0.f))
+			.SetImagePadding(FMargin(4.f, 0.f, 0.f, 0.f))
+			.SetImageSizeOverride(FVector2D(14.0f))
 			.SetLeftAlignSearchResultButtons(true)
 			.SetLeftAlignGlassImageAndClearButton(true)
 		);
@@ -1452,6 +1453,13 @@ void FStarshipCoreStyle::SetupButtonStyles(TSharedRef<FStyle>& Style)
 			.SetPressedPadding(FMargin(2.0f, 0.f));
 
 		Style->Set("HoverOnlyButton", SimpleButtonIconOnly);
+
+		SimpleButtonIconOnly
+			.SetNormalPadding(0.0f)
+			.SetPressedPadding(0.0f);
+
+		// A HoverOnlyButton that has no padding at all
+		Style->Set("ThinHoverOnlyButton", SimpleButtonIconOnly);
 	}
 }
 
diff --git a/Engine/Source/Runtime/SlateCore/Public/Styling/SlateTypes.h b/Engine/Source/Runtime/SlateCore/Public/Styling/SlateTypes.h
--- a/Engine/Source/Runtime/SlateCore/Public/Styling/SlateTypes.h
+++ b/Engine/Source/Runtime/SlateCore/Public/Styling/SlateTypes.h
@@ -1281,6 +1281,11 @@ struct FSearchBoxStyle : public FSlateWidgetStyle
 	FMargin ImagePadding;
 	FSearchBoxStyle& SetImagePadding(const FMargin& InImagePadding){ ImagePadding = InImagePadding; return *this; }
 
+	/** Size of the images (auto size if not set) */
+	UPROPERTY(EditAnywhere, Category = Appearance)
+	TOptional<FVector2D> ImageSizeOverride;
+	FSearchBoxStyle& SetImageSizeOverride(const TOptional<FVector2D>& InImageSizeOverride){ ImageSizeOverride = InImageSizeOverride; return *this; }
+
 	/** If true, buttons appear to the left of the search text */
 	UPROPERTY(meta = (DeprecatedProperty, DeprecationMessage="Use LeftAlignSearchResultButtons and LeftAlignGlassImageAndClearButton instead"))
 	bool bLeftAlignButtons_DEPRECATED;
