Breadcrumb trail button content customization

diff --git a/Engine/Source/Runtime/Slate/Public/Widgets/Navigation/SBreadcrumbTrail.h b/Engine/Source/Runtime/Slate/Public/Widgets/Navigation/SBreadcrumbTrail.h
--- a/Engine/Source/Runtime/Slate/Public/Widgets/Navigation/SBreadcrumbTrail.h
+++ b/Engine/Source/Runtime/Slate/Public/Widgets/Navigation/SBreadcrumbTrail.h
@@ -75,6 +75,9 @@ public:
 	/** Callback for getting the menu content to be displayed when clicking on a crumb's delimiter arrow */
 	DECLARE_DELEGATE_RetVal_OneParam( TSharedRef< SWidget >, FGetCrumbMenuContent, const ItemType& /*CrumbData*/ );
 
+	/** Callback for customizing the crumb button content */
+	DECLARE_DELEGATE_RetVal_TwoParams( TSharedRef<SWidget>, FGetCrumbButtonContent, const ItemType& /*CrumbData*/, const FTextBlockStyle* InTextStyle );	
+
 	SLATE_BEGIN_ARGS( SBreadcrumbTrail )
 		: _ButtonStyle( &FCoreStyle::Get().GetWidgetStyle< FButtonStyle >( "BreadcrumbButton" ) )
 		, _TextStyle( &FCoreStyle::Get().GetWidgetStyle<FTextBlockStyle>("NormalText") )
@@ -117,6 +120,8 @@ public:
 
 		SLATE_EVENT( FGetCrumbMenuContent, GetCrumbMenuContent )
 
+		SLATE_EVENT( FGetCrumbButtonContent, GetCrumbButtonContent )
+
 	SLATE_END_ARGS()
 
 	/** Constructs this widget with InArgs */
@@ -133,6 +138,7 @@ public:
 		bHasStaticBreadcrumbs = InArgs._PersistentBreadcrumbs;
 		HasCrumbMenuContentCallback = InArgs._HasCrumbMenuContent;
 		GetCrumbMenuContentCallback = InArgs._GetCrumbMenuContent;
+		GetCrumbButtonContentCallback = InArgs._GetCrumbButtonContent;
 
 		NextValidCrumbID = 0;
 
@@ -156,6 +162,8 @@ public:
 		TSharedPtr<SVerticalBox> NewButtonBox;
 		TSharedPtr<SVerticalBox> NewDelimiterBox;
 
+		TSharedRef<SWidget> ContentWidget = GetCrumbButtonContentCallback.IsBound() ? GetCrumbButtonContentCallback.Execute(NewCrumbData, TextStyle) : SNullWidget::NullWidget; 
+
 		// Add the crumb button
 		CrumbBox->AddSlot()
 		[
@@ -171,7 +179,10 @@ public:
 				.TextStyle(TextStyle)
 				.Text(CrumbText)
 				.OnClicked( this, &SBreadcrumbTrail::CrumbButtonClicked, NextValidCrumbID )
-				// .ForegroundColor( FSlateColor::Uthis, &SBreadcrumbTrail::GetButtonForegroundColor, NextValidCrumbID )
+				.Content()
+				[
+					ContentWidget
+				]
 			]
 		];
 
@@ -465,6 +476,9 @@ private:
 	/** Delegate to invoke to retrieve the content for a crumb's menu */
 	FGetCrumbMenuContent GetCrumbMenuContentCallback;
 
+	/** Delegate to invoke to retrieve content a crumb's SButton */
+	FGetCrumbButtonContent GetCrumbButtonContentCallback;
+
 	/** If true, a leading delimiter will be added */
 	TAttribute<bool> ShowLeadingDelimiter;
 
