Handle tab/shift+tab navigation as a special case

diff --git a/Engine/Source/Runtime/Slate/Private/Framework/Application/NavigationConfig.cpp b/Engine/Source/Runtime/Slate/Private/Framework/Application/NavigationConfig.cpp
--- a/Engine/Source/Runtime/Slate/Private/Framework/Application/NavigationConfig.cpp
+++ b/Engine/Source/Runtime/Slate/Private/Framework/Application/NavigationConfig.cpp
@@ -58,7 +58,8 @@ void FNavigationConfig::OnUserRemoved(int32 UserIndex)
 
 EUINavigation FNavigationConfig::GetNavigationDirectionFromKey(const FKeyEvent& InKeyEvent) const
 {
-	const bool bModifierHeld = InKeyEvent.IsControlDown() || InKeyEvent.IsAltDown() || InKeyEvent.IsCommandDown() || InKeyEvent.IsShiftDown();
+	const bool bIsShiftDown = InKeyEvent.IsShiftDown();
+	const bool bModifierHeld = InKeyEvent.IsControlDown() || InKeyEvent.IsAltDown() || InKeyEvent.IsCommandDown() || bIsShiftDown;
 	if (bIgnoreModifiersForNavigationActions || !bModifierHeld)
 	{
 		if (const EUINavigation* Rule = KeyEventRules.Find(InKeyEvent.GetKey()))
@@ -68,11 +69,12 @@ EUINavigation FNavigationConfig::GetNavigationDirectionFromKey(const FKeyEvent&
 				return *Rule;
 			}
 		}
-		else if (bTabNavigation && InKeyEvent.GetKey() == EKeys::Tab )
-		{
-			//@TODO: Really these uses of input should be at a lower priority, only occurring if nothing else handled them
-			return ( InKeyEvent.IsShiftDown() ) ? EUINavigation::Previous : EUINavigation::Next;
-		}
+	}
+
+	if (bTabNavigation && InKeyEvent.GetKey() == EKeys::Tab )
+	{
+		//@TODO: Really these uses of input should be at a lower priority, only occurring if nothing else handled them
+		return bIsShiftDown ? EUINavigation::Previous : EUINavigation::Next;
 	}
 
 	return EUINavigation::Invalid;
