UE-218896 - Can delete a member while renaming by backspacing the member's name in the Details panel

diff --git a/Engine/Source/Runtime/Slate/Private/Widgets/Text/SInlineEditableTextBlock.cpp b/Engine/Source/Runtime/Slate/Private/Widgets/Text/SInlineEditableTextBlock.cpp
--- a/Engine/Source/Runtime/Slate/Private/Widgets/Text/SInlineEditableTextBlock.cpp
+++ b/Engine/Source/Runtime/Slate/Private/Widgets/Text/SInlineEditableTextBlock.cpp
@@ -142,16 +142,25 @@ void SInlineEditableTextBlock::ExitEditingMode()
 	// Clear the error so it will vanish.
 	SetTextBoxError( FText::GetEmpty() );
 
-	// Restore the original widget focus
-	TSharedPtr<SWidget> WidgetToFocusPin = WidgetToFocus.Pin();
-	if(WidgetToFocusPin.IsValid())
+	TSharedPtr<SWidget> CurrentFocus = FSlateApplication::Get().GetKeyboardFocusedWidget();
+	if (CurrentFocus.IsValid())
 	{
-		FSlateApplication::Get().SetKeyboardFocus(WidgetToFocusPin, EFocusCause::SetDirectly);
-	}
-	else
-	{
-		FSlateApplication::Get().ClearKeyboardFocus(EFocusCause::SetDirectly);
+		if (CurrentFocus == GetEditableTextWidget())
+		{
+			// Restore the original widget focus
+			TSharedPtr<SWidget> WidgetToFocusPin = WidgetToFocus.Pin();
+			if (WidgetToFocusPin.IsValid())
+			{
+				FSlateApplication::Get().SetKeyboardFocus(WidgetToFocusPin, EFocusCause::SetDirectly);
+			}
+			else
+			{
+				FSlateApplication::Get().ClearKeyboardFocus(EFocusCause::SetDirectly);
+			}
+		}
 	}
+
+	WidgetToFocus.Reset();
 }
 
 bool SInlineEditableTextBlock::IsInEditMode() const
