[Editor] Changed WrapButton logic for toolbars to allow properly right-aligned widgets

diff --git a/Engine/Source/Runtime/Slate/Private/Framework/MultiBox/SClippingHorizontalBox.cpp b/Engine/Source/Runtime/Slate/Private/Framework/MultiBox/SClippingHorizontalBox.cpp
--- a/Engine/Source/Runtime/Slate/Private/Framework/MultiBox/SClippingHorizontalBox.cpp
+++ b/Engine/Source/Runtime/Slate/Private/Framework/MultiBox/SClippingHorizontalBox.cpp
@@ -41,16 +41,22 @@ void SClippingHorizontalBox::OnArrangeChildren( const FGeometry& AllottedGeometr
 
 	if (IndexClippedAt == NumChildren)
 	{
-		// Note do not increment NumClippedChildrenn here. This is to remove the wrap button
+		// Note do not increment NumClippedChildren here. This is to remove the wrap button
 		// None of the children are being clipped, so remove the wrap button
 		ArrangedChildren.Remove(ArrangedChildren.Num() - 1);
 	}
-	else
+	else if (ArrangedChildren.Num() > 0)
 	{
-		// Right align the wrap button
-		FArrangedWidget& ArrangedButton = ArrangedChildren[ArrangedChildren.Num() - 1];
-		ArrangedButton.Geometry = AllottedGeometry.MakeChild(ArrangedButton.Geometry.GetLocalSize(), FSlateLayoutTransform(AllottedGeometry.GetLocalSize() - ArrangedButton.Geometry.GetLocalSize()));
-		const int32 WrapButtonXPosition = FMath::TruncToInt(ArrangedButton.Geometry.AbsolutePosition.X);
+		// Insert and right align the wrap button, for occlusion checking and insertion below
+		FArrangedWidget& ArrangedWrapButton = ArrangedChildren[ArrangedChildren.Num() - 1];
+		FGeometry& WrapButtonGeometry = ArrangedWrapButton.Geometry;
+
+		const FVector2D WrapButtonSize = FVector2D(WrapButtonWidth, WrapButtonGeometry.GetLocalSize().Y);
+		WrapButtonGeometry = AllottedGeometry.MakeChild(
+			WrapButtonSize,
+			FSlateLayoutTransform(AllottedGeometry.GetLocalSize() - WrapButtonSize));
+
+		const int32 WrapButtonXPosition = FMath::TruncToInt(WrapButtonGeometry.AbsolutePosition.X);
 
 		// Further remove any children that the wrap button overlaps with
 		for (int32 ChildIdx = IndexClippedAt - 1; ChildIdx >= 0; --ChildIdx)
@@ -116,9 +122,13 @@ void SClippingHorizontalBox::AddWrapButton()
 			.Image(&ToolBarStyle.ExpandBrush)
 		];
 
+	// Perform a prepass to get a valid DesiredSize value below
+	WrapButton->SlatePrepass(1.0f);
+	WrapButtonWidth = WrapButton->GetDesiredSize().X;
+
 	// Add the wrap button
 	AddSlot()
-	.AutoWidth()
+	.FillWidth(0.0f) // Effectively makes this widget 0 width, so it exists as a slot/child, but isn't considered for layout
 	.Padding( 0.f )
 	[
 		WrapButton.ToSharedRef()
@@ -149,4 +159,3 @@ EActiveTimerReturnType SClippingHorizontalBox::UpdateWrapButtonStatus(double Cur
 
 	return EActiveTimerReturnType::Continue;
 }
-
diff --git a/Engine/Source/Runtime/Slate/Private/Framework/MultiBox/SClippingHorizontalBox.h b/Engine/Source/Runtime/Slate/Private/Framework/MultiBox/SClippingHorizontalBox.h
--- a/Engine/Source/Runtime/Slate/Private/Framework/MultiBox/SClippingHorizontalBox.h
+++ b/Engine/Source/Runtime/Slate/Private/Framework/MultiBox/SClippingHorizontalBox.h
@@ -20,9 +20,7 @@ public:
 	SLATE_BEGIN_ARGS(SClippingHorizontalBox) 
 		: _StyleSet(&FCoreStyle::Get())
 		, _StyleName(NAME_None)
-		, _IsFocusable(true)
-		{ }
-
+		, _IsFocusable(true) { }
 		SLATE_ARGUMENT(FOnGetContent, OnWrapButtonClicked)
 		SLATE_ARGUMENT(const ISlateStyle*, StyleSet)
 		SLATE_ARGUMENT(FName, StyleName)
@@ -48,6 +46,7 @@ public:
 private:
 	void OnWrapButtonOpenChanged(bool bIsOpen);
 	EActiveTimerReturnType UpdateWrapButtonStatus(double CurrentTime, float DeltaTime);
+
 private:
 	/** The button that is displayed when a toolbar or menubar is clipped */
 	TSharedPtr<SComboButton> WrapButton;
@@ -68,6 +67,8 @@ private:
 	/** Can the wrap button be focused? */
 	bool bIsFocusable;
 
+	float WrapButtonWidth = 0.0f; // Fixed width, initialized after button creation
+
 	/** The style to use */
 	const ISlateStyle* StyleSet;
 	FName StyleName;
