Updated FAnalogCursor to reference a static empty keyset (FTouchKeySet::EmptySet) instead of creating a new one that may be dereferenced later after becoming invalid

diff --git a/Engine/Source/Runtime/Slate/Private/Framework/Application/AnalogCursor.cpp b/Engine/Source/Runtime/Slate/Private/Framework/Application/AnalogCursor.cpp
--- a/Engine/Source/Runtime/Slate/Private/Framework/Application/AnalogCursor.cpp
+++ b/Engine/Source/Runtime/Slate/Private/Framework/Application/AnalogCursor.cpp
@@ -61,7 +61,7 @@ bool FAnalogCursor::HandleKeyDownEvent(FSlateApplication& SlateApp, const FKeyEv
 						FSlateApplication::CursorPointerIndex,
 						SlateUser->GetCursorPosition(),
 						SlateUser->GetPreviousCursorPosition(),
-						bIsPrimaryUser ? SlateApp.GetPressedMouseButtons() : TSet<FKey>(),
+						bIsPrimaryUser ? SlateApp.GetPressedMouseButtons() : FTouchKeySet::EmptySet,
 						EKeys::LeftMouseButton,
 						0,
 						bIsPrimaryUser ? SlateApp.GetModifierKeys() : FModifierKeysState()
@@ -101,13 +101,12 @@ bool FAnalogCursor::HandleKeyUpEvent(FSlateApplication& SlateApp, const FKeyEven
 			{
 				const bool bIsPrimaryUser = FSlateApplication::CursorUserIndex == SlateUser->GetUserIndex();
 
-				TSet<FKey> EmptySet;
 				FPointerEvent MouseEvent(
 					SlateUser->GetUserIndex(),
 					FSlateApplication::CursorPointerIndex,
 					SlateUser->GetCursorPosition(),
 					SlateUser->GetPreviousCursorPosition(),
-					bIsPrimaryUser ? SlateApp.GetPressedMouseButtons() : EmptySet,
+					bIsPrimaryUser ? SlateApp.GetPressedMouseButtons() : FTouchKeySet::EmptySet,
 					EKeys::LeftMouseButton,
 					0,
 					bIsPrimaryUser ? SlateApp.GetModifierKeys() : FModifierKeysState()
@@ -238,7 +237,7 @@ void FAnalogCursor::UpdateCursorPosition(FSlateApplication& SlateApp, TSharedRef
 				SlateApp.CursorPointerIndex,
 				UpdatedPosition,
 				OldPosition,
-				bIsPrimaryUser ? SlateApp.GetPressedMouseButtons() : TSet<FKey>(),
+				bIsPrimaryUser ? SlateApp.GetPressedMouseButtons() : FTouchKeySet::EmptySet,
 				EKeys::Invalid,
 				0,
 				bIsPrimaryUser ? SlateApp.GetModifierKeys() : FModifierKeysState()
@@ -273,7 +272,7 @@ void FAnalogCursor::UpdateCursorPosition(FSlateApplication& SlateApp, TSharedRef
 			FSlateApplication::CursorPointerIndex,
 			UpdatedPosition,
 			OldPosition,
-			bIsPrimaryUser ? SlateApp.GetPressedMouseButtons() : TSet<FKey>(),
+			bIsPrimaryUser ? SlateApp.GetPressedMouseButtons() : FTouchKeySet::EmptySet,
 			EKeys::Invalid,
 			0,
 			bIsPrimaryUser ? SlateApp.GetModifierKeys() : FModifierKeysState()
