Added tooltip to MultiBoxBuilder's AddToolBarWidget

diff --git a/Engine/Source/Runtime/Slate/Private/Framework/MultiBox/MultiBoxBuilder.cpp b/Engine/Source/Runtime/Slate/Private/Framework/MultiBox/MultiBoxBuilder.cpp
--- a/Engine/Source/Runtime/Slate/Private/Framework/MultiBox/MultiBoxBuilder.cpp
+++ b/Engine/Source/Runtime/Slate/Private/Framework/MultiBox/MultiBoxBuilder.cpp
@@ -569,7 +569,7 @@ void FToolBarBuilder::AddToolbarStackButton(const TSharedPtr< const FUICommandIn
 	MultiBox->AddMultiBlock(NewToolBarStackButtonBlock);
 }
 
-void FToolBarBuilder::AddToolBarWidget( TSharedRef<SWidget> InWidget, const TAttribute<FText>& InLabel, FName InTutorialHighlightName, bool bSearchable )
+void FToolBarBuilder::AddToolBarWidget( TSharedRef<SWidget> InWidget, const TAttribute<FText>& InLabel, FName InTutorialHighlightName, bool bSearchable, const TAttribute<FText>& InToolTip)
 {
 	ApplySectionBeginning();
 
@@ -579,6 +579,7 @@ void FToolBarBuilder::AddToolBarWidget( TSharedRef<SWidget> InWidget, const TAtt
 	InWidget = 
 		SNew( SVerticalBox )
 		.AddMetaData<FTagMetaData>(FTagMetaData(InTutorialHighlightName))
+		.ToolTipText(InToolTip)
 
 		+SVerticalBox::Slot()
 		.AutoHeight()
diff --git a/Engine/Source/Runtime/Slate/Public/Framework/MultiBox/MultiBoxBuilder.h b/Engine/Source/Runtime/Slate/Public/Framework/MultiBox/MultiBoxBuilder.h
--- a/Engine/Source/Runtime/Slate/Public/Framework/MultiBox/MultiBoxBuilder.h
+++ b/Engine/Source/Runtime/Slate/Public/Framework/MultiBox/MultiBoxBuilder.h
@@ -594,9 +594,10 @@ public:
 	 * @param	InWidget				The widget that should be shown in the toolbar
 	 * @param	InLabel                 Optional Label.  
 	 * @param	InTutorialHighlightName	Name to identify this widget and highlight during tutorials
-	 * @param	bInSearchable				If true, widget will be searchable (default == true)
+	 * @param	bInSearchable			If true, widget will be searchable (default == true)
+	 * @param	InToolTip				Optional tool tip 
 	 */
-	SLATE_API void AddToolBarWidget(TSharedRef<SWidget> InWidget, const TAttribute<FText>& InLabel = TAttribute<FText>(), FName InTutorialHighlightName = NAME_None, bool bInSearchable = true);
+	SLATE_API void AddToolBarWidget(TSharedRef<SWidget> InWidget, const TAttribute<FText>& InLabel = TAttribute<FText>(), FName InTutorialHighlightName = NAME_None, bool bInSearchable = true, const TAttribute<FText>& InToolTip = TAttribute<FText>());
 
 
 	/**
