Collapse empty labels on toolbar buttons to prevent them from taking up space

diff --git a/Engine/Source/Runtime/Slate/Private/Framework/MultiBox/SToolBarButtonBlock.cpp b/Engine/Source/Runtime/Slate/Private/Framework/MultiBox/SToolBarButtonBlock.cpp
--- a/Engine/Source/Runtime/Slate/Private/Framework/MultiBox/SToolBarButtonBlock.cpp
+++ b/Engine/Source/Runtime/Slate/Private/Framework/MultiBox/SToolBarButtonBlock.cpp
@@ -233,8 +233,28 @@ void SToolBarButtonBlock::BuildMultiBlockWidget(const ISlateStyle* StyleSet, con
 		|| bIsSlimHorizontalUniformToolBar)
 	{
 		const FVector2f IconSize = ToolBarStyle.IconSize;
-		const TSharedRef<STextBlock> TextBlock = SNew(STextBlock)
-				.Visibility(LabelVisibility)
+		const TSharedRef<STextBlock> TextBlock =
+			SNew(STextBlock)
+				// Collapse empty labels to prevent them from taking up visible space.
+				.Visibility_Lambda(
+					[WeakBlock = SharedThis(this).ToWeakPtr(), ActualLabel]() -> EVisibility
+					{
+						// Check first if the label is empty, and if so collapse it.
+						if (ActualLabel.IsSet() && ActualLabel.Get().IsEmpty())
+						{
+							return EVisibility::Collapsed;
+						}
+						// Only now check the set override.
+						else if (TSharedPtr<SToolBarButtonBlock> Block = WeakBlock.Pin())
+						{
+							return Block->LabelVisibility.Get();
+						}
+						else
+						{
+							return EVisibility::Visible;
+						}
+					}
+				)
 				.Text(ActualLabel)
 				.TextStyle(&ToolBarStyle.LabelStyle); // Smaller font for tool tip labels
 
