Fix tooltip placement so that it can't be pushed offscreen by forcefields when its size changes dynamically.

diff --git a/Engine/Source/Runtime/Slate/Private/Framework/Application/SlateUser.cpp b/Engine/Source/Runtime/Slate/Private/Framework/Application/SlateUser.cpp
--- a/Engine/Source/Runtime/Slate/Private/Framework/Application/SlateUser.cpp
+++ b/Engine/Source/Runtime/Slate/Private/Framework/Application/SlateUser.cpp
@@ -1405,7 +1405,8 @@ void FSlateUser::UpdateTooltip(const FMenuStack& MenuStack, bool bCanSpawnNewToo
 		const float SlideProgress = bAllowAnimations ? FMath::Pow(1.0f - TooltipOpacity, 3.0f) : 0.0f;
 
 		FVector2f WindowLocation = DesiredLocation + SlideProgress * SlideDistance;
-		if (WindowLocation != TooltipWindow->GetPositionInScreen())
+		if (WindowLocation != TooltipWindow->GetPositionInScreen()
+			|| TooltipWindow->GetDesiredSize() != ActiveTooltipInfo.DesiredSize)
 		{
 			// already handled
 			const bool bAutoAdjustForDPIScale = false;
@@ -1414,6 +1415,9 @@ void FSlateUser::UpdateTooltip(const FMenuStack& MenuStack, bool bCanSpawnNewToo
 			FSlateRect Anchor(WindowLocation.X, WindowLocation.Y, WindowLocation.X, WindowLocation.Y);
 			WindowLocation = SlateApp.CalculateTooltipWindowPosition(Anchor, TooltipWindow->GetDesiredSizeDesktopPixels(), bAutoAdjustForDPIScale);
 
+			// Cache the size to compare against in future frames
+			ActiveTooltipInfo.DesiredSize = TooltipWindow->GetDesiredSize();
+
 			// Update the tool tip window positioning
 			// SetCachedScreenPosition is a hack (issue tracked as TTP #347070) which is needed because code in TickWindowAndChildren()/DrawPrepass()
 			// assumes GetPositionInScreen() to correspond to the new window location in the same tick. This is true on Windows, but other
diff --git a/Engine/Source/Runtime/Slate/Public/Framework/Application/SlateUser.h b/Engine/Source/Runtime/Slate/Public/Framework/Application/SlateUser.h
--- a/Engine/Source/Runtime/Slate/Public/Framework/Application/SlateUser.h
+++ b/Engine/Source/Runtime/Slate/Public/Framework/Application/SlateUser.h
@@ -316,6 +316,9 @@ private:
 		/** Desired position of the tooltip in screen space, updated whenever the mouse moves */
 		FVector2f DesiredLocation = FVector2f::ZeroVector;
 
+		/** Desired size of the tooltip in screen space */
+		FVector2f DesiredSize = FVector2f::ZeroVector;
+
 		/** The time at which the tooltip was summoned */
 		double SummonTime = 0.0;
 
