SListView: When a pinned item is clicked, also select it (deslecting any currently selected items) alongside navigating to it

diff --git a/Engine/Source/Runtime/Slate/Public/Widgets/Views/SListView.h b/Engine/Source/Runtime/Slate/Public/Widgets/Views/SListView.h
--- a/Engine/Source/Runtime/Slate/Public/Widgets/Views/SListView.h
+++ b/Engine/Source/Runtime/Slate/Public/Widgets/Views/SListView.h
@@ -1235,8 +1235,8 @@ private:
 
 				if (PinnedItem)
 				{
-					// Navigate to the pinned item on click
-					OwnerListView->RequestNavigateToItem(*PinnedItem);
+					// Select the pinned item on click 
+					OwnerListView->NavigationSelect(*PinnedItem, MouseEvent);
 					return FReply::Handled();
 				}
 			}
@@ -2509,6 +2509,7 @@ protected:
 
 	/**
 	 * Selects the specified item and scrolls it into view. If shift is held, it will be a range select.
+	 * If ctrl is held and multi selection is allowed, the item will be added to the list of currently selected items
 	 * 
 	 * @param ItemToSelect		The item that was selected by a keystroke
 	 * @param InInputEvent	The key event that caused this selection
@@ -2544,7 +2545,12 @@ protected:
 
 					this->Private_SelectRangeFromCurrentTo(ItemToSelect.GetValue());
 				}
-
+				// For ctrl select, simply add the item to the list of currently selected items
+				else
+				{
+					this->Private_SetItemSelection(ItemToSelect.GetValue(), true, true);
+				}
+				
 				this->Private_SignalSelectionChanged(ESelectInfo::OnNavigation);
 			}
 			else
