Backing out CL 34061862 and 35266395 - to be revisited later
#rnx

diff --git a/Engine/Source/Runtime/Slate/Private/Framework/Application/NavigationConfig.cpp b/Engine/Source/Runtime/Slate/Private/Framework/Application/NavigationConfig.cpp
--- a/Engine/Source/Runtime/Slate/Private/Framework/Application/NavigationConfig.cpp
+++ b/Engine/Source/Runtime/Slate/Private/Framework/Application/NavigationConfig.cpp
@@ -58,23 +58,23 @@ void FNavigationConfig::OnUserRemoved(int32 UserIndex)
 
 EUINavigation FNavigationConfig::GetNavigationDirectionFromKey(const FKeyEvent& InKeyEvent) const
 {
-	const bool bIsShiftDown = InKeyEvent.IsShiftDown();
-	const bool bModifierHeld = InKeyEvent.IsControlDown() || InKeyEvent.IsAltDown() || InKeyEvent.IsCommandDown() || bIsShiftDown;
-	if (bIgnoreModifiersForNavigationActions || !bModifierHeld)
+	if (const EUINavigation* Rule = KeyEventRules.Find(InKeyEvent.GetKey()))
 	{
-		if (const EUINavigation* Rule = KeyEventRules.Find(InKeyEvent.GetKey()))
+		if (bKeyNavigation)
 		{
-			if (bKeyNavigation)
-			{
-				return *Rule;
-			}
+			return *Rule;
 		}
 	}
-
-	if (bTabNavigation && InKeyEvent.GetKey() == EKeys::Tab )
+	else if (bTabNavigation && InKeyEvent.GetKey() == EKeys::Tab )
 	{
 		//@TODO: Really these uses of input should be at a lower priority, only occurring if nothing else handled them
-		return bIsShiftDown ? EUINavigation::Previous : EUINavigation::Next;
+		// For now this code prevents consuming them when some modifiers are held down, allowing some limited binding
+		const bool bAllowEatingKeyEvents = !InKeyEvent.IsControlDown() && !InKeyEvent.IsAltDown() && !InKeyEvent.IsCommandDown();
+
+		if ( bAllowEatingKeyEvents )
+		{
+			return ( InKeyEvent.IsShiftDown() ) ? EUINavigation::Previous : EUINavigation::Next;
+		}
 	}
 
 	return EUINavigation::Invalid;
