Added rich tooltip support for segmented control widget

#rnx

diff --git a/Engine/Source/Editor/PropertyEditor/Private/UserInterface/PropertyEditor/SPropertyEditorCombo.cpp b/Engine/Source/Editor/PropertyEditor/Private/UserInterface/PropertyEditor/SPropertyEditorCombo.cpp
--- a/Engine/Source/Editor/PropertyEditor/Private/UserInterface/PropertyEditor/SPropertyEditorCombo.cpp
+++ b/Engine/Source/Editor/PropertyEditor/Private/UserInterface/PropertyEditor/SPropertyEditorCombo.cpp
@@ -77,20 +77,23 @@ void SPropertyEditorCombo::Construct( const FArguments& InArgs, const TSharedPtr
 		for (int ItemIndex = 0; ItemIndex < ComboItems.Num(); ItemIndex++)
 		{
 			const FString& ComboItem = *ComboItems[ItemIndex];
+			SSegmentedControl<FString>::FScopedWidgetSlotArguments Slot = SegmentControl->AddSlot(ComboItem);
 			
 			FText DisplayName = FText::FromString(ComboItem);
 			FText TooltipText = DisplayName;
 			if (RichToolTips.IsValidIndex(ItemIndex) && !RichToolTips[ItemIndex]->IsEmpty())
 			{
-				TooltipText = RichToolTips[ItemIndex]->GetTextTooltip();
+				Slot.ToolTipWidget(RichToolTips[ItemIndex]);
 			}
-
-			SSegmentedControl<FString>::FScopedWidgetSlotArguments Slot = SegmentControl->AddSlot(ComboItem);
+			else
+			{
+				Slot.ToolTip(DisplayName);
+			}
+			
 			Slot
 			  .HAlign(HAlign_Center)
 			  .VAlign(VAlign_Center)
-			  .Text(DisplayName)
-			  .ToolTip(TooltipText);
+			  .Text(DisplayName);
 		}
 		
 		ChildSlot
diff --git a/Engine/Source/Runtime/Slate/Public/Widgets/Input/SSegmentedControl.h b/Engine/Source/Runtime/Slate/Public/Widgets/Input/SSegmentedControl.h
--- a/Engine/Source/Runtime/Slate/Public/Widgets/Input/SSegmentedControl.h
+++ b/Engine/Source/Runtime/Slate/Public/Widgets/Input/SSegmentedControl.h
@@ -44,6 +44,7 @@ public:
 		SLATE_SLOT_BEGIN_ARGS_OneMixin(FSlot, TSlotBase<FSlot>, TAlignmentWidgetSlotMixin<FSlot>)
 			SLATE_ATTRIBUTE(FText, Text)
 			SLATE_ATTRIBUTE(FText, ToolTip)
+			SLATE_ATTRIBUTE(TSharedPtr<IToolTip>, ToolTipWidget)
 			SLATE_ATTRIBUTE(const FSlateBrush*, Icon)
 			SLATE_ARGUMENT(TOptional<OptionType>, Value)
 		SLATE_SLOT_END_ARGS()
@@ -60,6 +61,10 @@ public:
 			{
 				_Tooltip = MoveTemp(InArgs._ToolTip);
 			}
+			if (InArgs._ToolTipWidget.IsSet())
+			{
+				_ToolTipWidget = MoveTemp(InArgs._ToolTipWidget);
+			}
 			if (InArgs._Icon.IsSet())
 			{
 				_Icon = MoveTemp(InArgs._Icon);
@@ -105,6 +110,7 @@ public:
 	private:
 		TAttribute<FText> _Text;
 		TAttribute<FText> _Tooltip;
+		TAttribute<TSharedPtr<IToolTip>> _ToolTipWidget;
 		TAttribute<const FSlateBrush*> _Icon;
 
 		OptionType _Value;
@@ -283,6 +289,7 @@ public:
 				.Clipping(EWidgetClipping::ClipToBounds)
 				.HAlign(ChildSlotPtr->GetHorizontalAlignment())
 				.ToolTipText(ChildSlotPtr->_Tooltip)
+				.ToolTip(ChildSlotPtr->_ToolTipWidget)
 				.Style(CheckBoxStyle)
 				.IsChecked(GetCheckBoxStateAttribute(ChildValue))
 				.OnCheckStateChanged(this, &SSegmentedControl::CommitValue, ChildValue)
