Fixed a problem with tooltip force fields, where if the force field would push the tooltip towards an invalid position (e.g. when forced down but the tooltip ends up outside the screen), the force field would end up being ineffective.
With the fix, we're still trying to push the tooltip in the best direction, but if that fails, we're trying the other direction

diff --git a/Engine/Source/Runtime/Slate/Private/Framework/Application/SlateUser.cpp b/Engine/Source/Runtime/Slate/Private/Framework/Application/SlateUser.cpp
--- a/Engine/Source/Runtime/Slate/Private/Framework/Application/SlateUser.cpp
+++ b/Engine/Source/Runtime/Slate/Private/Framework/Application/SlateUser.cpp
@@ -1324,10 +1324,12 @@ void FSlateUser::UpdateTooltip(const FMenuStack& MenuStack, bool bCanSpawnNewToo
 		}
 	}
 
+	FVector2f TooltipSize(ForceInitToZero);
 	if (TooltipWindowPtr.IsValid())
 	{
+		TooltipSize = TooltipWindowPtr.Pin()->GetDesiredSizeDesktopPixels();
 		FSlateRect Anchor(DesiredLocation.X, DesiredLocation.Y, DesiredLocation.X, DesiredLocation.Y);
-		DesiredLocation = SlateApp.CalculatePopupWindowPosition(Anchor, TooltipWindowPtr.Pin()->GetDesiredSizeDesktopPixels(), /*bAutoAdjustForDPIScale =*/false);
+		DesiredLocation = SlateApp.CalculatePopupWindowPosition(Anchor, TooltipSize, /*bAutoAdjustForDPIScale =*/false);
 	}
 
 	// Repel tooltip from a force field, if necessary
@@ -1338,11 +1340,41 @@ void FSlateUser::UpdateTooltip(const FMenuStack& MenuStack, bool bCanSpawnNewToo
 		TooltipShift.Y = (ForceFieldRect->Bottom + SlateDefs::TooltipOffsetFromForceField.Y) - DesiredLocation.Y;
 
 		// Make sure the tooltip needs to be offset
-		if (TooltipShift.X != 0.0f && TooltipShift.Y != 0.0f)
+		if (TooltipShift.X != 0.0f || TooltipShift.Y != 0.0f)
 		{
 			// Find the best edge to move the tooltip towards
-			if (ActiveTooltipInfo.OffsetDirection == ETooltipOffsetDirection::Right ||
-				(ActiveTooltipInfo.OffsetDirection == ETooltipOffsetDirection::Undetermined && FMath::Abs(TooltipShift.X) < FMath::Abs(TooltipShift.Y)))
+			if (ActiveTooltipInfo.OffsetDirection == ETooltipOffsetDirection::Undetermined)
+			{
+				ETooltipOffsetDirection PotentialOffsetDirection = FMath::Abs(TooltipShift.X) < FMath::Abs(TooltipShift.Y) ? ETooltipOffsetDirection::Right : ETooltipOffsetDirection::Down;
+				if (PotentialOffsetDirection == ETooltipOffsetDirection::Right)
+				{
+					FVector2f TentativeDesiredLocation = DesiredLocation + FVector2f(TooltipShift.X, 0);
+					FSlateRect Anchor(TentativeDesiredLocation.X, TentativeDesiredLocation.Y, TentativeDesiredLocation.X, TentativeDesiredLocation.Y);
+					FVector2f NewDesiredLocation = SlateApp.CalculatePopupWindowPosition(Anchor, TooltipSize, /*bAutoAdjustForDPIScale =*/false);
+					// If after adjustment the tooltip still overlaps with the force field, try the other direction
+					if (FSlateRect::DoRectanglesIntersect(ForceFieldRect.GetValue(), FSlateRect(NewDesiredLocation, NewDesiredLocation + TooltipSize)))
+					{
+						PotentialOffsetDirection = ETooltipOffsetDirection::Down;
+					}
+				}
+
+				if (PotentialOffsetDirection == ETooltipOffsetDirection::Down)
+				{
+					FVector2f TentativeDesiredLocation = DesiredLocation + FVector2f(0, TooltipShift.Y);
+					FSlateRect Anchor(TentativeDesiredLocation.X, TentativeDesiredLocation.Y, TentativeDesiredLocation.X, TentativeDesiredLocation.Y);
+					FVector2f NewDesiredLocation = SlateApp.CalculatePopupWindowPosition(Anchor, TooltipSize, /*bAutoAdjustForDPIScale =*/false);
+					// If after adjustment the tooltip still overlaps with the force field, try the other direction
+					if (FSlateRect::DoRectanglesIntersect(ForceFieldRect.GetValue(), FSlateRect(NewDesiredLocation, NewDesiredLocation + TooltipSize)))
+					{
+						PotentialOffsetDirection = ETooltipOffsetDirection::Right;
+					}
+				}
+
+				ActiveTooltipInfo.OffsetDirection = PotentialOffsetDirection;
+			}
+
+			check(ActiveTooltipInfo.OffsetDirection != ETooltipOffsetDirection::Undetermined);
+			if (ActiveTooltipInfo.OffsetDirection == ETooltipOffsetDirection::Right)
 			{
 				// Move right
 				DesiredLocation.X += TooltipShift.X;
