Added CVar "Slate.AllowTooltipsWithHiddenCursor" to SlateUser to enable tooltips even when the cursor is hidden. This would enable tooltips when using something like the CommonUI analog cursor

diff --git a/Engine/Source/Runtime/Slate/Private/Framework/Application/SlateUser.cpp b/Engine/Source/Runtime/Slate/Private/Framework/Application/SlateUser.cpp
--- a/Engine/Source/Runtime/Slate/Private/Framework/Application/SlateUser.cpp
+++ b/Engine/Source/Runtime/Slate/Private/Framework/Application/SlateUser.cpp
@@ -62,6 +62,12 @@ FAutoConsoleVariableRef CVarCursorSignificantMoveDetectionThreshold(
 	CursorSignificantMoveDetectionThreshold,
 	TEXT("The distance from previous cursor position above which the move will be considered significant (used to trigger the display of the tooltips)."));
 
+static bool bAllowTooltipsWithHiddenCursor = false;
+FAutoConsoleVariableRef CVarAllowTooltipsWithHiddenCursor(
+	TEXT("Slate.AllowTooltipsWithHiddenCursor"),
+	bAllowTooltipsWithHiddenCursor,
+	TEXT(""));
+
 //////////////////////////////////////////////////////////////////////////
 // FSlateVirtualUserHandle
 //////////////////////////////////////////////////////////////////////////
@@ -1195,7 +1201,7 @@ void FSlateUser::UpdateTooltip(const FMenuStack& MenuStack, bool bCanSpawnNewToo
 	SCOPE_CYCLE_COUNTER(STAT_SlateUpdateTooltip);
 
 	const double MotionLessDurationBeforeAllowingNewToolTip = 0.05;
-	bCanSpawnNewTooltip = bCanSpawnNewTooltip && bCanDrawCursor && (FPlatformTime::Seconds() - LastCursorSignificantMoveTime > MotionLessDurationBeforeAllowingNewToolTip);
+	bCanSpawnNewTooltip = bCanSpawnNewTooltip && (bCanDrawCursor || bAllowTooltipsWithHiddenCursor) && (FPlatformTime::Seconds() - LastCursorSignificantMoveTime > MotionLessDurationBeforeAllowingNewToolTip);
 
 	float DPIScaleFactor = 1.0f; //todo: this value is never changed, we should investigate if it is necessary or not to handle it for the force field.
 	FWidgetPath WidgetsToQueryForTooltip;
