Disable breakpoints during tab restore to prevent editor freeze / infinite loop

diff --git a/Engine/Source/Runtime/Slate/Private/Framework/Docking/TabManager.cpp b/Engine/Source/Runtime/Slate/Private/Framework/Docking/TabManager.cpp
--- a/Engine/Source/Runtime/Slate/Private/Framework/Docking/TabManager.cpp
+++ b/Engine/Source/Runtime/Slate/Private/Framework/Docking/TabManager.cpp
@@ -1650,6 +1650,10 @@ TSharedPtr<SDockingArea> FTabManager::RestoreArea(const TSharedRef<FArea>& AreaT
 TSharedPtr<SDockingNode> FTabManager::RestoreArea_Helper(const TSharedRef<FLayoutNode>& LayoutNode, const TSharedPtr<SWindow>& ParentWindow, const bool bEmbedTitleAreaContent,
 	FSidebarTabLists& OutSidebarTabs, const EOutputCanBeNullptr OutputCanBeNullptr, bool bForceOpenWindowIfNeeded)
 {
+#if WITH_EDITOR
+	FSlateApplication::FScopedPreventDebuggingMode Scope(LOCTEXT("RestoringTabsDebugScope", "Disabling debug due to being in tab restore, breakpoints in constructors can infinitely stall during restore."));
+#endif
+
 	TSharedPtr<FTabManager::FStack> NodeAsStack = LayoutNode->AsStack();
 	TSharedPtr<FTabManager::FSplitter> NodeAsSplitter = LayoutNode->AsSplitter();
 	TSharedPtr<FTabManager::FArea> NodeAsArea = LayoutNode->AsArea();
diff --git a/Engine/Source/Runtime/Slate/Public/Framework/Application/SlateApplication.h b/Engine/Source/Runtime/Slate/Public/Framework/Application/SlateApplication.h
--- a/Engine/Source/Runtime/Slate/Public/Framework/Application/SlateApplication.h
+++ b/Engine/Source/Runtime/Slate/Public/Framework/Application/SlateApplication.h
@@ -899,10 +899,10 @@ public:
 	SLATE_API void LeaveDebuggingMode( bool bLeavingDebugForSingleStep = false );
 
 #if WITH_EDITOR
-	struct FScopedPreventDebuggingMode
+	struct SLATE_API FScopedPreventDebuggingMode
 	{
-		UE_NODISCARD_CTOR SLATE_API FScopedPreventDebuggingMode(FText Reason);
-		SLATE_API ~FScopedPreventDebuggingMode();
+		UE_NODISCARD_CTOR FScopedPreventDebuggingMode(FText Reason);
+		~FScopedPreventDebuggingMode();
 	private:
 		int32 Id;
 	};
