[ToolMenus] Allow variable widget scaling in toolbars

#rnx

diff --git a/Engine/Source/Developer/ToolMenus/Private/ToolMenus.cpp b/Engine/Source/Developer/ToolMenus/Private/ToolMenus.cpp
--- a/Engine/Source/Developer/ToolMenus/Private/ToolMenus.cpp
+++ b/Engine/Source/Developer/ToolMenus/Private/ToolMenus.cpp
@@ -353,7 +353,8 @@ public:
 			}
 			else
 			{
-				MenuBuilder.AddWidget(Widget.ToSharedRef(), LabelToDisplay.Get(), Block.WidgetData.bNoIndent, Block.WidgetData.bSearchable, Block.ToolTip.Get());
+				Block.WidgetData.StyleParams.bNoIndent = Block.WidgetData.bNoIndent;
+				MenuBuilder.AddWidget(Widget.ToSharedRef(), LabelToDisplay.Get(), Block.WidgetData.StyleParams, Block.WidgetData.bSearchable, Block.ToolTip.Get());
 			}
 		}
 		else
@@ -1701,7 +1702,13 @@ void UToolMenus::PopulateToolBarBuilderWithEntry(
 		if (OnGetContent.IsBound())
 		{
 			ToolBarBuilder.AddComboButton(
-				UIAction, OnGetContent, Label, Block.ToolTip, Block.Icon, Block.ToolBarData.bSimpleComboBox, Block.TutorialHighlightName
+				UIAction,
+				OnGetContent,
+				Label,
+				Block.ToolTip,
+				Block.Icon,
+				Block.ToolBarData.bSimpleComboBox,
+				Block.TutorialHighlightName
 			);
 		}
 		else
@@ -1710,7 +1717,13 @@ void UToolMenus::PopulateToolBarBuilderWithEntry(
 				this, &UToolMenus::GenerateToolbarComboButtonMenu, TWeakObjectPtr<UToolMenu>(MenuData), Block.Name);
 
 			ToolBarBuilder.AddComboButton(
-				UIAction, Delegate, Label, Block.ToolTip, Block.Icon, Block.ToolBarData.bSimpleComboBox, Block.TutorialHighlightName
+				UIAction,
+				Delegate,
+				Label,
+				Block.ToolTip,
+				Block.Icon,
+				Block.ToolBarData.bSimpleComboBox,
+				Block.TutorialHighlightName
 			);
 
 			// Also add any top-level flagged children to the toolbar.
@@ -1743,7 +1756,10 @@ void UToolMenus::PopulateToolBarBuilderWithEntry(
 		}
 		PRAGMA_ENABLE_DEPRECATION_WARNINGS
 
-		ToolBarBuilder.AddWidget(Widget.ToSharedRef(), Block.TutorialHighlightName, Block.WidgetData.bSearchable);
+		FMenuEntryStyleParams StyleParams = Block.WidgetData.StyleParams;
+		StyleParams.HorizontalAlignment = HAlign_Fill;
+
+		ToolBarBuilder.AddWidget(Widget.ToSharedRef(), StyleParams, Block.TutorialHighlightName, Block.WidgetData.bSearchable);
 	}
 	else
 	{
@@ -1795,7 +1811,12 @@ void UToolMenus::PopulateToolBarBuilder(FToolBarBuilder& ToolBarBuilder, UToolMe
 		// previous alignment group.
 		if (bIsMiddleOrLast && bDidAddSection)
 		{
-			ToolBarBuilder.AddWidget(SNew(SSpacer), NAME_None, true, HAlign_Right);
+			FMenuEntryStyleParams StyleParams;
+			StyleParams.HorizontalAlignment = HAlign_Right;
+			StyleParams.SizeRule = FSizeParam::ESizeRule::SizeRule_StretchContent;
+			StyleParams.MinSize = 0.5f;
+
+			ToolBarBuilder.AddWidget(SNew(SSpacer), StyleParams, NAME_None, true, {});
 		}
 
 		// Keep track if this is the first section we're adding for the CurrentAlignment. Make an exception if the
diff --git a/Engine/Source/Developer/ToolMenus/Public/ToolMenuEntry.h b/Engine/Source/Developer/ToolMenus/Public/ToolMenuEntry.h
--- a/Engine/Source/Developer/ToolMenus/Public/ToolMenuEntry.h
+++ b/Engine/Source/Developer/ToolMenus/Public/ToolMenuEntry.h
@@ -51,22 +51,21 @@ public:
 	{
 	}
 
-	/** Delegate that generates a widget for this combo button's menu content.  Called when the menu is summoned. */
+	/** Delegate that generates a widget for this combo button's menu content. Called when the menu is summoned. */
 	FNewToolMenuChoice ComboButtonContextMenuGenerator;
 
-	/** Legacy delegate that generates a widget for this combo button's menu content.  Called when the menu is summoned. */
+	/** Legacy delegate that generates a widget for this combo button's menu content. Called when the menu is summoned. */
 	FNewToolBarDelegateLegacy ConstructLegacy;
 
 	TSharedPtr<FToolMenuEntryOptionsDropdownData> OptionsDropdownData;
 
 	bool bSimpleComboBox;
 
-	/** Whether ToolBar will have Focusable buttons */
+	/** Whether ToolBar will have Focusable buttons. */
 	bool bIsFocusable;
 
-	/** Whether this toolbar should always use small icons, regardless of the current settings */
+	/** Whether this toolbar should always use small icons, regardless of the current settings. */
 	bool bForceSmallIcons;
-
 };
 
 
@@ -88,6 +87,9 @@ public:
 
 	/** If true, no padding will be added */
 	bool bNoPadding;
+
+	/** Various Style parameters and overrides */
+	FMenuEntryStyleParams StyleParams;
 };
 
 struct FToolMenuCustomWidgetContext
diff --git a/Engine/Source/Runtime/Slate/Private/Framework/MultiBox/MultiBox.cpp b/Engine/Source/Runtime/Slate/Private/Framework/MultiBox/MultiBox.cpp
--- a/Engine/Source/Runtime/Slate/Private/Framework/MultiBox/MultiBox.cpp
+++ b/Engine/Source/Runtime/Slate/Private/Framework/MultiBox/MultiBox.cpp
@@ -808,11 +808,8 @@ void SMultiBoxWidget::AddBlockWidget(const FMultiBlock& Block, TSharedPtr<SHoriz
 	case EMultiBoxType::ToolBar:
 	case EMultiBoxType::SlimHorizontalToolBar:
 		{
-			EHorizontalAlignment HAlign;
-			EVerticalAlignment VAlign;
-			bool bAutoWidth;
-
-			bool bOverride = Block.GetAlignmentOverrides(HAlign, VAlign, bAutoWidth);
+			FMenuEntryStyleParams StyleParams;
+			const bool bOverride = Block.GetAlignmentOverrides(StyleParams);
 
 			{
 				SHorizontalBox::FScopedWidgetSlotArguments NewSlot = HorizontalBox->AddSlot();
@@ -824,13 +821,38 @@ void SMultiBoxWidget::AddBlockWidget(const FMultiBlock& Block, TSharedPtr<SHoriz
 
 				if (bOverride)
 				{
-					if (bAutoWidth)
+					if (StyleParams.SizeRule.IsSet()
+						&& StyleParams.SizeRule.GetValue() == FSizeParam::SizeRule_Auto)
 					{
 						NewSlot.AutoWidth();
 					}
 
-					NewSlot.HAlign(HAlign)
-						.VAlign(VAlign);
+					NewSlot
+					.HAlign(StyleParams.HorizontalAlignment)
+					.VAlign(StyleParams.VerticalAlignment.Get(VAlign_Fill));
+
+					if (StyleParams.SizeRule.IsSet())
+					{
+						FSizeParam::ESizeRule SizeRule = StyleParams.SizeRule.GetValue();
+						if (SizeRule == FSizeParam::SizeRule_Stretch)
+						{
+							NewSlot.FillWidth(1.0f);
+						}
+						else if (SizeRule == FSizeParam::SizeRule_StretchContent)
+						{
+							NewSlot.FillContentWidth(1.0f);
+						}
+					}
+
+					if (StyleParams.MinSize.IsSet())
+					{
+						NewSlot.MinWidth(StyleParams.MinSize.GetValue());
+					}
+
+					if (StyleParams.MaxSize.IsSet())
+					{
+						NewSlot.MaxWidth(StyleParams.MaxSize.GetValue());
+					}
 				}
 				else
 				{
@@ -933,7 +955,10 @@ void SMultiBoxWidget::CreateSearchTextWidget()
 				SearchTextWidget.ToSharedRef()
 			];
 
-	TSharedRef<FWidgetBlock> NewWidgetBlock(new FWidgetBlock(SearchBox, FText::GetEmpty(), false));
+	FMenuEntryStyleParams StyleParams;
+	StyleParams.bNoIndent = false;
+
+	TSharedRef<FWidgetBlock> NewWidgetBlock = MakeShared<FWidgetBlock>(SearchBox, FText::GetEmpty(), FText(), StyleParams);
 	NewWidgetBlock->SetSearchable(false);
 
 	MultiBox->AddMultiBlockToFront(NewWidgetBlock);
diff --git a/Engine/Source/Runtime/Slate/Private/Framework/MultiBox/MultiBoxBuilder.cpp b/Engine/Source/Runtime/Slate/Private/Framework/MultiBox/MultiBoxBuilder.cpp
--- a/Engine/Source/Runtime/Slate/Private/Framework/MultiBox/MultiBoxBuilder.cpp
+++ b/Engine/Source/Runtime/Slate/Private/Framework/MultiBox/MultiBoxBuilder.cpp
@@ -350,10 +350,18 @@ void FMenuBuilder::AddWrapperSubMenu( const FText& InMenuLabel, const FText& InT
 }
 
 void FMenuBuilder::AddWidget( TSharedRef<SWidget> InWidget, const FText& Label, bool bNoIndent, bool bInSearchable, const TAttribute<FText>& InToolTipText )
+{
+	FMenuEntryStyleParams StyleParams;
+	StyleParams.bNoIndent = bNoIndent;
+
+	return AddWidget(InWidget, Label, StyleParams, bInSearchable, InToolTipText);
+}
+
+void FMenuBuilder::AddWidget(const TSharedRef<SWidget>& InWidget, const FText& InLabel, const FMenuEntryStyleParams& InStyleParams, const bool bInSearchable, const TAttribute<FText>& InToolTipText)
 {
 	ApplySectionBeginning();
 
-	TSharedRef< FWidgetBlock > NewWidgetBlock(new FWidgetBlock( InWidget, Label, bNoIndent, EHorizontalAlignment::HAlign_Fill, InToolTipText ));
+	TSharedRef<FWidgetBlock> NewWidgetBlock = MakeShared<FWidgetBlock>(InWidget, InLabel, InToolTipText, InStyleParams);
 	NewWidgetBlock->SetSearchable( bInSearchable );
 
 	MultiBox->AddMultiBlock( NewWidgetBlock );
@@ -609,27 +617,37 @@ void FToolBarBuilder::AddToolBarWidget( TSharedRef<SWidget> InWidget, const TAtt
 			.Text( InLabel )
 			.TextStyle(&ToolBarStyle.LabelStyle)	// Smaller font for tool tip labels
 		] ;
-	
-	TSharedRef< FWidgetBlock > NewWidgetBlock( new FWidgetBlock( InWidget, FText::GetEmpty(), true ) );
+
+	static const TAttribute<FText> EmptyTextAttribute;
+	TSharedRef<FWidgetBlock> NewWidgetBlock = MakeShared<FWidgetBlock>(InWidget, FText::GetEmpty(), EmptyTextAttribute, FMenuEntryStyleParams(true));
 	MultiBox->AddMultiBlock( NewWidgetBlock );
 	NewWidgetBlock->SetSearchable(bSearchable);
 }
 
-void FToolBarBuilder::AddWidget( TSharedRef<SWidget> InWidget, FName InTutorialHighlightName, bool bSearchable, EHorizontalAlignment Alignment, FNewMenuDelegate InCustomMenuDelegate )
+void FToolBarBuilder::AddWidget( TSharedRef<SWidget> InWidget, FName InTutorialHighlightName, bool bSearchable, EHorizontalAlignment Alignment, FNewMenuDelegate InCustomMenuDelegate)
+{
+	AddWidget(InWidget, FMenuEntryStyleParams(Alignment), InTutorialHighlightName, bSearchable, InCustomMenuDelegate);
+}
+
+void FToolBarBuilder::AddWidget(TSharedRef<SWidget> InWidget, const FMenuEntryStyleParams& InStyleParams, FName InTutorialHighlightName, bool bInSearchable, FNewMenuDelegate InCustomMenuDelegate)
 {
 	ApplySectionBeginning();
 
 	TSharedRef<SWidget> ChildWidget = InWidget;
-	InWidget = 
+	InWidget =
 		SNew( SBox )
 		.AddMetaData<FTagMetaData>(FTagMetaData(InTutorialHighlightName))
 		[
 			ChildWidget
 		];
-	
-	TSharedRef< FWidgetBlock > NewWidgetBlock( new FWidgetBlock( InWidget, FText::GetEmpty(), true, Alignment) );
+
+	FMenuEntryStyleParams StyleParams = InStyleParams;
+	StyleParams.bNoIndent = true;
+
+	static const TAttribute<FText> EmptyTextAttribute;
+	TSharedRef<FWidgetBlock> NewWidgetBlock = MakeShared<FWidgetBlock>(InWidget, FText::GetEmpty(), EmptyTextAttribute, StyleParams);
 	MultiBox->AddMultiBlock( NewWidgetBlock );
-	NewWidgetBlock->SetSearchable(bSearchable);
+	NewWidgetBlock->SetSearchable(bInSearchable);
 	NewWidgetBlock->SetCustomMenuDelegate(InCustomMenuDelegate);
 }
 
@@ -652,7 +670,7 @@ void FToolBarBuilder::BeginSection(FName InExtensionHook, bool bInSectionShouldH
 	checkf(CurrentSectionExtensionHook == NAME_None && !bSectionNeedsToBeApplied, TEXT("Did you forget to call EndSection()?"));
 
 	ApplyHook(InExtensionHook, EExtensionHook::Before);
-	
+
 	// Do not actually apply the section header, because if this section is ended immediately
 	// then nothing ever gets created, preventing empty sections from ever appearing
 	bSectionNeedsToBeApplied = true;
diff --git a/Engine/Source/Runtime/Slate/Private/Framework/MultiBox/SMenuEntryBlock.cpp b/Engine/Source/Runtime/Slate/Private/Framework/MultiBox/SMenuEntryBlock.cpp
--- a/Engine/Source/Runtime/Slate/Private/Framework/MultiBox/SMenuEntryBlock.cpp
+++ b/Engine/Source/Runtime/Slate/Private/Framework/MultiBox/SMenuEntryBlock.cpp
@@ -1573,10 +1573,10 @@ TSharedRef< SWidget > SMenuEntryBlock::MakeNewMenuWidget() const
 	}
 	else if (MenuEntryBlock->EntryWidget.IsValid())
 	{
-		const bool bCloseSelfOnly = false;
+		static constexpr bool bCloseSelfOnly = false;
 		FMenuBuilder MenuBuilder(MenuEntryBlock->bShouldCloseWindowAfterMenuSelection, nullptr, TSharedPtr<FExtender>(), bCloseSelfOnly, StyleSet );
 		{
-			MenuBuilder.AddWidget( MenuEntryBlock->EntryWidget.ToSharedRef(), FText::GetEmpty() );
+			MenuBuilder.AddWidget( MenuEntryBlock->EntryWidget.ToSharedRef(), FText::GetEmpty(), FMenuEntryStyleParams() );
 		}
 
 		return MenuBuilder.MakeWidget();
diff --git a/Engine/Source/Runtime/Slate/Private/Framework/MultiBox/SWidgetBlock.cpp b/Engine/Source/Runtime/Slate/Private/Framework/MultiBox/SWidgetBlock.cpp
--- a/Engine/Source/Runtime/Slate/Private/Framework/MultiBox/SWidgetBlock.cpp
+++ b/Engine/Source/Runtime/Slate/Private/Framework/MultiBox/SWidgetBlock.cpp
@@ -6,29 +6,20 @@
 #include "Styling/ToolBarStyle.h"
 #include "Widgets/SToolTip.h"
 
-
-/**
- * Constructor
- *
- * @param	InHeadingText	Heading text
- */
-FWidgetBlock::FWidgetBlock( TSharedRef<SWidget> InContent, const FText& InLabel, bool bInNoIndent, EHorizontalAlignment InHorizontalAlignment, const TAttribute<FText>& InToolTipText)
-	: FMultiBlock( nullptr, nullptr, NAME_None, EMultiBlockType::Widget )
-	, ContentWidget( InContent )
-	, Label( InLabel )
+FWidgetBlock::FWidgetBlock(TSharedRef<SWidget> InContent, const FText& InLabel, const TAttribute<FText>& InToolTipText, const FMenuEntryStyleParams& InStyleParams)
+	: FMultiBlock(nullptr, nullptr, NAME_None, EMultiBlockType::Widget)
+    , ContentWidget(InContent)
+	, Label(InLabel)
 	, ToolTipText(InToolTipText)
-	, bNoIndent( bInNoIndent )
-	, HorizontalAlignment(InHorizontalAlignment)
+	, StyleParams(InStyleParams)
 {
 }
 
-
 void FWidgetBlock::SetCustomMenuDelegate( FNewMenuDelegate& InCustomMenuDelegate )
 {
 	CustomMenuDelegate = InCustomMenuDelegate;	
 }
 
-
 void FWidgetBlock::CreateMenuEntry(FMenuBuilder& MenuBuilder) const
 {
 	if ( CustomMenuDelegate.IsBound() )
@@ -38,11 +29,14 @@ void FWidgetBlock::CreateMenuEntry(FMenuBuilder& MenuBuilder) const
 	else
 	{
 		FText EntryLabel = (!Label.IsEmpty()) ? Label : NSLOCTEXT("WidgetBlock", "CustomControl", "Custom Control");
-		MenuBuilder.AddWidget(ContentWidget, FText::GetEmpty(), true);
+
+		FMenuEntryStyleParams MenuStyleParams;
+		MenuStyleParams.bNoIndent = true;
+
+		MenuBuilder.AddWidget(ContentWidget, FText::GetEmpty(), MenuStyleParams);
 	}
 }
 
-
 /**
  * Allocates a widget for this type of MultiBlock.  Override this in derived classes.
  *
@@ -54,12 +48,23 @@ TSharedRef< class IMultiBlockBaseWidget > FWidgetBlock::ConstructWidget() const
 			.Cursor(EMouseCursor::Default);
 }
 
-
-bool FWidgetBlock::GetAlignmentOverrides(EHorizontalAlignment& OutHorizontalAlignment, EVerticalAlignment& OutVerticalAlignment, bool& bOutAutoWidth) const
+bool FWidgetBlock::GetAlignmentOverrides(FMenuEntryStyleParams& OutAlignmentParameters) const
 {
-	OutHorizontalAlignment = HorizontalAlignment;
-	OutVerticalAlignment = VAlign_Fill;
-	bOutAutoWidth = HorizontalAlignment != HAlign_Fill ? false : true;
+	OutAlignmentParameters.HorizontalAlignment = StyleParams.HorizontalAlignment;
+	OutAlignmentParameters.VerticalAlignment = StyleParams.VerticalAlignment.Get(VAlign_Fill);
+
+	if (StyleParams.SizeRule.IsSet())
+	{
+		OutAlignmentParameters.SizeRule = StyleParams.SizeRule;
+	}
+	else if (StyleParams.HorizontalAlignment == HAlign_Fill)
+	{
+		OutAlignmentParameters.SizeRule = FSizeParam::ESizeRule::SizeRule_Auto;
+	}
+
+	OutAlignmentParameters.MinSize = StyleParams.MinSize;
+	OutAlignmentParameters.MaxSize = StyleParams.MaxSize;
+
 	return true;
 }
 
@@ -72,17 +77,15 @@ void SWidgetBlock::Construct( const FArguments& InArgs )
 {
 }
 
-
-
 /**
  * Builds this MultiBlock widget up from the MultiBlock associated with it
  */
 void SWidgetBlock::BuildMultiBlockWidget(const ISlateStyle* StyleSet, const FName& StyleName)
 {
-	TSharedPtr<SMultiBoxWidget> OwnerMultiBoxWidgetPinned = OwnerMultiBoxWidget.Pin();
+	const TSharedPtr<SMultiBoxWidget> OwnerMultiBoxWidgetPinned = OwnerMultiBoxWidget.Pin();
 
 	TSharedPtr< const FMultiBox > MultiBox = OwnerMultiBoxWidgetPinned->GetMultiBox();
-	TSharedRef< const FWidgetBlock > WidgetBlock = StaticCastSharedRef< const FWidgetBlock >( MultiBlock.ToSharedRef() );
+	const TSharedRef<const FWidgetBlock> WidgetBlock = StaticCastSharedRef<const FWidgetBlock>(MultiBlock.ToSharedRef());
 
 	// Support menus which do not have a defined widget style yet
 	bool bHasLabel = !WidgetBlock->Label.IsEmpty();
@@ -92,14 +95,12 @@ void SWidgetBlock::BuildMultiBlockWidget(const ISlateStyle* StyleSet, const FNam
 	{
 		const FToolBarStyle& ToolBarStyle = StyleSet->GetWidgetStyle<FToolBarStyle>(StyleName);
 
-		Padding = WidgetBlock->bNoIndent ? ToolBarStyle.BlockPadding : ToolBarStyle.IndentedBlockPadding;
+		Padding = WidgetBlock->StyleParams.bNoIndent ? ToolBarStyle.BlockPadding : ToolBarStyle.IndentedBlockPadding;
 		LabelStyle = &ToolBarStyle.LabelStyle;
-
 	}
 	else
 	{
-		Padding = WidgetBlock->bNoIndent ? StyleSet->GetMargin(StyleName, ".Block.Padding") : StyleSet->GetMargin(StyleName, ".Block.IndentedPadding");
-
+		Padding = WidgetBlock->StyleParams.bNoIndent ? StyleSet->GetMargin(StyleName, ".Block.Padding") : StyleSet->GetMargin(StyleName, ".Block.IndentedPadding");
 		LabelStyle = &StyleSet->GetWidgetStyle<FTextBlockStyle>(ISlateStyle::Join(StyleName, ".Label"));
 	}
 
@@ -150,6 +151,8 @@ void SWidgetBlock::BuildMultiBlockWidget(const ISlateStyle* StyleSet, const FNam
 		OwnerMultiBoxWidgetPinned->SetSearchBlockWidget(ThisWidget);
 	}
 
+	static constexpr float IconRightPadding = 4.f;
+
 	ChildSlot
 	.Padding( Padding )	// Large left margin mimics the indent of normal menu items when bNoIndent is false
 	[
@@ -162,8 +165,8 @@ void SWidgetBlock::BuildMultiBlockWidget(const ISlateStyle* StyleSet, const FNam
 			.Visibility( bHasLabel ? EVisibility::Visible : EVisibility::Collapsed )
 			+ SHorizontalBox::Slot()
 			.AutoWidth()
-			.Padding( bHasLabel ? FMargin(0.f, 0.f, 4.f, 0.f) : FMargin(0))
-			.VAlign( VAlign_Center )
+			.Padding(bHasLabel ? FMargin(0.f, 0.f, IconRightPadding, 0.f) : FMargin(0))
+			.VAlign(VAlign_Center)
 			[
 				SNew( STextBlock )
 				.TextStyle(LabelStyle)
@@ -217,4 +220,3 @@ void SWidgetBlock::OnMouseEnter(const FGeometry& MyGeometry, const FPointerEvent
 
 	SMultiBlockBaseWidget::OnMouseEnter(MyGeometry, MouseEvent);
 }
-
diff --git a/Engine/Source/Runtime/Slate/Private/Framework/MultiBox/SWidgetBlock.h b/Engine/Source/Runtime/Slate/Private/Framework/MultiBox/SWidgetBlock.h
--- a/Engine/Source/Runtime/Slate/Private/Framework/MultiBox/SWidgetBlock.h
+++ b/Engine/Source/Runtime/Slate/Private/Framework/MultiBox/SWidgetBlock.h
@@ -8,7 +8,6 @@
 #include "Framework/MultiBox/MultiBox.h"
 #include "Framework/MultiBox/MultiBoxBuilder.h"
 
-
 /**
  * Arbitrary Widget MultiBlock
  */
@@ -17,16 +16,15 @@ class FWidgetBlock
 {
 
 public:
-
 	/**
 	 * Constructor
 	 *
-	 * @param	InContent		The widget to place in the block
-	 * @param	InLabel			Optional label text to be added to the left of the content
-	 * @param	bInNoIndent		If true, removes the padding from the left of the widget that lines it up with other menu items
-	 * @param	InToolTipText	Optional tooltip text to be added to the widget and label
+	 * @param	InContent					The widget to place in the block
+	 * @param	InLabel						Optional label text to be added to the left of the content
+	 * @param	InToolTipText				Optional tooltip text to be added to the widget and label
+	 * @param	InStyleParams				Optional additional style parameters. @see FMenuEntryStyleParams
 	 */
-	FWidgetBlock(TSharedRef<SWidget> InContent, const FText& InLabel, bool bInNoIndent, EHorizontalAlignment InHorizontalAlignment = HAlign_Fill, const TAttribute<FText>& InToolTipText = FText());
+	FWidgetBlock(TSharedRef<SWidget> InContent, const FText& InLabel, const TAttribute<FText>& InToolTipText = FText(), const FMenuEntryStyleParams& InStyleParams = {});
 
 	/** FMultiBlock interface */
 	virtual void CreateMenuEntry(class FMenuBuilder& MenuBuilder) const override;
@@ -38,7 +36,8 @@ private:
 
 	/** FMultiBlock private interface */
 	virtual TSharedRef< class IMultiBlockBaseWidget > ConstructWidget() const override;
-	virtual bool GetAlignmentOverrides(EHorizontalAlignment& OutHorizontalAlignment, EVerticalAlignment& OutVerticalAlignment, bool& bOutAutoWidth) const;
+	virtual bool GetAlignmentOverrides(FMenuEntryStyleParams& OutAlignmentParameters) const override;
+
 private:
 
 	// Friend our corresponding widget class
@@ -53,11 +52,8 @@ private:
 	/** Optional ToolTip text */
 	TAttribute<FText> ToolTipText;
 
-	/** Remove the padding from the left of the widget that lines it up with other menu items? */
-	bool bNoIndent;
-
-	/** Hortizontal aligment for this widget in its parent container. Note: only applies to toolbars */
-	EHorizontalAlignment HorizontalAlignment;
+	/** Style parameters */
+	FMenuEntryStyleParams StyleParams;
 
 	/** Optional delegate to customize when a menu appears instead of the widget, such as in toolbars */
 	FNewMenuDelegate CustomMenuDelegate;
diff --git a/Engine/Source/Runtime/Slate/Private/Widgets/Views/SHeaderRow.cpp b/Engine/Source/Runtime/Slate/Private/Widgets/Views/SHeaderRow.cpp
--- a/Engine/Source/Runtime/Slate/Private/Widgets/Views/SHeaderRow.cpp
+++ b/Engine/Source/Runtime/Slate/Private/Widgets/Views/SHeaderRow.cpp
@@ -771,10 +771,14 @@ void SHeaderRow::RegenerateWidgets()
 				TSharedRef<SWidget> HeaderMenuContent = SomeColumn.HeaderMenuContent.Widget;
 				if ( bCanSelectGeneratedColumn && HeaderMenuContent != SNullWidget::NullWidget )
 				{
-					const bool CloseAfterSelection = true;
-					const bool bNoIndent = true;
-					FMenuBuilder MenuBuilder(CloseAfterSelection, nullptr);
-					MenuBuilder.AddWidget(SomeColumn.HeaderMenuContent.Widget, FText::GetEmpty(), bNoIndent);
+					static constexpr bool bCloseAfterSelection = true;
+					static constexpr bool bNoIndent = true;
+
+					FMenuEntryStyleParams StyleParams;
+					StyleParams.bNoIndent = bNoIndent;
+
+					FMenuBuilder MenuBuilder(bCloseAfterSelection, nullptr);
+					MenuBuilder.AddWidget(SomeColumn.HeaderMenuContent.Widget, FText::GetEmpty(), StyleParams);
 					MenuBuilder.AddMenuSeparator();
 					MenuBuilder.AddSubMenu(
 						LOCTEXT("SelectColumns", "Select Columns"),
diff --git a/Engine/Source/Runtime/Slate/Public/Framework/MultiBox/MultiBox.h b/Engine/Source/Runtime/Slate/Public/Framework/MultiBox/MultiBox.h
--- a/Engine/Source/Runtime/Slate/Public/Framework/MultiBox/MultiBox.h
+++ b/Engine/Source/Runtime/Slate/Public/Framework/MultiBox/MultiBox.h
@@ -254,14 +254,23 @@ private:
 	virtual TSharedRef< class IMultiBlockBaseWidget > ConstructWidget() const = 0;
 
 	/**
- 	 * Gets any aligment overrides for this block
+ 	 * Gets any alignment overrides for this block
 	 *
-	 * @param OutHorizontalAligment	Horizontal alignment override
+	 * @param OutHorizontalAlignment	Horizontal alignment override
 	 * @param OutVerticalAlignment	Vertical Alignment override 
 	 * @param bOutAutoWidth		Fill or Auto width override
 	 * @return true if overrides should be applied, false to use defaults 
- 	 */ 
+ 	 */
+	UE_DEPRECATED(5.5, "Use the version that outputs FMenuEntryStyleParams.")
 	virtual bool GetAlignmentOverrides(EHorizontalAlignment& OutHorizontalAlignment, EVerticalAlignment& OutVerticalAlignment, bool& bOutAutoWidth) const { return false; }
+
+	/**
+	  * Gets any alignment overrides for this block
+	 *
+	 * @param OutAlignmentParameters	Alignment parameters and overrides
+	 * @return true if overrides should be applied, false to use defaults
+	  */
+	virtual bool GetAlignmentOverrides(FMenuEntryStyleParams& OutAlignmentParameters) const { return false; }
 private:
 
 	// We're friends with SMultiBoxWidget so that it can call MakeWidget() directly
diff --git a/Engine/Source/Runtime/Slate/Public/Framework/MultiBox/MultiBoxBuilder.h b/Engine/Source/Runtime/Slate/Public/Framework/MultiBox/MultiBoxBuilder.h
--- a/Engine/Source/Runtime/Slate/Public/Framework/MultiBox/MultiBoxBuilder.h
+++ b/Engine/Source/Runtime/Slate/Public/Framework/MultiBox/MultiBoxBuilder.h
@@ -384,11 +384,21 @@ public:
 	 * @param	InWidget			The widget that should be shown in the menu
 	 * @param	InLabel				Optional label text to be added to the left of the content
 	 * @param	bInNoIndent			If true, removes the padding from the left of the widget that lines it up with other menu items (default == false)
-	 * @param	bInSearchable			If true, widget will be searchable (default == true)
-	 * @param	InToolTipText	Optional tooltip text to be added to the widget and label
+	 * @param	bInSearchable		If true, widget will be searchable (default == true)
+	 * @param	InToolTipText		Optional tooltip text to be added to the widget and label
+	 */
+	SLATE_API void AddWidget( TSharedRef<SWidget> InWidget, const FText& InLabel, bool bInNoIndent, bool bInSearchable = true, const TAttribute<FText>&  InToolTipText = FText());
+
+	/**
+	 * Adds any widget to the menu
+	 *
+	 * @param	InWidget			The widget that should be shown in the menu
+	 * @param	InLabel				Optional label text to be added to the left of the content
+	 * @param	InStyleParams		Optional style parameters and overrides, @see FMenuEntryStyleParams
+	 * @param	bInSearchable		If true, widget will be searchable (default == true)
+	 * @param	InToolTipText		Optional tooltip text to be added to the widget and label
 	 */
-	SLATE_API void AddWidget( TSharedRef<SWidget> InWidget, const FText& InLabel, bool bInNoIndent = false, bool bInSearchable = true, const 
-	TAttribute<FText>&  InToolTipText = FText());
+	SLATE_API void AddWidget(const TSharedRef<SWidget>& InWidget, const FText& InLabel, const FMenuEntryStyleParams& InStyleParams = FMenuEntryStyleParams(false), bool bInSearchable = true, const TAttribute<FText>& InToolTipText = FText());
 
 	/**
 	* Adds the widget the multibox will use for searching
@@ -609,8 +619,19 @@ public:
 	 * @param	InAlignment				Horizontal alignment for the widget inside the toolbar
 	 * @param	InCustomMenuDelegate	Optional custom menu delegate for cases where the toolbar is compressed into a menu
 	 */
-	SLATE_API void AddWidget(TSharedRef<SWidget> InWidget, FName InTutorialHighlightName = NAME_None, bool bInSearchable = true, EHorizontalAlignment InAlignment = HAlign_Fill, FNewMenuDelegate InCustomMenuDelegate = FNewMenuDelegate());
-	
+	SLATE_API void AddWidget(TSharedRef<SWidget> InWidget, FName InTutorialHighlightName, bool bInSearchable = true, EHorizontalAlignment InAlignment = HAlign_Fill, FNewMenuDelegate InCustomMenuDelegate = FNewMenuDelegate());
+
+	/**
+	 * Adds any widget to the toolbar
+	 *
+	 * @param	InWidget				The widget that should be shown in the toolbar
+	 * @param	InStyleParams			Style parameters for the widget inside the toolbar
+	 * @param	InTutorialHighlightName	Name to identify this widget and highlight during tutorials
+	 * @param	bInSearchable			If true, widget will be searchable (default == true)
+	 * @param	InCustomMenuDelegate	Optional custom menu delegate for cases where the toolbar is compressed into a menu
+	 */
+	SLATE_API void AddWidget(TSharedRef<SWidget> InWidget, const FMenuEntryStyleParams& InStyleParams = FMenuEntryStyleParams(HAlign_Fill), FName InTutorialHighlightName = NAME_None, bool bInSearchable = true, FNewMenuDelegate InCustomMenuDelegate = FNewMenuDelegate());
+
 	/**
 	 * Adds a toolbar separator
 	 */
diff --git a/Engine/Source/Runtime/Slate/Public/Framework/MultiBox/MultiBoxDefs.h b/Engine/Source/Runtime/Slate/Public/Framework/MultiBox/MultiBoxDefs.h
--- a/Engine/Source/Runtime/Slate/Public/Framework/MultiBox/MultiBoxDefs.h
+++ b/Engine/Source/Runtime/Slate/Public/Framework/MultiBox/MultiBoxDefs.h
@@ -6,6 +6,7 @@
 #include "Misc/Attribute.h"
 #include "Styling/ISlateStyle.h"
 #include "Framework/Commands/UICommandInfo.h"
+#include "Types/SlateStructs.h"
 #include "MultiBoxDefs.generated.h"
 
 class SToolTip;
@@ -137,3 +138,37 @@ namespace EMultiBlockLocation
 		return StyleName;
 	}
 }
+
+/** Contains various Style parameters and overrides. Not all are applicable to a given entry */
+struct FMenuEntryStyleParams
+{
+	FMenuEntryStyleParams() = default;
+
+	explicit FMenuEntryStyleParams(const EHorizontalAlignment InHorizontalAlignment)
+		: HorizontalAlignment(InHorizontalAlignment)
+	{
+	}
+
+	explicit FMenuEntryStyleParams(const bool bInNoIndent)
+		: bNoIndent(bInNoIndent)
+	{
+	}
+
+	/** If true, removes the padding from the left of the widget that lines it up with other menu items */
+	bool bNoIndent = false;
+
+	/** Horizontal alignment for this widget in its parent container. Note: only applies to toolbars */
+	EHorizontalAlignment HorizontalAlignment = HAlign_Fill;
+
+	/** (Optional) Vertical alignment for this widget in its parent container */
+	TOptional<EVerticalAlignment> VerticalAlignment;
+
+	/** (Optionally) override the size rule, where the default is generally Auto */
+	TOptional<FSizeParam::ESizeRule> SizeRule;
+
+	/** (Optionally) override the minimum size. This will apply to the width or height, depending on the menu orientation */
+	TOptional<float> MinSize;
+
+	/** (Optionally) override the maximum size. This will apply to the width or height, depending on the menu orientation */
+	TOptional<float> MaxSize;
+};
