[Slate] Removed unecessary constructors from FMenuEntryStyleParams

diff --git a/Engine/Source/Runtime/Slate/Private/Framework/MultiBox/MultiBoxBuilder.cpp b/Engine/Source/Runtime/Slate/Private/Framework/MultiBox/MultiBoxBuilder.cpp
--- a/Engine/Source/Runtime/Slate/Private/Framework/MultiBox/MultiBoxBuilder.cpp
+++ b/Engine/Source/Runtime/Slate/Private/Framework/MultiBox/MultiBoxBuilder.cpp
@@ -619,14 +619,21 @@ void FToolBarBuilder::AddToolBarWidget( TSharedRef<SWidget> InWidget, const TAtt
 		] ;
 
 	static const TAttribute<FText> EmptyTextAttribute;
-	TSharedRef<FWidgetBlock> NewWidgetBlock = MakeShared<FWidgetBlock>(InWidget, FText::GetEmpty(), EmptyTextAttribute, FMenuEntryStyleParams(true));
+
+	FMenuEntryStyleParams StyleParams;
+	StyleParams.bNoIndent = true;
+
+	TSharedRef<FWidgetBlock> NewWidgetBlock = MakeShared<FWidgetBlock>(InWidget, FText::GetEmpty(), EmptyTextAttribute, StyleParams);
 	MultiBox->AddMultiBlock( NewWidgetBlock );
 	NewWidgetBlock->SetSearchable(bSearchable);
 }
 
 void FToolBarBuilder::AddWidget( TSharedRef<SWidget> InWidget, FName InTutorialHighlightName, bool bSearchable, EHorizontalAlignment Alignment, FNewMenuDelegate InCustomMenuDelegate)
 {
-	AddWidget(InWidget, FMenuEntryStyleParams(Alignment), InTutorialHighlightName, bSearchable, InCustomMenuDelegate);
+	FMenuEntryStyleParams StyleParams;
+	StyleParams.HorizontalAlignment = Alignment;
+
+	AddWidget(InWidget, StyleParams, InTutorialHighlightName, bSearchable, InCustomMenuDelegate);
 }
 
 void FToolBarBuilder::AddWidget(TSharedRef<SWidget> InWidget, const FMenuEntryStyleParams& InStyleParams, FName InTutorialHighlightName, bool bInSearchable, FNewMenuDelegate InCustomMenuDelegate)
diff --git a/Engine/Source/Runtime/Slate/Public/Framework/MultiBox/MultiBoxBuilder.h b/Engine/Source/Runtime/Slate/Public/Framework/MultiBox/MultiBoxBuilder.h
--- a/Engine/Source/Runtime/Slate/Public/Framework/MultiBox/MultiBoxBuilder.h
+++ b/Engine/Source/Runtime/Slate/Public/Framework/MultiBox/MultiBoxBuilder.h
@@ -387,7 +387,7 @@ public:
 	 * @param	bInSearchable		If true, widget will be searchable (default == true)
 	 * @param	InToolTipText		Optional tooltip text to be added to the widget and label
 	 */
-	SLATE_API void AddWidget( TSharedRef<SWidget> InWidget, const FText& InLabel, bool bInNoIndent, bool bInSearchable = true, const TAttribute<FText>&  InToolTipText = FText());
+	SLATE_API void AddWidget( TSharedRef<SWidget> InWidget, const FText& InLabel, bool bInNoIndent = false, bool bInSearchable = true, const TAttribute<FText>&  InToolTipText = FText());
 
 	/**
 	 * Adds any widget to the menu
@@ -398,7 +398,7 @@ public:
 	 * @param	bInSearchable		If true, widget will be searchable (default == true)
 	 * @param	InToolTipText		Optional tooltip text to be added to the widget and label
 	 */
-	SLATE_API void AddWidget(const TSharedRef<SWidget>& InWidget, const FText& InLabel, const FMenuEntryStyleParams& InStyleParams = FMenuEntryStyleParams(false), bool bInSearchable = true, const TAttribute<FText>& InToolTipText = FText());
+	SLATE_API void AddWidget(const TSharedRef<SWidget>& InWidget, const FText& InLabel, const FMenuEntryStyleParams& InStyleParams, bool bInSearchable = true, const TAttribute<FText>& InToolTipText = FText());
 
 	/**
 	* Adds the widget the multibox will use for searching
@@ -619,7 +619,7 @@ public:
 	 * @param	InAlignment				Horizontal alignment for the widget inside the toolbar
 	 * @param	InCustomMenuDelegate	Optional custom menu delegate for cases where the toolbar is compressed into a menu
 	 */
-	SLATE_API void AddWidget(TSharedRef<SWidget> InWidget, FName InTutorialHighlightName, bool bInSearchable = true, EHorizontalAlignment InAlignment = HAlign_Fill, FNewMenuDelegate InCustomMenuDelegate = FNewMenuDelegate());
+	SLATE_API void AddWidget(TSharedRef<SWidget> InWidget, FName InTutorialHighlightName = NAME_None, bool bInSearchable = true, EHorizontalAlignment InAlignment = HAlign_Fill, FNewMenuDelegate InCustomMenuDelegate = FNewMenuDelegate());
 
 	/**
 	 * Adds any widget to the toolbar
@@ -630,7 +630,7 @@ public:
 	 * @param	bInSearchable			If true, widget will be searchable (default == true)
 	 * @param	InCustomMenuDelegate	Optional custom menu delegate for cases where the toolbar is compressed into a menu
 	 */
-	SLATE_API void AddWidget(TSharedRef<SWidget> InWidget, const FMenuEntryStyleParams& InStyleParams = FMenuEntryStyleParams(HAlign_Fill), FName InTutorialHighlightName = NAME_None, bool bInSearchable = true, FNewMenuDelegate InCustomMenuDelegate = FNewMenuDelegate());
+	SLATE_API void AddWidget(TSharedRef<SWidget> InWidget, const FMenuEntryStyleParams& InStyleParams, FName InTutorialHighlightName = NAME_None, bool bInSearchable = true, FNewMenuDelegate InCustomMenuDelegate = FNewMenuDelegate());
 
 	/**
 	 * Adds a toolbar separator
diff --git a/Engine/Source/Runtime/Slate/Public/Framework/MultiBox/MultiBoxDefs.h b/Engine/Source/Runtime/Slate/Public/Framework/MultiBox/MultiBoxDefs.h
--- a/Engine/Source/Runtime/Slate/Public/Framework/MultiBox/MultiBoxDefs.h
+++ b/Engine/Source/Runtime/Slate/Public/Framework/MultiBox/MultiBoxDefs.h
@@ -144,16 +144,6 @@ struct FMenuEntryStyleParams
 {
 	FMenuEntryStyleParams() = default;
 
-	explicit FMenuEntryStyleParams(const EHorizontalAlignment InHorizontalAlignment)
-		: HorizontalAlignment(InHorizontalAlignment)
-	{
-	}
-
-	explicit FMenuEntryStyleParams(const bool bInNoIndent)
-		: bNoIndent(bInNoIndent)
-	{
-	}
-
 	/** If true, removes the padding from the left of the widget that lines it up with other menu items */
 	bool bNoIndent = false;
 
