Add support for visibility overrides of multi blocks

* Add FMultiBlock::VisibilityOverride.
* Take a new InVisibilityOverride as parameter in FToolBarBuilder::AddToolBarButton, FToolBarBuilder::AddComboButton, FToolBarBuilder::AddWidget, and FToolBarBuilder::AddSeparator and forward to FMultiBlock::SetVisibilityOverride.
* Make SMenuEntryBlock, SToolBarButtonBlock, SToolBarComboButtonBlock, and SToolBarSeparatorBlock all use the visibility of the VisibilityOverride unless it returns Visible, in which case other methods get a chance to affect visibility too (command visibility, for example).

diff --git a/Engine/Source/Developer/ToolMenus/Private/ToolMenus.cpp b/Engine/Source/Developer/ToolMenus/Private/ToolMenus.cpp
--- a/Engine/Source/Developer/ToolMenus/Private/ToolMenus.cpp
+++ b/Engine/Source/Developer/ToolMenus/Private/ToolMenus.cpp
@@ -1595,28 +1595,24 @@ void UToolMenus::PopulateToolBarBuilderWithEntry(
 		ToolBarBuilder.BeginStyleOverride(OverrideStyleName);
 	}
 
-	FUIAction UIAction = UToolMenus::ConvertUIAction(Block, MenuData->Context);
+	TAttribute<EVisibility> VisibilityOverride;
 	if (bIsRaisingToTopLevel && Block.ShowInToolbarTopLevel.IsBound())
 	{
-		// Patch the IsVisible delegate with the top-level status of the menu entry.
-		UIAction.IsActionVisibleDelegate = FIsActionButtonVisible::CreateLambda(
-			[OriginalIsVisible = UIAction.IsActionVisibleDelegate, ShowInToolbarTopLevel = Block.ShowInToolbarTopLevel]() -> bool
+		VisibilityOverride = TAttribute<EVisibility>::CreateLambda(
+			[ShowInToolbarTopLevel = Block.ShowInToolbarTopLevel]() -> EVisibility
 			{
 				if (ShowInToolbarTopLevel.IsSet() && !ShowInToolbarTopLevel.Get())
 				{
-					return false;
+					return EVisibility::Collapsed;
 				}
 
-				if (OriginalIsVisible.IsBound())
-				{
-					return OriginalIsVisible.Execute();
-				}
-
-				return true;
+				return EVisibility::Visible;
 			}
 		);
 	}
 
+	const FUIAction UIAction = UToolMenus::ConvertUIAction(Block, MenuData->Context);
+
 	TAttribute<FText> Label;
 	{
 		const bool bHasIcon = Block.Icon.IsSet() || (Block.Command.IsValid() && Block.Command->GetIcon().IsSet());
@@ -1655,7 +1651,7 @@ void UToolMenus::PopulateToolBarBuilderWithEntry(
 			}
 
 			ToolBarBuilder.AddToolBarButton(
-				Block.Command, Block.Name, Label, Block.ToolTip, Block.Icon, Block.TutorialHighlightName
+				Block.Command, Block.Name, Label, Block.ToolTip, Block.Icon, Block.TutorialHighlightName, FNewMenuDelegate(), VisibilityOverride
 			);
 
 			if (bPopCommandList)
@@ -1667,15 +1663,21 @@ void UToolMenus::PopulateToolBarBuilderWithEntry(
 		{
 			UToolMenuEntryScript* ScriptObject = Block.ScriptObject;
 			TAttribute<FSlateIcon> Icon = ScriptObject->CreateIconAttribute(MenuData->Context);
-			ToolBarBuilder.AddToolBarButton(UIAction, ScriptObject->Data.Name,
+			ToolBarBuilder.AddToolBarButton(
+				UIAction,
+				ScriptObject->Data.Name,
 				ScriptObject->CreateLabelAttribute(MenuData->Context),
-				ScriptObject->CreateToolTipAttribute(MenuData->Context), Icon, Block.UserInterfaceActionType,
-				Block.TutorialHighlightName);
+				ScriptObject->CreateToolTipAttribute(MenuData->Context),
+				Icon,
+				Block.UserInterfaceActionType,
+				Block.TutorialHighlightName,
+				VisibilityOverride
+			);
 		}
 		else
 		{
 			ToolBarBuilder.AddToolBarButton(
-				UIAction, Block.Name, Label, Block.ToolTip, Block.Icon, Block.UserInterfaceActionType, Block.TutorialHighlightName
+				UIAction, Block.Name, Label, Block.ToolTip, Block.Icon, Block.UserInterfaceActionType, Block.TutorialHighlightName, VisibilityOverride
 			);
 		}
 
@@ -1690,7 +1692,8 @@ void UToolMenus::PopulateToolBarBuilderWithEntry(
 				Block.ToolBarData.OptionsDropdownData->ToolTip,
 				Block.Icon,
 				true,
-				Block.TutorialHighlightName
+				Block.TutorialHighlightName,
+				VisibilityOverride
 			);
 		}
 	}
@@ -1708,7 +1711,8 @@ void UToolMenus::PopulateToolBarBuilderWithEntry(
 				Block.ToolTip,
 				Block.Icon,
 				Block.ToolBarData.bSimpleComboBox,
-				Block.TutorialHighlightName
+				Block.TutorialHighlightName,
+				VisibilityOverride
 			);
 		}
 		else
@@ -1717,13 +1721,7 @@ void UToolMenus::PopulateToolBarBuilderWithEntry(
 				this, &UToolMenus::GenerateToolbarComboButtonMenu, TWeakObjectPtr<UToolMenu>(MenuData), Block.Name);
 
 			ToolBarBuilder.AddComboButton(
-				UIAction,
-				Delegate,
-				Label,
-				Block.ToolTip,
-				Block.Icon,
-				Block.ToolBarData.bSimpleComboBox,
-				Block.TutorialHighlightName
+				UIAction, Delegate, Label, Block.ToolTip, Block.Icon, Block.ToolBarData.bSimpleComboBox, Block.TutorialHighlightName, VisibilityOverride
 			);
 
 			// Also add any top-level flagged children to the toolbar.
@@ -1759,7 +1757,7 @@ void UToolMenus::PopulateToolBarBuilderWithEntry(
 		FMenuEntryStyleParams StyleParams = Block.WidgetData.StyleParams;
 		StyleParams.HorizontalAlignment = HAlign_Fill;
 
-		ToolBarBuilder.AddWidget(Widget.ToSharedRef(), StyleParams, Block.TutorialHighlightName, Block.WidgetData.bSearchable);
+		ToolBarBuilder.AddWidget(Widget.ToSharedRef(), StyleParams, Block.TutorialHighlightName, Block.WidgetData.bSearchable, FNewMenuDelegate(), VisibilityOverride);
 	}
 	else
 	{
diff --git a/Engine/Source/Runtime/Slate/Private/Framework/MultiBox/MultiBox.cpp b/Engine/Source/Runtime/Slate/Private/Framework/MultiBox/MultiBox.cpp
--- a/Engine/Source/Runtime/Slate/Private/Framework/MultiBox/MultiBox.cpp
+++ b/Engine/Source/Runtime/Slate/Private/Framework/MultiBox/MultiBox.cpp
@@ -263,6 +263,16 @@ bool FMultiBlock::GetSearchable() const
 	return bSearchable;
 }
 
+TAttribute<EVisibility> FMultiBlock::GetVisibilityOverride() const
+{
+	return VisibilityOverride;
+}
+
+void FMultiBlock::SetVisibilityOverride(TAttribute<EVisibility> InVisibilityOverride)
+{
+	VisibilityOverride = InVisibilityOverride;
+}
+
 /**
  * Constructor
  *
diff --git a/Engine/Source/Runtime/Slate/Private/Framework/MultiBox/MultiBoxBuilder.cpp b/Engine/Source/Runtime/Slate/Private/Framework/MultiBox/MultiBoxBuilder.cpp
--- a/Engine/Source/Runtime/Slate/Private/Framework/MultiBox/MultiBoxBuilder.cpp
+++ b/Engine/Source/Runtime/Slate/Private/Framework/MultiBox/MultiBoxBuilder.cpp
@@ -456,8 +456,16 @@ void FToolBarBuilder::SetIsFocusable(bool bInIsFocusable)
 	MultiBox->bIsFocusable = bIsFocusable;
 }
 
-
-void FToolBarBuilder::AddToolBarButton(const TSharedPtr< const FUICommandInfo > InCommand, FName InExtensionHook, const TAttribute<FText>& InLabelOverride, const TAttribute<FText>& InToolTipOverride, const TAttribute<FSlateIcon>& InIconOverride, FName InTutorialHighlightName, FNewMenuDelegate InCustomMenuDelegate )
+void FToolBarBuilder::AddToolBarButton(
+	const TSharedPtr<const FUICommandInfo> InCommand,
+	FName InExtensionHook,
+	const TAttribute<FText>& InLabelOverride,
+	const TAttribute<FText>& InToolTipOverride,
+	const TAttribute<FSlateIcon>& InIconOverride,
+	FName InTutorialHighlightName,
+	FNewMenuDelegate InCustomMenuDelegate,
+	TAttribute<EVisibility> InVisibilityOverride
+)
 {
 	ApplySectionBeginning();
 
@@ -475,6 +483,7 @@ void FToolBarBuilder::AddToolBarButton(const TSharedPtr< const FUICommandInfo >
 	NewToolBarButtonBlock->SetTutorialHighlightName(GenerateTutorialIdentifierName(TutorialHighlightName, InTutorialHighlightName, InCommand, MultiBox->GetBlocks().Num()));
 	NewToolBarButtonBlock->SetStyleNameOverride(CurrentStyleOverride);
 	NewToolBarButtonBlock->SetCustomMenuDelegate(InCustomMenuDelegate);
+	NewToolBarButtonBlock->SetVisibilityOverride(InVisibilityOverride);
 
 	MultiBox->AddMultiBlock( NewToolBarButtonBlock );
 
@@ -518,7 +527,16 @@ void FToolBarBuilder::AddToolBarButton(const FButtonArgs& Args)
 	ApplyHook(Args.ExtensionHook, EExtensionHook::After);
 }
 
-void FToolBarBuilder::AddToolBarButton(const FUIAction& InAction, FName InExtensionHook, const TAttribute<FText>& InLabelOverride, const TAttribute<FText>& InToolTipOverride, const TAttribute<FSlateIcon>& InIconOverride, const EUserInterfaceActionType UserInterfaceActionType, FName InTutorialHighlightName )
+void FToolBarBuilder::AddToolBarButton(
+	const FUIAction& InAction,
+	FName InExtensionHook,
+	const TAttribute<FText>& InLabelOverride,
+	const TAttribute<FText>& InToolTipOverride,
+	const TAttribute<FSlateIcon>& InIconOverride,
+	const EUserInterfaceActionType UserInterfaceActionType,
+	FName InTutorialHighlightName,
+	TAttribute<EVisibility> InVisibilityOverride
+)
 {
 	ApplySectionBeginning();
 
@@ -535,13 +553,23 @@ void FToolBarBuilder::AddToolBarButton(const FUIAction& InAction, FName InExtens
 	NewToolBarButtonBlock->SetForceSmallIcons(bForceSmallIcons);
 	NewToolBarButtonBlock->SetTutorialHighlightName(GenerateTutorialIdentifierName(TutorialHighlightName, InTutorialHighlightName, nullptr, MultiBox->GetBlocks().Num()));
 	NewToolBarButtonBlock->SetStyleNameOverride(CurrentStyleOverride);
+	NewToolBarButtonBlock->SetVisibilityOverride(InVisibilityOverride);
 
 	MultiBox->AddMultiBlock( NewToolBarButtonBlock );
 
 	ApplyHook(InExtensionHook, EExtensionHook::After);
 }
 
-void FToolBarBuilder::AddComboButton( const FUIAction& InAction, const FOnGetContent& InMenuContentGenerator, const TAttribute<FText>& InLabelOverride, const TAttribute<FText>& InToolTipOverride, const TAttribute<FSlateIcon>& InIconOverride, bool bInSimpleComboBox, FName InTutorialHighlightName )
+void FToolBarBuilder::AddComboButton(
+	const FUIAction& InAction,
+	const FOnGetContent& InMenuContentGenerator,
+	const TAttribute<FText>& InLabelOverride,
+	const TAttribute<FText>& InToolTipOverride,
+	const TAttribute<FSlateIcon>& InIconOverride,
+	bool bInSimpleComboBox,
+	FName InTutorialHighlightName,
+	TAttribute<EVisibility> InVisibilityOverride
+)
 {
 	ApplySectionBeginning();
 
@@ -555,6 +583,7 @@ void FToolBarBuilder::AddComboButton( const FUIAction& InAction, const FOnGetCon
 	NewToolBarComboButtonBlock->SetForceSmallIcons(bForceSmallIcons);
 	NewToolBarComboButtonBlock->SetTutorialHighlightName(GenerateTutorialIdentifierName(TutorialHighlightName, InTutorialHighlightName, nullptr, MultiBox->GetBlocks().Num()));
 	NewToolBarComboButtonBlock->SetStyleNameOverride(CurrentStyleOverride);
+	NewToolBarComboButtonBlock->SetVisibilityOverride(InVisibilityOverride);
 
 	MultiBox->AddMultiBlock( NewToolBarComboButtonBlock );
 }
@@ -628,15 +657,29 @@ void FToolBarBuilder::AddToolBarWidget( TSharedRef<SWidget> InWidget, const TAtt
 	NewWidgetBlock->SetSearchable(bSearchable);
 }
 
-void FToolBarBuilder::AddWidget( TSharedRef<SWidget> InWidget, FName InTutorialHighlightName, bool bSearchable, EHorizontalAlignment Alignment, FNewMenuDelegate InCustomMenuDelegate)
+void FToolBarBuilder::AddWidget(
+	TSharedRef<SWidget> InWidget,
+	FName InTutorialHighlightName,
+	bool bSearchable,
+	EHorizontalAlignment Alignment,
+	FNewMenuDelegate InCustomMenuDelegate,
+	TAttribute<EVisibility> InVisibilityOverride
+)
 {
 	FMenuEntryStyleParams StyleParams;
 	StyleParams.HorizontalAlignment = Alignment;
 
-	AddWidget(InWidget, StyleParams, InTutorialHighlightName, bSearchable, InCustomMenuDelegate);
+	AddWidget(InWidget, StyleParams, InTutorialHighlightName, bSearchable, InCustomMenuDelegate, InVisibilityOverride);
 }
 
-void FToolBarBuilder::AddWidget(TSharedRef<SWidget> InWidget, const FMenuEntryStyleParams& InStyleParams, FName InTutorialHighlightName, bool bInSearchable, FNewMenuDelegate InCustomMenuDelegate)
+void FToolBarBuilder::AddWidget(
+	TSharedRef<SWidget> InWidget,
+	const FMenuEntryStyleParams& InStyleParams,
+	FName InTutorialHighlightName,
+	bool bInSearchable,
+	FNewMenuDelegate InCustomMenuDelegate,
+	TAttribute<EVisibility> InVisibilityOverride
+)
 {
 	ApplySectionBeginning();
 
@@ -656,9 +699,10 @@ void FToolBarBuilder::AddWidget(TSharedRef<SWidget> InWidget, const FMenuEntrySt
 	MultiBox->AddMultiBlock( NewWidgetBlock );
 	NewWidgetBlock->SetSearchable(bInSearchable);
 	NewWidgetBlock->SetCustomMenuDelegate(InCustomMenuDelegate);
+	NewWidgetBlock->SetVisibilityOverride(InVisibilityOverride);
 }
 
-void FToolBarBuilder::AddSeparator(FName InExtensionHook)
+void FToolBarBuilder::AddSeparator(FName InExtensionHook, TAttribute<EVisibility> InVisibilityOverride)
 {
 	ApplySectionBeginning();
 
@@ -666,6 +710,7 @@ void FToolBarBuilder::AddSeparator(FName InExtensionHook)
 
 	TSharedRef<FToolBarSeparatorBlock> NewSeparatorBlock = MakeShared<FToolBarSeparatorBlock>(InExtensionHook);
 	NewSeparatorBlock->SetStyleNameOverride(CurrentStyleOverride);
+	NewSeparatorBlock->SetVisibilityOverride(InVisibilityOverride);
 
 	MultiBox->AddMultiBlock(NewSeparatorBlock);
 
diff --git a/Engine/Source/Runtime/Slate/Private/Framework/MultiBox/SMenuEntryBlock.cpp b/Engine/Source/Runtime/Slate/Private/Framework/MultiBox/SMenuEntryBlock.cpp
--- a/Engine/Source/Runtime/Slate/Private/Framework/MultiBox/SMenuEntryBlock.cpp
+++ b/Engine/Source/Runtime/Slate/Private/Framework/MultiBox/SMenuEntryBlock.cpp
@@ -333,7 +333,18 @@ FText SMenuEntryBlock::GetFilteredToolTipText( TAttribute<FText> ToolTipText ) c
 
 EVisibility SMenuEntryBlock::GetVisibility() const
 {
-	TSharedPtr< const FUICommandList > ActionList = MultiBlock->GetActionList();
+	// Let the visibility override take prescedence here.
+	// However, if it returns Visible, let the other methods have a chance to change that.
+	if (MultiBlock->GetVisibilityOverride().IsSet())
+	{
+		const EVisibility OverrideVisibility = MultiBlock->GetVisibilityOverride().Get();
+		if (OverrideVisibility != EVisibility::Visible)
+		{
+			return OverrideVisibility;
+		}
+	}
+
+	TSharedPtr<const FUICommandList> ActionList = MultiBlock->GetActionList();
 	TSharedPtr< const FUICommandInfo > Action = MultiBlock->GetAction();
 	const FUIAction& DirectActions = MultiBlock->GetDirectActions();
 
diff --git a/Engine/Source/Runtime/Slate/Private/Framework/MultiBox/SToolBarButtonBlock.cpp b/Engine/Source/Runtime/Slate/Private/Framework/MultiBox/SToolBarButtonBlock.cpp
--- a/Engine/Source/Runtime/Slate/Private/Framework/MultiBox/SToolBarButtonBlock.cpp
+++ b/Engine/Source/Runtime/Slate/Private/Framework/MultiBox/SToolBarButtonBlock.cpp
@@ -557,7 +557,18 @@ bool SToolBarButtonBlock::IsEnabled() const
  */
 EVisibility SToolBarButtonBlock::GetBlockVisibility() const
 {
-	TSharedPtr< const FUICommandList > ActionList = MultiBlock->GetActionList();
+	// Let the visibility override take prescedence here.
+	// However, if it returns Visible, let the other methods have a chance to change that.
+	if (MultiBlock->GetVisibilityOverride().IsSet())
+	{
+		const EVisibility OverrideVisibility = MultiBlock->GetVisibilityOverride().Get();
+		if (OverrideVisibility != EVisibility::Visible)
+		{
+			return OverrideVisibility;
+		}
+	}
+
+	TSharedPtr<const FUICommandList> ActionList = MultiBlock->GetActionList();
 	const FUIAction& DirectActions = MultiBlock->GetDirectActions();
 	if( ActionList.IsValid() )
 	{
diff --git a/Engine/Source/Runtime/Slate/Private/Framework/MultiBox/SToolBarComboButtonBlock.cpp b/Engine/Source/Runtime/Slate/Private/Framework/MultiBox/SToolBarComboButtonBlock.cpp
--- a/Engine/Source/Runtime/Slate/Private/Framework/MultiBox/SToolBarComboButtonBlock.cpp
+++ b/Engine/Source/Runtime/Slate/Private/Framework/MultiBox/SToolBarComboButtonBlock.cpp
@@ -246,6 +246,17 @@ bool SToolBarComboButtonBlock::IsEnabled() const
 
 EVisibility SToolBarComboButtonBlock::GetVisibility() const
 {
+	// Let the visibility override take prescedence here.
+	// However, if it returns Visible, let the other methods have a chance to change that.
+	if (MultiBlock->GetVisibilityOverride().IsSet())
+	{
+		const EVisibility OverrideVisibility = MultiBlock->GetVisibilityOverride().Get();
+		if (OverrideVisibility != EVisibility::Visible)
+		{
+			return OverrideVisibility;
+		}
+	}
+
 	const FUIAction& UIAction = MultiBlock->GetDirectActions();
 	if (UIAction.IsActionVisibleDelegate.IsBound())
 	{
diff --git a/Engine/Source/Runtime/Slate/Private/Framework/MultiBox/SToolBarSeparatorBlock.cpp b/Engine/Source/Runtime/Slate/Private/Framework/MultiBox/SToolBarSeparatorBlock.cpp
--- a/Engine/Source/Runtime/Slate/Private/Framework/MultiBox/SToolBarSeparatorBlock.cpp
+++ b/Engine/Source/Runtime/Slate/Private/Framework/MultiBox/SToolBarSeparatorBlock.cpp
@@ -52,7 +52,7 @@ void SToolBarSeparatorBlock::Construct( const FArguments& InArgs )
 void SToolBarSeparatorBlock::BuildMultiBlockWidget(const ISlateStyle* StyleSet, const FName& StyleName)
 {
 	const FToolBarStyle& ToolBarStyle = StyleSet->GetWidgetStyle<FToolBarStyle>(StyleName);
-	
+
 	ChildSlot
 	[
 		SNew( SHorizontalBox )
@@ -69,4 +69,6 @@ void SToolBarSeparatorBlock::BuildMultiBlockWidget(const ISlateStyle* StyleSet,
 
 	// Add this widget to the search list of the multibox and hide it
 	OwnerMultiBoxWidget.Pin()->AddElement(this->AsWidget(), FText::GetEmpty(), MultiBlock->GetSearchable());
+
+	SetVisibility(MultiBlock->GetVisibilityOverride());
 }
diff --git a/Engine/Source/Runtime/Slate/Private/Framework/MultiBox/SToolBarStackButtonBlock.cpp b/Engine/Source/Runtime/Slate/Private/Framework/MultiBox/SToolBarStackButtonBlock.cpp
--- a/Engine/Source/Runtime/Slate/Private/Framework/MultiBox/SToolBarStackButtonBlock.cpp
+++ b/Engine/Source/Runtime/Slate/Private/Framework/MultiBox/SToolBarStackButtonBlock.cpp
@@ -491,6 +491,17 @@ bool SToolBarStackButtonBlock::IsEnabled() const
 
 EVisibility SToolBarStackButtonBlock::GetVisibility() const
 {
+	// Let the visibility override take prescedence here.
+	// However, if it returns Visible, let the other methods have a chance to change that.
+	if (MultiBlock->GetVisibilityOverride().IsSet())
+	{
+		const EVisibility OverrideVisibility = MultiBlock->GetVisibilityOverride().Get();
+		if (OverrideVisibility != EVisibility::Visible)
+		{
+			return OverrideVisibility;
+		}
+	}
+
 	const FUIAction& UIAction = MultiBlock->GetDirectActions();
 	if (UIAction.IsActionVisibleDelegate.IsBound())
 	{
diff --git a/Engine/Source/Runtime/Slate/Private/Framework/MultiBox/SWidgetBlock.cpp b/Engine/Source/Runtime/Slate/Private/Framework/MultiBox/SWidgetBlock.cpp
--- a/Engine/Source/Runtime/Slate/Private/Framework/MultiBox/SWidgetBlock.cpp
+++ b/Engine/Source/Runtime/Slate/Private/Framework/MultiBox/SWidgetBlock.cpp
@@ -182,6 +182,8 @@ void SWidgetBlock::BuildMultiBlockWidget(const ISlateStyle* StyleSet, const FNam
 			WidgetBlock->ContentWidget
 		]
 	];
+
+	SetVisibility(MultiBlock->GetVisibilityOverride());
 }
 
 TSharedRef<SWidget> SWidgetBlock::FindTextBlockWidget( TSharedRef<SWidget> Content )
diff --git a/Engine/Source/Runtime/Slate/Public/Framework/MultiBox/MultiBox.h b/Engine/Source/Runtime/Slate/Public/Framework/MultiBox/MultiBox.h
--- a/Engine/Source/Runtime/Slate/Public/Framework/MultiBox/MultiBox.h
+++ b/Engine/Source/Runtime/Slate/Public/Framework/MultiBox/MultiBox.h
@@ -245,6 +245,12 @@ public:
 	/** Gets the extension hook so users can see what hooks are where */
 	FName GetExtensionHook() const { return ExtensionHook; }
 
+	/** Get the visibility override for this block. */
+	TAttribute<EVisibility> GetVisibilityOverride() const;
+
+	/** Set the visibility override for this block. */
+	void SetVisibilityOverride(TAttribute<EVisibility> InVisibilityOverride);
+
 private:
 	/**
 	 * Allocates a widget for this type of MultiBlock.  Override this in derived classes.
@@ -301,6 +307,9 @@ private:
 
 	/** Whether this block is part of the heading blocks for a section */
 	bool bIsPartOfHeading;
+
+	/** Visibility override which, if set and returning anything but Visible, takes priority over other methods of determining visibility. */
+	TAttribute<EVisibility> VisibilityOverride;
 };
 
 
diff --git a/Engine/Source/Runtime/Slate/Public/Framework/MultiBox/MultiBoxBuilder.h b/Engine/Source/Runtime/Slate/Public/Framework/MultiBox/MultiBoxBuilder.h
--- a/Engine/Source/Runtime/Slate/Public/Framework/MultiBox/MultiBoxBuilder.h
+++ b/Engine/Source/Runtime/Slate/Public/Framework/MultiBox/MultiBoxBuilder.h
@@ -550,7 +550,6 @@ public:
 	 */
 	SLATE_API virtual void AddToolBarButton(const FButtonArgs& ButtonArgs);
 
-
 	/**
 	 * Adds a tool bar button
 	 *
@@ -560,10 +559,20 @@ public:
 	 * @param	InToolTipOverride		Optional tool tip override.	 If omitted, then the action's label will be used instead.
 	 * @param	InIconOverride			Optional name of the slate brush to use for the tool bar image.  If omitted, then the action's icon will be used instead.
 	 * @param	InTutorialHighlightName	Name to identify this widget and highlight during tutorials
-	 * @param	InCustomMenuDelegate  Optional custom menu delegate for cases where the toolbar is compressed into a menu
+	 * @param	InCustomMenuDelegate	Optional custom menu delegate for cases where the toolbar is compressed into a menu
+	 * @param	InVisibilityOverride	Optional visibility override which, if set and returning anything but Visible, overrides the visibility of the toolbar button.
 	 */
-	SLATE_API void AddToolBarButton(const TSharedPtr< const FUICommandInfo > InCommand, FName InExtensionHook = NAME_None, const TAttribute<FText>& InLabelOverride = TAttribute<FText>(), const TAttribute<FText>& InToolTipOverride = TAttribute<FText>(), const TAttribute<FSlateIcon>& InIconOverride = TAttribute<FSlateIcon>(), FName InTutorialHighlightName = NAME_None, FNewMenuDelegate InCustomMenuDelegate = FNewMenuDelegate());
-	
+	SLATE_API void AddToolBarButton(
+		const TSharedPtr<const FUICommandInfo> InCommand,
+		FName InExtensionHook = NAME_None,
+		const TAttribute<FText>& InLabelOverride = TAttribute<FText>(),
+		const TAttribute<FText>& InToolTipOverride = TAttribute<FText>(),
+		const TAttribute<FSlateIcon>& InIconOverride = TAttribute<FSlateIcon>(),
+		FName InTutorialHighlightName = NAME_None,
+		FNewMenuDelegate InCustomMenuDelegate = FNewMenuDelegate(),
+		TAttribute<EVisibility> InVisibilityOverride = TAttribute<EVisibility>()
+	);
+
 	/**
 	 * Adds a tool bar button
 	 *
@@ -574,21 +583,41 @@ public:
 	 * @param	InIconOverride			Optional icon to use for the tool bar image.  If omitted, then the action's icon will be used instead.
 	 * @param	UserInterfaceActionType	Type of interface action
 	 * @param	InTutorialHighlightName	Name to identify this widget and highlight during tutorials
+	 * @param	InVisibilityOverride	Optional visibility override which, if set and returning anything but Visible, overrides the visibility of the toolbar button.
 	 */
-	SLATE_API void AddToolBarButton(const FUIAction& InAction, FName InExtensionHook = NAME_None, const TAttribute<FText>& InLabelOverride = TAttribute<FText>(), const TAttribute<FText>& InToolTipOverride = TAttribute<FText>(), const TAttribute<FSlateIcon>& InIconOverride = TAttribute<FSlateIcon>(), const EUserInterfaceActionType UserInterfaceActionType = EUserInterfaceActionType::Button, FName InTutorialHighlightName = NAME_None );
+	SLATE_API void AddToolBarButton(
+		const FUIAction& InAction,
+		FName InExtensionHook = NAME_None,
+		const TAttribute<FText>& InLabelOverride = TAttribute<FText>(),
+		const TAttribute<FText>& InToolTipOverride = TAttribute<FText>(),
+		const TAttribute<FSlateIcon>& InIconOverride = TAttribute<FSlateIcon>(),
+		const EUserInterfaceActionType UserInterfaceActionType = EUserInterfaceActionType::Button,
+		FName InTutorialHighlightName = NAME_None,
+		TAttribute<EVisibility> InVisibilityOverride = TAttribute<EVisibility>()
+	);
 
 	/**
 	 * Adds a combo button
 	 *
 	 * @param	InAction					UI action that sets the enabled state for this combo button
-	 * @param	InMenuContentGenerator		Delegate that generates a widget for this combo button's menu content.  Called when the menu is summoned.
+	 * @param	InMenuContentGenerator		Delegate that generates a widget for this combo button's menu content. Called when the menu is summoned.
 	 * @param	InLabelOverride				Optional label override.  If omitted, then the action's label will be used instead.
 	 * @param	InToolTipOverride			Optional tool tip override.	 If omitted, then the action's label will be used instead.
 	 * @param	InIconOverride				Optional icon to use for the tool bar image.  If omitted, then the action's icon will be used instead.
 	 * @param	bInSimpleComboBox			If true, the icon and label won't be displayed
 	 * @param	InTutorialHighlightName		Name to identify this widget and highlight during tutorials
+	 * @param	InVisibilityOverride		Optional visibility override which, if set and returning anything but Visible, overrides the visibility of the combo button.
 	 */
-	SLATE_API void AddComboButton( const FUIAction& InAction, const FOnGetContent& InMenuContentGenerator, const TAttribute<FText>& InLabelOverride = TAttribute<FText>(), const TAttribute<FText>& InToolTipOverride = TAttribute<FText>(), const TAttribute<FSlateIcon>& InIconOverride = TAttribute<FSlateIcon>(), bool bInSimpleComboBox = false, FName InTutorialHighlightName = NAME_None );
+	SLATE_API void AddComboButton(
+		const FUIAction& InAction,
+		const FOnGetContent& InMenuContentGenerator,
+		const TAttribute<FText>& InLabelOverride = TAttribute<FText>(),
+		const TAttribute<FText>& InToolTipOverride = TAttribute<FText>(),
+		const TAttribute<FSlateIcon>& InIconOverride = TAttribute<FSlateIcon>(),
+		bool bInSimpleComboBox = false,
+		FName InTutorialHighlightName = NAME_None,
+		TAttribute<EVisibility> InVisibilityOverride = TAttribute<EVisibility>()
+	);
 
 	/**
 	 * Adds a tool bar stack button
@@ -609,17 +638,24 @@ public:
 	 */
 	SLATE_API void AddToolBarWidget(TSharedRef<SWidget> InWidget, const TAttribute<FText>& InLabel = TAttribute<FText>(), FName InTutorialHighlightName = NAME_None, bool bInSearchable = true, const TAttribute<FText>& InToolTip = TAttribute<FText>());
 
-
 	/**
 	 * Adds any widget to the toolbar
-	 * 
+	 *
 	 * @param	InWidget				The widget that should be shown in the toolbar
 	 * @param	InTutorialHighlightName	Name to identify this widget and highlight during tutorials
 	 * @param	bInSearchable			If true, widget will be searchable (default == true)
 	 * @param	InAlignment				Horizontal alignment for the widget inside the toolbar
 	 * @param	InCustomMenuDelegate	Optional custom menu delegate for cases where the toolbar is compressed into a menu
+	 * @param	InVisibilityOverride	Optional visibility override which, if set and returning anything but Visible, overrides the visibility of the widget.
 	 */
-	SLATE_API void AddWidget(TSharedRef<SWidget> InWidget, FName InTutorialHighlightName = NAME_None, bool bInSearchable = true, EHorizontalAlignment InAlignment = HAlign_Fill, FNewMenuDelegate InCustomMenuDelegate = FNewMenuDelegate());
+	SLATE_API void AddWidget(
+		TSharedRef<SWidget> InWidget,
+		FName InTutorialHighlightName = NAME_None,
+		bool bInSearchable = true,
+		EHorizontalAlignment InAlignment = HAlign_Fill,
+		FNewMenuDelegate InCustomMenuDelegate = FNewMenuDelegate(),
+		TAttribute<EVisibility> InVisibilityOverride = TAttribute<EVisibility>()
+	);
 
 	/**
 	 * Adds any widget to the toolbar
@@ -629,13 +665,27 @@ public:
 	 * @param	InTutorialHighlightName	Name to identify this widget and highlight during tutorials
 	 * @param	bInSearchable			If true, widget will be searchable (default == true)
 	 * @param	InCustomMenuDelegate	Optional custom menu delegate for cases where the toolbar is compressed into a menu
+	 * @param	InVisibilityOverride	Optional visibility override which, if set and returning anything but Visible, overrides the visibility of the widget.
 	 */
-	SLATE_API void AddWidget(TSharedRef<SWidget> InWidget, const FMenuEntryStyleParams& InStyleParams, FName InTutorialHighlightName = NAME_None, bool bInSearchable = true, FNewMenuDelegate InCustomMenuDelegate = FNewMenuDelegate());
+	SLATE_API void AddWidget(
+		TSharedRef<SWidget> InWidget,
+		const FMenuEntryStyleParams& InStyleParams,
+		FName InTutorialHighlightName = NAME_None,
+		bool bInSearchable = true,
+		FNewMenuDelegate InCustomMenuDelegate = FNewMenuDelegate(),
+		TAttribute<EVisibility> InVisibilityOverride = TAttribute<EVisibility>()
+	);
 
 	/**
 	 * Adds a toolbar separator
+	 *
+	 * @param	InExtensionHook			Optional extensibility hook.
+	 * @param	InVisibilityOverride	Optional visibility override which, if set and returning anything but Visible,
+	 * overrides the visibility of the separator.
 	 */
-	SLATE_API void AddSeparator(FName InExtensionHook = NAME_None);
+	SLATE_API void AddSeparator(
+		FName InExtensionHook = NAME_None, TAttribute<EVisibility> InVisibilityOverride = TAttribute<EVisibility>()
+	);
 
 	/**
 	 * Starts a section on to the extender section hook stack
