Add ToolbarLabelOverride to ToolBarButtonBlock and ToolBarComboButtonBlock

* Add FToolBarButtonBlock::ToolbarLabelOverride and FToolBarComboButtonBlock::ToolbarLabelOverride.
* Add InToolbarLabelOverride to FToolBarBuilder::AddToolBarButton and FToolBarBuilder::AddComboButton methods and forward it to FToolBarButtonBlock::ToolbarLabelOverride and FToolBarComboButtonBlock::ToolbarLabelOverride.
* Use the ToolbarLabelOverride, if set, as the label when the blocks are added to a toolbar. The ToolbarLabelOverride does not affect the label of menu entries representing the blocks in toolbar overflow >> menus.
* Set the ToolbarLabelOverride to the empty string for ToolMenu entries raised to the top-level of a toolbar because our design is for those to not have labels by default. These entries will still have a proper label in the overflow >> menu of a toolbar.

diff --git a/Engine/Source/Developer/ToolMenus/Private/ToolMenus.cpp b/Engine/Source/Developer/ToolMenus/Private/ToolMenus.cpp
--- a/Engine/Source/Developer/ToolMenus/Private/ToolMenus.cpp
+++ b/Engine/Source/Developer/ToolMenus/Private/ToolMenus.cpp
@@ -1613,23 +1613,18 @@ void UToolMenus::PopulateToolBarBuilderWithEntry(
 
 	const FUIAction UIAction = UToolMenus::ConvertUIAction(Block, MenuData->Context);
 
-	TAttribute<FText> Label;
+	TAttribute<FText> ToolbarLabelOverride;
+	if (Block.ToolbarLabelOverride.IsSet())
 	{
-		const bool bHasIcon = Block.Icon.IsSet() || (Block.Command.IsValid() && Block.Command->GetIcon().IsSet());
-
-		if (Block.ToolbarLabelOverride.IsSet())
-		{
-			Label = Block.ToolbarLabelOverride;
-		}
-		else if (!bIsRaisingToTopLevel || (bIsRaisingToTopLevel && !bHasIcon))
-		{
-			Label = Block.Label;
-		}
-		else
-		{
-			// This explicitly sets it empty, which will override the label on commands, for example.
-			Label = FText();
-		}
+		ToolbarLabelOverride = Block.ToolbarLabelOverride;
+	}
+	else if (const bool bHasIcon = Block.Icon.IsSet() || (Block.Command.IsValid() && Block.Command->GetIcon().IsSet());
+			 bHasIcon && bIsRaisingToTopLevel)
+	{
+		// Set the toolbar label to the empty string if we're raising an entry that has an icon. This makes
+		// raising/pinning of icons less annoying because the intended design is for them to not have a label. We can
+		// still use the ToolbarLabelOverride to bypass this.
+		ToolbarLabelOverride = FText();
 	}
 
 	if (Block.Type == EMultiBlockType::ToolBarButton || (Block.Type == EMultiBlockType::MenuEntry && !Block.IsSubMenu()))
@@ -1651,7 +1646,15 @@ void UToolMenus::PopulateToolBarBuilderWithEntry(
 			}
 
 			ToolBarBuilder.AddToolBarButton(
-				Block.Command, Block.Name, Label, Block.ToolTip, Block.Icon, Block.TutorialHighlightName, FNewMenuDelegate(), VisibilityOverride
+				Block.Command,
+				Block.Name,
+				Block.Label,
+				Block.ToolTip,
+				Block.Icon,
+				Block.TutorialHighlightName,
+				FNewMenuDelegate(),
+				VisibilityOverride,
+				ToolbarLabelOverride
 			);
 
 			if (bPopCommandList)
@@ -1671,13 +1674,22 @@ void UToolMenus::PopulateToolBarBuilderWithEntry(
 				Icon,
 				Block.UserInterfaceActionType,
 				Block.TutorialHighlightName,
-				VisibilityOverride
+				VisibilityOverride,
+				ToolbarLabelOverride
 			);
 		}
 		else
 		{
 			ToolBarBuilder.AddToolBarButton(
-				UIAction, Block.Name, Label, Block.ToolTip, Block.Icon, Block.UserInterfaceActionType, Block.TutorialHighlightName, VisibilityOverride
+				UIAction,
+				Block.Name,
+				Block.Label,
+				Block.ToolTip,
+				Block.Icon,
+				Block.UserInterfaceActionType,
+				Block.TutorialHighlightName,
+				VisibilityOverride,
+				ToolbarLabelOverride
 			);
 		}
 
@@ -1688,12 +1700,13 @@ void UToolMenus::PopulateToolBarBuilderWithEntry(
 			ToolBarBuilder.AddComboButton(
 				Block.ToolBarData.OptionsDropdownData->Action,
 				OnGetContent,
-				Label,
+				Block.Label,
 				Block.ToolBarData.OptionsDropdownData->ToolTip,
 				Block.Icon,
 				true,
 				Block.TutorialHighlightName,
-				VisibilityOverride
+				VisibilityOverride,
+				ToolbarLabelOverride
 			);
 		}
 	}
@@ -1707,12 +1720,13 @@ void UToolMenus::PopulateToolBarBuilderWithEntry(
 			ToolBarBuilder.AddComboButton(
 				UIAction,
 				OnGetContent,
-				Label,
+				Block.Label,
 				Block.ToolTip,
 				Block.Icon,
 				Block.ToolBarData.bSimpleComboBox,
 				Block.TutorialHighlightName,
-				VisibilityOverride
+				VisibilityOverride,
+				ToolbarLabelOverride
 			);
 		}
 		else
@@ -1721,7 +1735,15 @@ void UToolMenus::PopulateToolBarBuilderWithEntry(
 				this, &UToolMenus::GenerateToolbarComboButtonMenu, TWeakObjectPtr<UToolMenu>(MenuData), Block.Name);
 
 			ToolBarBuilder.AddComboButton(
-				UIAction, Delegate, Label, Block.ToolTip, Block.Icon, Block.ToolBarData.bSimpleComboBox, Block.TutorialHighlightName, VisibilityOverride
+				UIAction,
+				Delegate,
+				Block.Label,
+				Block.ToolTip,
+				Block.Icon,
+				Block.ToolBarData.bSimpleComboBox,
+				Block.TutorialHighlightName,
+				VisibilityOverride,
+				ToolbarLabelOverride
 			);
 
 			// Also add any top-level flagged children to the toolbar.
diff --git a/Engine/Source/Runtime/Slate/Private/Framework/MultiBox/MultiBoxBuilder.cpp b/Engine/Source/Runtime/Slate/Private/Framework/MultiBox/MultiBoxBuilder.cpp
--- a/Engine/Source/Runtime/Slate/Private/Framework/MultiBox/MultiBoxBuilder.cpp
+++ b/Engine/Source/Runtime/Slate/Private/Framework/MultiBox/MultiBoxBuilder.cpp
@@ -464,14 +464,17 @@ void FToolBarBuilder::AddToolBarButton(
 	const TAttribute<FSlateIcon>& InIconOverride,
 	FName InTutorialHighlightName,
 	FNewMenuDelegate InCustomMenuDelegate,
-	TAttribute<EVisibility> InVisibilityOverride
+	TAttribute<EVisibility> InVisibilityOverride,
+	TAttribute<FText> InToolbarLabelOverride
 )
 {
 	ApplySectionBeginning();
 
 	ApplyHook(InExtensionHook, EExtensionHook::Before);
 
-	TSharedRef< FToolBarButtonBlock > NewToolBarButtonBlock( new FToolBarButtonBlock( InCommand.ToSharedRef(), CommandListStack.Last(), InLabelOverride, InToolTipOverride, InIconOverride ) );
+	TSharedRef<FToolBarButtonBlock> NewToolBarButtonBlock(new FToolBarButtonBlock(
+		InCommand.ToSharedRef(), CommandListStack.Last(), InLabelOverride, InToolTipOverride, InIconOverride, InToolbarLabelOverride
+	));
 
 	if ( LabelVisibility.IsSet() )
 	{
@@ -497,11 +500,20 @@ void FToolBarBuilder::AddToolBarButton(const FButtonArgs& Args)
 	ApplyHook(Args.ExtensionHook, EExtensionHook::Before);
 
 	const bool bHasUserInteractionType = Args.UserInterfaceActionType != EUserInterfaceActionType::None;
-	
-	const TSharedPtr< FToolBarButtonBlock > NewToolBarButtonBlock( bHasUserInteractionType ?
-		new FToolBarButtonBlock( Args.LabelOverride, Args.ToolTipOverride, Args.IconOverride, Args.Action, Args.UserInterfaceActionType ) :
-		new FToolBarButtonBlock( Args.Command.ToSharedRef(), CommandListStack.Last(), Args.LabelOverride, Args.ToolTipOverride, Args.IconOverride )
-		);
+
+	const TSharedPtr<FToolBarButtonBlock> NewToolBarButtonBlock(
+		bHasUserInteractionType ? new FToolBarButtonBlock(
+			Args.LabelOverride, Args.ToolTipOverride, Args.IconOverride, Args.Action, Args.UserInterfaceActionType, Args.ToolbarLabelOverride
+		)
+								: new FToolBarButtonBlock(
+									Args.Command.ToSharedRef(),
+									CommandListStack.Last(),
+									Args.LabelOverride,
+									Args.ToolTipOverride,
+									Args.IconOverride,
+									Args.ToolbarLabelOverride
+								)
+	);
 
 	if ( LabelVisibility.IsSet() )
 	{
@@ -535,14 +547,17 @@ void FToolBarBuilder::AddToolBarButton(
 	const TAttribute<FSlateIcon>& InIconOverride,
 	const EUserInterfaceActionType UserInterfaceActionType,
 	FName InTutorialHighlightName,
-	TAttribute<EVisibility> InVisibilityOverride
+	TAttribute<EVisibility> InVisibilityOverride,
+	TAttribute<FText> InToolbarLabelOverride
 )
 {
 	ApplySectionBeginning();
 
 	ApplyHook(InExtensionHook, EExtensionHook::Before);
 
-	TSharedRef< FToolBarButtonBlock > NewToolBarButtonBlock( new FToolBarButtonBlock( InLabelOverride, InToolTipOverride, InIconOverride, InAction, UserInterfaceActionType ) );
+	TSharedRef<FToolBarButtonBlock> NewToolBarButtonBlock(new FToolBarButtonBlock(
+		InLabelOverride, InToolTipOverride, InIconOverride, InAction, UserInterfaceActionType, InToolbarLabelOverride
+	));
 
 	if ( LabelVisibility.IsSet() )
 	{
@@ -568,12 +583,15 @@ void FToolBarBuilder::AddComboButton(
 	const TAttribute<FSlateIcon>& InIconOverride,
 	bool bInSimpleComboBox,
 	FName InTutorialHighlightName,
-	TAttribute<EVisibility> InVisibilityOverride
+	TAttribute<EVisibility> InVisibilityOverride,
+	TAttribute<FText> InToolbarLabelOverride
 )
 {
 	ApplySectionBeginning();
 
-	TSharedRef<FToolBarComboButtonBlock> NewToolBarComboButtonBlock( new FToolBarComboButtonBlock( InAction, InMenuContentGenerator, InLabelOverride, InToolTipOverride, InIconOverride, bInSimpleComboBox ) );
+	TSharedRef<FToolBarComboButtonBlock> NewToolBarComboButtonBlock(new FToolBarComboButtonBlock(
+		InAction, InMenuContentGenerator, InLabelOverride, InToolTipOverride, InIconOverride, bInSimpleComboBox, InToolbarLabelOverride
+	));
 
 	if ( LabelVisibility.IsSet() )
 	{
diff --git a/Engine/Source/Runtime/Slate/Private/Framework/MultiBox/SToolBarButtonBlock.cpp b/Engine/Source/Runtime/Slate/Private/Framework/MultiBox/SToolBarButtonBlock.cpp
--- a/Engine/Source/Runtime/Slate/Private/Framework/MultiBox/SToolBarButtonBlock.cpp
+++ b/Engine/Source/Runtime/Slate/Private/Framework/MultiBox/SToolBarButtonBlock.cpp
@@ -13,25 +13,36 @@
 #include "Widgets/Images/SLayeredImage.h"
 #include "Widgets/Layout/SSeparator.h"
 
-FToolBarButtonBlock::FToolBarButtonBlock( FButtonArgs ButtonArgs )
-: FMultiBlock( ButtonArgs.Command, ButtonArgs.CommandList, NAME_None, EMultiBlockType::ToolBarButton )
-	, LabelOverride( ButtonArgs.LabelOverride )
-	, ToolTipOverride( ButtonArgs.ToolTipOverride )
-	, IconOverride( ButtonArgs.IconOverride )
+FToolBarButtonBlock::FToolBarButtonBlock(FButtonArgs ButtonArgs)
+	: FMultiBlock(ButtonArgs.Command, ButtonArgs.CommandList, NAME_None, EMultiBlockType::ToolBarButton)
+	, LabelOverride(ButtonArgs.LabelOverride)
+	, ToolbarLabelOverride(ButtonArgs.ToolbarLabelOverride)
+	, ToolTipOverride(ButtonArgs.ToolTipOverride)
+	, IconOverride(ButtonArgs.IconOverride)
 	, LabelVisibility()
-	, UserInterfaceActionType(ButtonArgs.UserInterfaceActionType != EUserInterfaceActionType::None ?
-		ButtonArgs.UserInterfaceActionType : EUserInterfaceActionType::Button)
+	, UserInterfaceActionType(
+		  ButtonArgs.UserInterfaceActionType != EUserInterfaceActionType::None ? ButtonArgs.UserInterfaceActionType
+																			   : EUserInterfaceActionType::Button
+	  )
 	, bIsFocusable(false)
 	, bForceSmallIcons(false)
-	, GetDecoratedButtonDelegate( ButtonArgs.GetDecoratedButtonDelegate )
+	, GetDecoratedButtonDelegate(ButtonArgs.GetDecoratedButtonDelegate)
 {
 }
 
-FToolBarButtonBlock::FToolBarButtonBlock( const TSharedPtr< const FUICommandInfo > InCommand, TSharedPtr< const FUICommandList > InCommandList, const TAttribute<FText>& InLabelOverride, const TAttribute<FText>& InToolTipOverride, const TAttribute<FSlateIcon>& InIconOverride )
-	: FMultiBlock( InCommand, InCommandList, NAME_None, EMultiBlockType::ToolBarButton )
-	, LabelOverride( InLabelOverride )
-	, ToolTipOverride( InToolTipOverride )
-	, IconOverride( InIconOverride )
+FToolBarButtonBlock::FToolBarButtonBlock(
+	const TSharedPtr<const FUICommandInfo> InCommand,
+	TSharedPtr<const FUICommandList> InCommandList,
+	const TAttribute<FText>& InLabelOverride,
+	const TAttribute<FText>& InToolTipOverride,
+	const TAttribute<FSlateIcon>& InIconOverride,
+	TAttribute<FText> InToolbarLabelOverride
+)
+	: FMultiBlock(InCommand, InCommandList, NAME_None, EMultiBlockType::ToolBarButton)
+	, LabelOverride(InLabelOverride)
+	, ToolbarLabelOverride(InToolbarLabelOverride)
+	, ToolTipOverride(InToolTipOverride)
+	, IconOverride(InIconOverride)
 	, LabelVisibility()
 	, UserInterfaceActionType(EUserInterfaceActionType::Button)
 	, bIsFocusable(false)
@@ -39,11 +50,19 @@ FToolBarButtonBlock::FToolBarButtonBlock( const TSharedPtr< const FUICommandInfo
 {
 }
 
-FToolBarButtonBlock::FToolBarButtonBlock( const TAttribute<FText>& InLabel, const TAttribute<FText>& InToolTip, const TAttribute<FSlateIcon>& InIcon, const FUIAction& InUIAction, const EUserInterfaceActionType InUserInterfaceActionType )
-	: FMultiBlock( InUIAction )
-	, LabelOverride( InLabel )
-	, ToolTipOverride( InToolTip )
-	, IconOverride( InIcon )
+FToolBarButtonBlock::FToolBarButtonBlock(
+	const TAttribute<FText>& InLabel,
+	const TAttribute<FText>& InToolTip,
+	const TAttribute<FSlateIcon>& InIcon,
+	const FUIAction& InUIAction,
+	const EUserInterfaceActionType InUserInterfaceActionType,
+	TAttribute<FText> InToolbarLabelOverride
+)
+	: FMultiBlock(InUIAction)
+	, LabelOverride(InLabel)
+	, ToolbarLabelOverride(InToolbarLabelOverride)
+	, ToolTipOverride(InToolTip)
+	, IconOverride(InIcon)
 	, LabelVisibility()
 	, UserInterfaceActionType(InUserInterfaceActionType)
 	, bIsFocusable(false)
@@ -191,7 +210,11 @@ void SToolBarButtonBlock::BuildMultiBlockWidget(const ISlateStyle* StyleSet, con
 
 	// Allow the block to override the action's label and tool tip string, if desired
 	TAttribute<FText> ActualLabel;
-	if (ToolBarButtonBlock->LabelOverride.IsSet())
+	if (ToolBarButtonBlock->ToolbarLabelOverride.IsSet())
+	{
+		ActualLabel = ToolBarButtonBlock->ToolbarLabelOverride;
+	}
+	else if (ToolBarButtonBlock->LabelOverride.IsSet())
 	{
 		ActualLabel = ToolBarButtonBlock->LabelOverride;
 	}
diff --git a/Engine/Source/Runtime/Slate/Private/Framework/MultiBox/SToolBarComboButtonBlock.cpp b/Engine/Source/Runtime/Slate/Private/Framework/MultiBox/SToolBarComboButtonBlock.cpp
--- a/Engine/Source/Runtime/Slate/Private/Framework/MultiBox/SToolBarComboButtonBlock.cpp
+++ b/Engine/Source/Runtime/Slate/Private/Framework/MultiBox/SToolBarComboButtonBlock.cpp
@@ -9,16 +9,24 @@
 #include "Styling/ToolBarStyle.h"
 #include "Widgets/Images/SLayeredImage.h"
 
-
-FToolBarComboButtonBlock::FToolBarComboButtonBlock( const FUIAction& InAction, const FOnGetContent& InMenuContentGenerator, const TAttribute<FText>& InLabel, const TAttribute<FText>& InToolTip, const TAttribute<FSlateIcon>& InIcon, bool bInSimpleComboBox )
-	: FMultiBlock( InAction, NAME_None, EMultiBlockType::ToolBarComboButton )
-	, MenuContentGenerator( InMenuContentGenerator )
-	, Label( InLabel )
-	, ToolTip( InToolTip )
-	, Icon( InIcon )
+FToolBarComboButtonBlock::FToolBarComboButtonBlock(
+	const FUIAction& InAction,
+	const FOnGetContent& InMenuContentGenerator,
+	const TAttribute<FText>& InLabel,
+	const TAttribute<FText>& InToolTip,
+	const TAttribute<FSlateIcon>& InIcon,
+	bool bInSimpleComboBox,
+	TAttribute<FText> InToolbarLabelOverride
+)
+	: FMultiBlock(InAction, NAME_None, EMultiBlockType::ToolBarComboButton)
+	, MenuContentGenerator(InMenuContentGenerator)
+	, Label(InLabel)
+	, ToolbarLabelOverride(InToolbarLabelOverride)
+	, ToolTip(InToolTip)
+	, Icon(InIcon)
 	, LabelVisibility()
-	, bSimpleComboBox( bInSimpleComboBox )
-	, bForceSmallIcons( false )
+	, bSimpleComboBox(bInSimpleComboBox)
+	, bForceSmallIcons(false)
 {
 }
 
@@ -111,7 +119,14 @@ void SToolBarComboButtonBlock::BuildMultiBlockWidget(const ISlateStyle* StyleSet
 			IconWidget = ActualIconWidget;
 		}
 
-		Label = ToolBarComboButtonBlock->Label;
+		if (ToolBarComboButtonBlock->ToolbarLabelOverride.IsSet())
+		{
+			Label = ToolBarComboButtonBlock->ToolbarLabelOverride;
+		}
+		else
+		{
+			Label = ToolBarComboButtonBlock->Label;
+		}
 	}
 
 	// Add this widget to the search list of the multibox
diff --git a/Engine/Source/Runtime/Slate/Public/Framework/MultiBox/MultiBoxBuilder.h b/Engine/Source/Runtime/Slate/Public/Framework/MultiBox/MultiBoxBuilder.h
--- a/Engine/Source/Runtime/Slate/Public/Framework/MultiBox/MultiBoxBuilder.h
+++ b/Engine/Source/Runtime/Slate/Public/Framework/MultiBox/MultiBoxBuilder.h
@@ -561,6 +561,7 @@ public:
 	 * @param	InTutorialHighlightName	Name to identify this widget and highlight during tutorials
 	 * @param	InCustomMenuDelegate	Optional custom menu delegate for cases where the toolbar is compressed into a menu
 	 * @param	InVisibilityOverride	Optional visibility override which, if set and returning anything but Visible, overrides the visibility of the toolbar button.
+	 * @param	InToolbarLabelOverride	Optional label override for wh.  If omitted, then the action's label will be used instead.
 	 */
 	SLATE_API void AddToolBarButton(
 		const TSharedPtr<const FUICommandInfo> InCommand,
@@ -570,7 +571,8 @@ public:
 		const TAttribute<FSlateIcon>& InIconOverride = TAttribute<FSlateIcon>(),
 		FName InTutorialHighlightName = NAME_None,
 		FNewMenuDelegate InCustomMenuDelegate = FNewMenuDelegate(),
-		TAttribute<EVisibility> InVisibilityOverride = TAttribute<EVisibility>()
+		TAttribute<EVisibility> InVisibilityOverride = TAttribute<EVisibility>(),
+		TAttribute<FText> InToolbarLabelOverride = TAttribute<FText>()
 	);
 
 	/**
@@ -584,6 +586,7 @@ public:
 	 * @param	UserInterfaceActionType	Type of interface action
 	 * @param	InTutorialHighlightName	Name to identify this widget and highlight during tutorials
 	 * @param	InVisibilityOverride	Optional visibility override which, if set and returning anything but Visible, overrides the visibility of the toolbar button.
+	 * @param	InToolbarLabelOverride	Optional label override for wh.  If omitted, then the action's label will be used instead.
 	 */
 	SLATE_API void AddToolBarButton(
 		const FUIAction& InAction,
@@ -593,8 +596,8 @@ public:
 		const TAttribute<FSlateIcon>& InIconOverride = TAttribute<FSlateIcon>(),
 		const EUserInterfaceActionType UserInterfaceActionType = EUserInterfaceActionType::Button,
 		FName InTutorialHighlightName = NAME_None,
-		TAttribute<EVisibility> InVisibilityOverride = TAttribute<EVisibility>()
-	);
+		TAttribute<EVisibility> InVisibilityOverride = TAttribute<EVisibility>(),
+		TAttribute<FText> InToolbarLabelOverride = TAttribute<FText>() );
 
 	/**
 	 * Adds a combo button
@@ -607,6 +610,7 @@ public:
 	 * @param	bInSimpleComboBox			If true, the icon and label won't be displayed
 	 * @param	InTutorialHighlightName		Name to identify this widget and highlight during tutorials
 	 * @param	InVisibilityOverride		Optional visibility override which, if set and returning anything but Visible, overrides the visibility of the combo button.
+	 * @param	InToolbarLabelOverride		Optional label override for wh.  If omitted, then the action's label will be used instead.
 	 */
 	SLATE_API void AddComboButton(
 		const FUIAction& InAction,
@@ -616,8 +620,8 @@ public:
 		const TAttribute<FSlateIcon>& InIconOverride = TAttribute<FSlateIcon>(),
 		bool bInSimpleComboBox = false,
 		FName InTutorialHighlightName = NAME_None,
-		TAttribute<EVisibility> InVisibilityOverride = TAttribute<EVisibility>()
-	);
+		TAttribute<EVisibility> InVisibilityOverride = TAttribute<EVisibility>(),
+		TAttribute<FText> InToolbarLabelOverride = TAttribute<FText>() );
 
 	/**
 	 * Adds a tool bar stack button
diff --git a/Engine/Source/Runtime/Slate/Public/Framework/MultiBox/SToolBarButtonBlock.h b/Engine/Source/Runtime/Slate/Public/Framework/MultiBox/SToolBarButtonBlock.h
--- a/Engine/Source/Runtime/Slate/Public/Framework/MultiBox/SToolBarButtonBlock.h
+++ b/Engine/Source/Runtime/Slate/Public/Framework/MultiBox/SToolBarButtonBlock.h
@@ -27,6 +27,7 @@ struct FButtonArgs : public TSharedFromThis<FButtonArgs>
 	TSharedPtr< const FUICommandList > CommandList;
 	FName ExtensionHook;
 	TAttribute<FText> LabelOverride;
+	TAttribute<FText> ToolbarLabelOverride;
 	TAttribute<FText> ToolTipOverride;
 	TAttribute<FSlateIcon> IconOverride;
 	FName TutorialHighlightName;
@@ -53,28 +54,47 @@ public:
 
 	
 	SLATE_API FToolBarButtonBlock( FButtonArgs ButtonArgs );
-	
+
 	/**
 	 * Constructor
 	 *
-	 * @param	InCommand			The command associated with this tool bar button
-	 * @param	InCommandList		The list of commands that are mapped to delegates so that we know what to execute for this button
-	 * @param	InLabelOverride		Optional label override.  If omitted, then the action's label will be used instead.
-	 * @param	InToolTipOverride	Optional tool tip override.	 If omitted, then the action's label will be used instead.
-	 * @param	InIconOverride		Optional icon to use for the tool bar image.  If omitted, then the action's icon will be used instead.
+	 * @param	InCommand				The command associated with this tool bar button
+	 * @param	InCommandList			The list of commands that are mapped to delegates so that we know what to execute for this button
+	 * @param	InLabelOverride			Optional label override.  If omitted, then the action's label will be used instead.
+	 * @param	InToolTipOverride		Optional tool tip override.	 If omitted, then the action's label will be used instead.
+	 * @param	InIconOverride			Optional icon to use for the tool bar image.  If omitted, then the action's icon will be used instead.
+	 * @param	InToolbarLabelOverride	Optional label to use when the block appears in a toolbar. If omitted, then the label override or command name will be used instead.
 	 */
-	SLATE_API FToolBarButtonBlock( const TSharedPtr< const FUICommandInfo > InCommand, TSharedPtr< const FUICommandList > InCommandList, const TAttribute<FText>& InLabelOverride = TAttribute<FText>(), const TAttribute<FText>& InToolTipOverride = TAttribute<FText>(), const TAttribute<FSlateIcon>& InIconOverride = TAttribute<FSlateIcon>() );
-	
+	SLATE_API FToolBarButtonBlock(
+		const TSharedPtr<const FUICommandInfo> InCommand,
+		TSharedPtr<const FUICommandList> InCommandList,
+		const TAttribute<FText>& InLabelOverride = TAttribute<FText>(),
+		const TAttribute<FText>& InToolTipOverride = TAttribute<FText>(),
+		const TAttribute<FSlateIcon>& InIconOverride = TAttribute<FSlateIcon>(),
+		TAttribute<FText> InToolbarLabelOverride = TAttribute<FText>()
+	);
+
 	/**
 	 * Constructor
 	 *
-	 * @param	InLabel				The label to display in the menu
-	 * @param	InToolTip			The tool tip to display when the menu entry is hovered over
-	 * @param	InIcon				The icon to display to the left of the label
-	 * @param	InUIAction			UI action to take when this menu item is clicked as well as to determine if the menu entry can be executed or appears "checked"
+	 * @param	InLabel						The label to display in the menu
+	 * @param	InToolTip					The tool tip to display when the menu entry is hovered over
+	 * @param	InIcon						The icon to display to the left of the label
+	 * @param	InUIAction					UI action to take when this menu item is clicked as well as to determine if
+	 * the menu entry can be executed or appears "checked"
 	 * @param	InUserInterfaceActionType	Type of interface action
+	 * @param	InToolbarLabelOverride	Optional label to use when the block appears in a toolbar. If omitted, then the
+	 * label override or command name will be used instead.
+	 *
 	 */
-	SLATE_API FToolBarButtonBlock( const TAttribute<FText>& InLabel, const TAttribute<FText>& InToolTip, const TAttribute<FSlateIcon>& InIcon, const FUIAction& InUIAction, const EUserInterfaceActionType InUserInterfaceActionType );
+	SLATE_API FToolBarButtonBlock(
+		const TAttribute<FText>& InLabel,
+		const TAttribute<FText>& InToolTip,
+		const TAttribute<FSlateIcon>& InIcon,
+		const FUIAction& InUIAction,
+		const EUserInterfaceActionType InUserInterfaceActionType,
+		TAttribute<FText> InToolbarLabelOverride = TAttribute<FText>()
+	);
 
 	void SetLabelVisibility( EVisibility InLabelVisibility ) { LabelVisibility = InLabelVisibility ; }
 
@@ -119,6 +139,9 @@ private:
 	/** Optional overridden text label for this tool bar button.  If not set, then the action's label will be used instead. */
 	TAttribute<FText> LabelOverride;
 
+	/** Optional overridden text label for when this tool bar button appears in a toolbar. If not set, then the label override or the action's label will be used instead. */
+	TAttribute<FText> ToolbarLabelOverride;
+
 	/** Optional overridden tool tip for this tool bar button.  If not set, then the action's tool tip will be used instead. */
 	TAttribute<FText> ToolTipOverride;
 	
diff --git a/Engine/Source/Runtime/Slate/Public/Framework/MultiBox/SToolBarComboButtonBlock.h b/Engine/Source/Runtime/Slate/Public/Framework/MultiBox/SToolBarComboButtonBlock.h
--- a/Engine/Source/Runtime/Slate/Public/Framework/MultiBox/SToolBarComboButtonBlock.h
+++ b/Engine/Source/Runtime/Slate/Public/Framework/MultiBox/SToolBarComboButtonBlock.h
@@ -20,18 +20,28 @@ class FToolBarComboButtonBlock
 {
 
 public:
-
 	/**
 	 * Constructor
 	 *
-	 * @param	InAction					UI action that sets the enabled state for this combo button
-	 * @param	InMenuContentGenerator		Delegate that generates a widget for this combo button's menu content.  Called when the menu is summoned.
-	 * @param	InLabel						Optional label for this combo button.
-	 * @param	InToolTip					Tool tip string (required!)
-	 * @param	InIcon						Optional icon to use for the tool bar image
-	 * @param	bInSimpleComboBox			If true, the icon and label won't be displayed
+	 * @param	InAction				UI action that sets the enabled state for this combo button
+	 * @param	InMenuContentGenerator	Delegate that generates a widget for this combo button's menu content.
+	 * Called when the menu is summoned.
+	 * @param	InLabel					Optional label for this combo button.
+	 * @param	InToolTip				Tool tip string (required!)
+	 * @param	InIcon					Optional icon to use for the tool bar image
+	 * @param	bInSimpleComboBox		If true, the icon and label won't be displayed
+	 * @param	InToolbarLabelOverride	Optional label to use when the block appears in a toolbar. If omitted, then the
+	 * label override or command name will be used instead.
 	 */
-	FToolBarComboButtonBlock( const FUIAction& InAction, const FOnGetContent& InMenuContentGenerator, const TAttribute<FText>& InLabel = TAttribute<FText>(), const TAttribute<FText>& InToolTip = TAttribute<FText>(), const TAttribute<FSlateIcon>& InIcon = TAttribute<FSlateIcon>(), bool bInSimpleComboBox = false );
+	FToolBarComboButtonBlock(
+		const FUIAction& InAction,
+		const FOnGetContent& InMenuContentGenerator,
+		const TAttribute<FText>& InLabel = TAttribute<FText>(),
+		const TAttribute<FText>& InToolTip = TAttribute<FText>(),
+		const TAttribute<FSlateIcon>& InIcon = TAttribute<FSlateIcon>(),
+		bool bInSimpleComboBox = false,
+		TAttribute<FText> InToolbarLabelOverride = TAttribute<FText>()
+	);
 
 	/** FMultiBlock interface */
 	virtual void CreateMenuEntry(class FMenuBuilder& MenuBuilder) const override;
@@ -69,6 +79,9 @@ private:
 	/** Optional overridden text label for this tool bar button.  If not set, then the action's label will be used instead. */
 	TAttribute<FText> Label;
 
+	/** Optional overridden text label for when this tool bar button appears in a toolbar. If not set, then the label override or the action's label will be used instead. */
+	TAttribute<FText> ToolbarLabelOverride;
+
 	/** Optional overridden tool tip for this tool bar button.  If not set, then the action's tool tip will be used instead. */
 	TAttribute<FText> ToolTip;
 
