[Slate] Fixed SClippingHorizontalBox overflow and padding

#rnx

diff --git a/Engine/Plugins/FX/Niagara/Source/NiagaraEditorWidgets/Private/NiagaraEditorWidgetsStyle.cpp b/Engine/Plugins/FX/Niagara/Source/NiagaraEditorWidgets/Private/NiagaraEditorWidgetsStyle.cpp
--- a/Engine/Plugins/FX/Niagara/Source/NiagaraEditorWidgets/Private/NiagaraEditorWidgetsStyle.cpp
+++ b/Engine/Plugins/FX/Niagara/Source/NiagaraEditorWidgets/Private/NiagaraEditorWidgetsStyle.cpp
@@ -448,7 +448,7 @@ void FNiagaraEditorWidgetsStyle::InitToolbarStyles()
 {
 	FToolBarStyle SlimToolbarStyle = FAppStyle::Get().GetWidgetStyle<FToolBarStyle>("SlimToolBar");
 	SlimToolbarStyle
-		.SetBackgroundPadding(FMargin(4.0f, 0.0f, 0.0f, 0.0f))
+		.SetBackgroundPadding(FMargin(4.0f, 0.0f))
 		.SetIconSize(Icon16x16);
 	SlimToolbarStyle.ToggleButton
 		.SetPadding(FMargin(8.0f, 2.0f));
@@ -463,7 +463,7 @@ void FNiagaraEditorWidgetsStyle::InitInlineDynamicInputBrushes()
 	Set("NiagaraEditor.Stack.InlineDynamicInput.GraphNode.Row.Bottom",			new BOX_BRUSH("Icons/MultiPartBorderRowBottom", FMargin(1.0f / 4.0f)));
 	Set("NiagaraEditor.Stack.InlineDynamicInput.GraphNode.Column.Left",			new BOX_BRUSH("Icons/MultiPartBorderColumnLeft", FMargin(1.0f / 4.0f)));
 	Set("NiagaraEditor.Stack.InlineDynamicInput.GraphNode.Column.Middle",		new BOX_BRUSH("Icons/MultiPartBorderColumnMiddle", FMargin(1.0f / 4.0f)));
-	Set("NiagaraEditor.Stack.InlineDynamicInput.GraphNode.Column.Right",			new BOX_BRUSH("Icons/MultiPartBorderColumnRight", FMargin(1.0f / 4.0f)));
+	Set("NiagaraEditor.Stack.InlineDynamicInput.GraphNode.Column.Right",		new BOX_BRUSH("Icons/MultiPartBorderColumnRight", FMargin(1.0f / 4.0f)));
 	Set("NiagaraEditor.Stack.InlineDynamicInput.GraphNode.Corner.BottomLeft",	new BOX_BRUSH("Icons/MultiPartBorderCornerBottomLeft", FMargin(1.0f / 4.0f)));
 	Set("NiagaraEditor.Stack.InlineDynamicInput.GraphNode.Corner.BottomRight",	new BOX_BRUSH("Icons/MultiPartBorderCornerBottomRight", FMargin(1.0f / 4.0f)));
 	Set("NiagaraEditor.Stack.InlineDynamicInput.GraphNode.Edge.Bottom",			new BOX_BRUSH("Icons/MultiPartBorderEdgeBottom", FMargin(1.0f / 4.0f)));
diff --git a/Engine/Source/Runtime/Slate/Private/Framework/MultiBox/SClippingHorizontalBox.cpp b/Engine/Source/Runtime/Slate/Private/Framework/MultiBox/SClippingHorizontalBox.cpp
--- a/Engine/Source/Runtime/Slate/Private/Framework/MultiBox/SClippingHorizontalBox.cpp
+++ b/Engine/Source/Runtime/Slate/Private/Framework/MultiBox/SClippingHorizontalBox.cpp
@@ -49,23 +49,32 @@ void SClippingHorizontalBox::OnArrangeChildren( const FGeometry& AllottedGeometr
 		FArrangedWidget& ArrangedWrapButton = ArrangedChildren[ArrangedChildren.Num() - 1];
 		FGeometry& WrapButtonGeometry = ArrangedWrapButton.Geometry;
 
-		const FVector2D WrapButtonSize = FVector2D(WrapButtonWidth, WrapButtonGeometry.GetLocalSize().Y);
-		WrapButtonGeometry = AllottedGeometry.MakeChild(
-			WrapButtonSize,
-			FSlateLayoutTransform(AllottedGeometry.GetLocalSize() - WrapButtonSize));
+		if (const bool bHasSpaceForWrapButton = WrapButtonWidth <= AllottedGeometry.GetLocalSize().X)
+		{
+			const float AdjustedWrapButtonWidth = FMath::Min(AllottedGeometry.GetLocalSize().X, WrapButtonWidth);
 
-		const int32 WrapButtonXPosition = FMath::TruncToInt(WrapButtonGeometry.AbsolutePosition.X);
+			const FVector2D WrapButtonSize = FVector2D(AdjustedWrapButtonWidth, WrapButtonGeometry.GetLocalSize().Y);
+			WrapButtonGeometry = AllottedGeometry.MakeChild(
+				WrapButtonSize,
+				FSlateLayoutTransform(AllottedGeometry.GetLocalSize() - WrapButtonSize));
 
-		// Further remove any children that the wrap button overlaps with
-		for (int32 ChildIdx = IndexClippedAt - 1; ChildIdx >= 0; --ChildIdx)
-		{
-			const FArrangedWidget& CurWidget = ArrangedChildren[ChildIdx];
-			if (FMath::TruncToInt(CurWidget.Geometry.AbsolutePosition.X + CurWidget.Geometry.GetLocalSize().X * CurWidget.Geometry.Scale) > WrapButtonXPosition)
+			const int32 WrapButtonXPosition = FMath::TruncToInt(WrapButtonGeometry.AbsolutePosition.X);
+
+			// Further remove any children that the wrap button overlaps with
+			for (int32 ChildIdx = IndexClippedAt - 1; ChildIdx >= 0; --ChildIdx)
 			{
-				++NumClippedChildren;
-				ArrangedChildren.Remove(ChildIdx);
+				const FArrangedWidget& CurWidget = ArrangedChildren[ChildIdx];
+				if (FMath::TruncToInt(CurWidget.Geometry.AbsolutePosition.X + CurWidget.Geometry.GetLocalSize().X * CurWidget.Geometry.Scale) > WrapButtonXPosition)
+				{
+					++NumClippedChildren;
+					ArrangedChildren.Remove(ChildIdx);
+				}
 			}
 		}
+		else // No space left for anything including WrapButton
+		{
+			ArrangedChildren.Empty();
+		}
 	}
 
 	ClippedIdx = ArrangedChildren.Num() - 1;
diff --git a/Engine/Source/Runtime/SlateCore/Private/Styling/StarshipCoreStyle.cpp b/Engine/Source/Runtime/SlateCore/Private/Styling/StarshipCoreStyle.cpp
--- a/Engine/Source/Runtime/SlateCore/Private/Styling/StarshipCoreStyle.cpp
+++ b/Engine/Source/Runtime/SlateCore/Private/Styling/StarshipCoreStyle.cpp
@@ -1894,7 +1894,7 @@ void FStarshipCoreStyle::SetupDockingStyles(TSharedRef<FStyle>& Style)
 	StatusBarToolBarStyle.SetBackground(FSlateNoResource());
 	StatusBarToolBarStyle.SetLabelPadding(FMargin(5, 5, 0, 5));
 	StatusBarToolBarStyle.SetIconSize(Icon16x16);
-	StatusBarToolBarStyle.SetBackgroundPadding(FMargin(4.f, 4.f, 0.f, 4.f));
+	StatusBarToolBarStyle.SetBackgroundPadding(FMargin(4.f, 4.f));
 
 	Style->Set("StatusBarToolBar", StatusBarToolBarStyle);
 
@@ -2140,7 +2140,7 @@ void FStarshipCoreStyle::SetupMultiboxStyles(TSharedRef<FStyle>& Style)
 		FToolBarStyle NormalToolbarStyle =
 			FToolBarStyle()
 			.SetBackground(FSlateColorBrush(FStyleColors::Panel))
-			.SetBackgroundPadding(FMargin(0.f, 4.f, 0.f, 4.f))
+			.SetBackgroundPadding(FMargin(0.f, 4.f))
 			.SetExpandBrush(IMAGE_BRUSH("Icons/toolbar_expand_16x", Icon16x16))
 			.SetComboButtonPadding(FMargin(4.0f, 0.0f))
 			.SetButtonPadding(FMargin(2.0f, 0.f))
@@ -2202,7 +2202,7 @@ void FStarshipCoreStyle::SetupMultiboxStyles(TSharedRef<FStyle>& Style)
 		FToolBarStyle SlimToolbarStyle =
 			FToolBarStyle()
 			.SetBackground(*SlimToolbarBackground)
-			.SetBackgroundPadding(FMargin(4.f, 6.f, 0.f, 6.f))
+			.SetBackgroundPadding(FMargin(4.f, 6.f))
 			.SetExpandBrush(IMAGE_BRUSH("Icons/toolbar_expand_16x", Icon16x16))
 			.SetComboButtonPadding(FMargin(0.0f, 0.0f))
 			.SetButtonPadding(FMargin(4.0f, 0.0f))
