Slate: SScrollBox allows for inserting slots at a given index, and retrieve number of slots

diff --git a/Engine/Source/Runtime/Slate/Private/Widgets/Layout/SScrollBox.cpp b/Engine/Source/Runtime/Slate/Private/Widgets/Layout/SScrollBox.cpp
--- a/Engine/Source/Runtime/Slate/Private/Widgets/Layout/SScrollBox.cpp
+++ b/Engine/Source/Runtime/Slate/Private/Widgets/Layout/SScrollBox.cpp
@@ -310,7 +310,12 @@ void SScrollBox::ConstructHorizontalLayout()
 /** Adds a slot to SScrollBox */
 SScrollBox::FScopedWidgetSlotArguments SScrollBox::AddSlot()
 {
-	return FScopedWidgetSlotArguments{ MakeUnique<FSlot>(), ScrollPanel->Children, INDEX_NONE };
+	return InsertSlot(INDEX_NONE);
+}
+
+SScrollBox::FScopedWidgetSlotArguments SScrollBox::InsertSlot(int32 Index)
+{
+	return FScopedWidgetSlotArguments{ MakeUnique<FSlot>(), ScrollPanel->Children, Index };
 }
 
 /** Removes a slot at the specified location */
@@ -319,6 +324,11 @@ void SScrollBox::RemoveSlot( const TSharedRef<SWidget>& WidgetToRemove )
 	ScrollPanel->Children.Remove(WidgetToRemove);
 }
 
+int32 SScrollBox::NumSlots() const
+{
+	return ScrollPanel->Children.Num();
+}
+
 void SScrollBox::ClearChildren()
 {
 	ScrollPanel->Children.Empty();
diff --git a/Engine/Source/Runtime/Slate/Public/Widgets/Layout/SScrollBox.h b/Engine/Source/Runtime/Slate/Public/Widgets/Layout/SScrollBox.h
--- a/Engine/Source/Runtime/Slate/Public/Widgets/Layout/SScrollBox.h
+++ b/Engine/Source/Runtime/Slate/Public/Widgets/Layout/SScrollBox.h
@@ -355,9 +355,15 @@ public:
 	/** Adds a slot to SScrollBox */
 	SLATE_API FScopedWidgetSlotArguments AddSlot();
 
+	/** Insert a slot at a given position. */
+	SLATE_API FScopedWidgetSlotArguments InsertSlot(int32 Index);
+
 	/** Removes a slot at the specified location */
 	SLATE_API void RemoveSlot( const TSharedRef<SWidget>& WidgetToRemove );
 
+	/** @return the number of slots. */
+	SLATE_API int32 NumSlots() const;
+
 	/** Removes all children from the box */
 	SLATE_API void ClearChildren();
 
@@ -535,7 +541,7 @@ protected:
 	/** Scrolls or begins scrolling a widget into view, only valid to call when we have layout geometry. */
 	SLATE_API bool InternalScrollDescendantIntoView(const FGeometry& MyGeometry, const TSharedPtr<SWidget>& WidgetToFind, bool InAnimateScroll = true, EDescendantScrollDestination InDestination = EDescendantScrollDestination::IntoView, float Padding = 0);
 
-	/** returns widget that can receive keyboard focus or nullprt **/
+	/** returns widget that can receive keyboard focus or nullptr **/
 	SLATE_API TSharedPtr<SWidget> GetKeyboardFocusableWidget(TSharedPtr<SWidget> InWidget);
 
 	/** The panel which stacks the child slots */
