Updated FNotificationButtonInfo to support dynamic Text, Visibility, and IsEnabled bindings

diff --git a/Engine/Source/Runtime/Slate/Private/Widgets/Notifications/SNotificationList.cpp b/Engine/Source/Runtime/Slate/Private/Widgets/Notifications/SNotificationList.cpp
--- a/Engine/Source/Runtime/Slate/Private/Widgets/Notifications/SNotificationList.cpp
+++ b/Engine/Source/Runtime/Slate/Private/Widgets/Notifications/SNotificationList.cpp
@@ -508,10 +508,8 @@ public:
 			// Adds any buttons that were passed in.
 			TSharedRef<SHorizontalBox> ButtonsBox = SNew(SHorizontalBox);
 
-			for (int32 idx = 0; idx < InArgs._ButtonDetails.Num(); idx++)
+			for (const FNotificationButtonInfo& Button : InArgs._ButtonDetails)
 			{
-				FNotificationButtonInfo Button = InArgs._ButtonDetails[idx];
-
 				ButtonsBox->AddSlot()
 				.AutoWidth()
 				.VAlign(VAlign_Center)
@@ -519,10 +517,11 @@ public:
 				[
 					SNew(SButton)
 					.Text(Button.Text)
-					.ToolTipText(Button.ToolTip.IsEmpty() ? Button.Text : Button.ToolTip)
+					.ToolTipText(Button.ToolTip.IsSet() ? Button.ToolTip : Button.Text)
 					.TextStyle(&FAppStyle::GetWidgetStyle<FTextBlockStyle>("NotificationList.WidgetText"))
 					.OnClicked(this, &SNotificationItemImpl::OnButtonClicked, Button.Callback) 
-					.Visibility( this, &SNotificationItemImpl::GetButtonVisibility, Button.VisibilityOnNone, Button.VisibilityOnPending, Button.VisibilityOnSuccess, Button.VisibilityOnFail )
+					.Visibility( this, &SNotificationItemImpl::GetButtonVisibility, Button.VisibilityCallback, Button.VisibilityOnNone, Button.VisibilityOnPending, Button.VisibilityOnSuccess, Button.VisibilityOnFail )
+					.IsEnabled( this, &SNotificationItemImpl::IsButtonEnabled, Button.IsEnabledCallback )
 				];
 			}
 
@@ -589,9 +588,19 @@ public:
 
 protected:
 
+	/* Used to determine whether the button is enabled */
+	bool IsButtonEnabled( const FNotificationButtonInfo::FIsEnabledDelegate IsEnabledCallback ) const
+	{
+		return !IsEnabledCallback.IsBound() || IsEnabledCallback.Execute(CompletionState);
+	}
+
 	/* Used to determine whether the button is visible */
-	EVisibility GetButtonVisibility( const EVisibility VisibilityOnNone, const EVisibility VisibilityOnPending, const EVisibility VisibilityOnSuccess, const EVisibility VisibilityOnFail ) const
+	EVisibility GetButtonVisibility( const FNotificationButtonInfo::FVisibilityDelegate VisibilityCallback, const EVisibility VisibilityOnNone, const EVisibility VisibilityOnPending, const EVisibility VisibilityOnSuccess, const EVisibility VisibilityOnFail ) const
 	{
+		if (VisibilityCallback.IsBound())
+		{
+			return VisibilityCallback.Execute(CompletionState);
+		}
 		switch ( CompletionState )
 		{
 		case SNotificationItem::CS_None: return VisibilityOnNone;
diff --git a/Engine/Source/Runtime/Slate/Public/Widgets/Notifications/SNotificationList.h b/Engine/Source/Runtime/Slate/Public/Widgets/Notifications/SNotificationList.h
--- a/Engine/Source/Runtime/Slate/Public/Widgets/Notifications/SNotificationList.h
+++ b/Engine/Source/Runtime/Slate/Public/Widgets/Notifications/SNotificationList.h
@@ -72,30 +72,51 @@ public:
  */
 struct FNotificationButtonInfo
 {
-	FNotificationButtonInfo( const FText& InText, const FText& InToolTip, FSimpleDelegate InCallback, SNotificationItem::ECompletionState VisibleInState = SNotificationItem::CS_Pending )
+	using FVisibilityDelegate = TDelegate<EVisibility(SNotificationItem::ECompletionState)>;
+	using FIsEnabledDelegate = TDelegate<bool(SNotificationItem::ECompletionState)>;
+
+	FNotificationButtonInfo( const TAttribute<FText>& InText, const TAttribute<FText>& InToolTip, FSimpleDelegate InCallback, SNotificationItem::ECompletionState VisibleInState = SNotificationItem::CS_Pending )
 		: Text( InText )
 		, ToolTip( InToolTip )
-		, Callback( InCallback )
+		, Callback( MoveTemp(InCallback) )
+		, VisibilityCallback()
+		, IsEnabledCallback()
 		, VisibilityOnNone( VisibleInState == SNotificationItem::CS_None ? EVisibility::Visible : EVisibility::Collapsed )
 		, VisibilityOnPending( VisibleInState == SNotificationItem::CS_Pending ? EVisibility::Visible : EVisibility::Collapsed )
 		, VisibilityOnSuccess( VisibleInState == SNotificationItem::CS_Success ? EVisibility::Visible : EVisibility::Collapsed )
 		, VisibilityOnFail( VisibleInState == SNotificationItem::CS_Fail ? EVisibility::Visible : EVisibility::Collapsed )
 	{ }
 
+	FNotificationButtonInfo( const TAttribute<FText>& InText, const TAttribute<FText>& InToolTip, FSimpleDelegate InCallback, FVisibilityDelegate InVisibilityCallback, FIsEnabledDelegate InIsEnabledCallback = FIsEnabledDelegate() )
+		: Text( InText )
+		, ToolTip( InToolTip )
+		, Callback( MoveTemp(InCallback) )
+		, VisibilityCallback( MoveTemp(InVisibilityCallback) )
+		, IsEnabledCallback( MoveTemp(InIsEnabledCallback) )
+		, VisibilityOnNone( VisibilityCallback.IsBound() ? EVisibility::Collapsed : EVisibility::Visible )
+		, VisibilityOnPending( VisibilityCallback.IsBound() ? EVisibility::Collapsed : EVisibility::Visible )
+		, VisibilityOnSuccess( VisibilityCallback.IsBound() ? EVisibility::Collapsed : EVisibility::Visible )
+		, VisibilityOnFail( VisibilityCallback.IsBound() ? EVisibility::Collapsed : EVisibility::Visible )
+	{ }
+
 	/** Message on the button */
-	FText Text;
+	TAttribute<FText> Text;
 	/** Tip displayed when moused over */
-	FText ToolTip;
+	TAttribute<FText> ToolTip;
 	/** Method called when button clicked */
 	FSimpleDelegate Callback;
+	/** Method called to decide if the button if visible; overrides the EVisibility data below if bound */
+	FVisibilityDelegate VisibilityCallback;
+	/** Method called to decide if the button is enabled; assumes enabled if unbound */
+	FIsEnabledDelegate IsEnabledCallback;
 
-	/** Visibility of the button when the completion state of the button is SNotificationItem::ECompletionState::None */
+	/** Visibility of the button when the completion state of the button is SNotificationItem::ECompletionState::None and VisibilityCallback is unbound */
 	EVisibility VisibilityOnNone;
-	/** Visibility of the button when the completion state of the button is SNotificationItem::ECompletionState::Pending */
+	/** Visibility of the button when the completion state of the button is SNotificationItem::ECompletionState::Pending and VisibilityCallback is unbound */
 	EVisibility VisibilityOnPending;
-	/** Visibility of the button when the completion state of the button is SNotificationItem::ECompletionState::Success */
+	/** Visibility of the button when the completion state of the button is SNotificationItem::ECompletionState::Success and VisibilityCallback is unbound */
 	EVisibility VisibilityOnSuccess;
-	/** Visibility of the button when the completion state of the button is SNotificationItem::ECompletionState::Fail */
+	/** Visibility of the button when the completion state of the button is SNotificationItem::ECompletionState::Fail and VisibilityCallback is unbound */
 	EVisibility VisibilityOnFail;
 };
 
