Updated FAsyncTaskNotification to use a single Slate notification

This restores its original behavior from UE4, where it would share a single notification for all state transitions rather than make a new notification for each state change

#rnx

diff --git a/Engine/Source/Runtime/Core/Public/Misc/CoreAsyncTaskNotificationImpl.h b/Engine/Source/Runtime/Core/Public/Misc/CoreAsyncTaskNotificationImpl.h
--- a/Engine/Source/Runtime/Core/Public/Misc/CoreAsyncTaskNotificationImpl.h
+++ b/Engine/Source/Runtime/Core/Public/Misc/CoreAsyncTaskNotificationImpl.h
@@ -121,7 +121,7 @@ protected:
 	/** The progress text displayed in the notification (if any) */
 	FText ProgressText;
 
-	/** The progress text displayed in the notification (if any) */
+	/** The text used by the prompt button (if any) */
 	FText PromptText;
 
 	/** When set this will display as a hyperlink on the right side of the notification. */
diff --git a/Engine/Source/Runtime/Slate/Private/Framework/Notifications/SlateAsyncTaskNotificationImpl.cpp b/Engine/Source/Runtime/Slate/Private/Framework/Notifications/SlateAsyncTaskNotificationImpl.cpp
--- a/Engine/Source/Runtime/Slate/Private/Framework/Notifications/SlateAsyncTaskNotificationImpl.cpp
+++ b/Engine/Source/Runtime/Slate/Private/Framework/Notifications/SlateAsyncTaskNotificationImpl.cpp
@@ -24,14 +24,17 @@
  * FSlateAsyncTaskNotificationImpl
  */
 
-FSlateAsyncTaskNotificationImpl::FSlateAsyncTaskNotificationImpl() : PromptAction(EAsyncTaskNotificationPromptAction::None)
-{
-	
+FSlateAsyncTaskNotificationImpl::FSlateAsyncTaskNotificationImpl()
+	: PromptAction(EAsyncTaskNotificationPromptAction::None)
+{	
 }
 
 FSlateAsyncTaskNotificationImpl::~FSlateAsyncTaskNotificationImpl()
 {
-	FTSTicker::GetCoreTicker().RemoveTicker(TickerHandle);
+	if (TickerHandle.IsValid())
+	{
+		FTSTicker::GetCoreTicker().RemoveTicker(TickerHandle);
+	}
 }
 
 void FSlateAsyncTaskNotificationImpl::Initialize(const FAsyncTaskNotificationConfig& InConfig)
@@ -43,31 +46,39 @@ void FSlateAsyncTaskNotificationImpl::Initialize(const FAsyncTaskNotificationCon
 	// Initialize the UI if the Notification is not headless
 	if (!NotificationConfig.bIsHeadless)
 	{
-		// Register the ticker to update the notification ever frame
-		TickerHandle = FTSTicker::GetCoreTicker().AddTicker(FTickerDelegate::CreateSP(this, &FSlateAsyncTaskNotificationImpl::TickNotification));
-
-		// Register this as a Staged Notification (Allows notifications to remain open even after this is destroyed)
+		// Register this as a Staged Notification (to keep 'this' alive until UnregisterStagedNotification is called)
 		FSlateNotificationManager::Get().RegisterStagedNotification(AsShared());
 
 		PromptAction = FApp::IsUnattended() ? EAsyncTaskNotificationPromptAction::Unattended : EAsyncTaskNotificationPromptAction::None;
 		bCanCancelAttr = InConfig.bCanCancel;
 		bKeepOpenOnSuccessAttr = InConfig.bKeepOpenOnSuccess;
 		bKeepOpenOnFailureAttr = InConfig.bKeepOpenOnFailure;
-		SyncAttributes();
 
-		// Mark the notification as pending so the UI can initialize
-		PreviousCompletionState = EAsyncTaskNotificationState::None;
-		SetPendingCompletionState(EAsyncTaskNotificationState::Pending); 
-	}
+		// Set the initial pending state prior to calling Tick to initialize the UI to that state
+		CurrentNotificationState = EAsyncTaskNotificationState::None;
+		SetPendingNotificationState(State);
 
+		// Create the notification UI
+		CreateNotification();
+
+		// Run a Tick to initialize the UI to the initial state
+		const bool bContinueTicking = TickNotification(0.0f);
+		if (bContinueTicking)
+		{
+			// Register the ticker to update the notification every frame
+			TickerHandle = FTSTicker::GetCoreTicker().AddTicker(FTickerDelegate::CreateSP(this, &FSlateAsyncTaskNotificationImpl::TickNotification));
+		}
+	}
 	
-	// This calls UpdateNotification to update the UI initialized above
+	// This calls UpdateNotification to update the UI initialized above, 
+	// which will happen immediately since bInitializedNotification is false
 	FCoreAsyncTaskNotificationImpl::Initialize(InConfig);
+	bInitializedNotification = true;
 }
 
-void FSlateAsyncTaskNotificationImpl::DestroyCurrentNotification()
+void FSlateAsyncTaskNotificationImpl::DestroyNotification()
 {
-	if(OwningNotification)
+	if (OwningNotification)
 	{
 		// Perform the normal automatic fadeout
 		OwningNotification->ExpireAndFadeout();
@@ -77,139 +88,58 @@ void FSlateAsyncTaskNotificationImpl::DestroyCurrentNotification()
 	}
 }
 
-void FSlateAsyncTaskNotificationImpl::CreateNewNotificationItem(EAsyncTaskNotificationState NewNotificationState)
+void FSlateAsyncTaskNotificationImpl::CreateNotification()
 {
-	DestroyCurrentNotification();
-	
-	switch (NewNotificationState)
+	check(!NotificationConfig.bIsHeadless);
+
+	if (OwningNotification)
 	{
-	case EAsyncTaskNotificationState::Pending:
-		CreatePendingNotification();
-		break;
-	case EAsyncTaskNotificationState::Failure:
-		CreateFailureNotification();
-		break;
-	case EAsyncTaskNotificationState::Success:
-		CreateSuccessNotification();
-		break;
-	case EAsyncTaskNotificationState::Prompt:
-		CreatePromptNotification();
-		break;
+		return;
 	}
 
-}
-
-void FSlateAsyncTaskNotificationImpl::SetOwner(TSharedPtr<SNotificationItem> InOwningNotification)
-{
-	OwningNotification = InOwningNotification;
-
-	// Update the notification here to make sure it has the correct Text/Subtext/Hyperlink etc
-	UpdateNotification();
-}
-
-TSharedPtr<SNotificationItem> FSlateAsyncTaskNotificationImpl::SetupNotificationItem(FNotificationInfo& NotificationInfo)
-{
+	FNotificationInfo NotificationInfo(FText::GetEmpty());
 	NotificationInfo.FadeOutDuration = NotificationConfig.FadeOutDuration;
 	NotificationInfo.ExpireDuration = NotificationConfig.ExpireDuration;
 	NotificationInfo.FadeInDuration = NotificationConfig.FadeInDuration;
 	NotificationInfo.bFireAndForget = false;
-
-	TSharedPtr<SNotificationItem> NotificationItem = FSlateNotificationManager::Get().AddNotification(NotificationInfo);
-	check(NotificationItem);
-	
-	SetOwner(NotificationItem);
-	return NotificationItem;
-}
-
-void FSlateAsyncTaskNotificationImpl::CreatePendingNotification()
-{
-	FNotificationInfo NotificationInfo(FText::GetEmpty());
-
-	// Pending Notifications have a throbber to show progress
 	NotificationInfo.bUseThrobber = true;
-	
-	AddCancelButton(NotificationInfo, SNotificationItem::CS_Pending);
-	
-	TSharedPtr<SNotificationItem> NotificationItem = SetupNotificationItem(NotificationInfo);
-}
-
-void FSlateAsyncTaskNotificationImpl::CreateSuccessNotification()
-{
-	FNotificationInfo NotificationInfo(FText::GetEmpty());
-
-	NotificationInfo.Image = FAppStyle::Get().GetBrush("NotificationList.SuccessImage");
-	AddCloseButton(NotificationInfo);
-	
-	TSharedPtr<SNotificationItem> NotificationItem = SetupNotificationItem(NotificationInfo);
-}
-
-void FSlateAsyncTaskNotificationImpl::CreateFailureNotification()
-{
-	FNotificationInfo NotificationInfo(FText::GetEmpty());
-	
-	NotificationInfo.Image = FAppStyle::Get().GetBrush("NotificationList.FailImage");
-	AddCloseButton(NotificationInfo);
-
-	TSharedPtr<SNotificationItem> NotificationItem = SetupNotificationItem(NotificationInfo);
-}
+	NotificationInfo.bUseSuccessFailIcons = true;
+	NotificationInfo.Image = NotificationConfig.Icon;
 
-void FSlateAsyncTaskNotificationImpl::CreatePromptNotification()
-{
-	FNotificationInfo NotificationInfo(FText::GetEmpty());
-
-	AddPromptButton(NotificationInfo);
-	
-	AddCancelButton(NotificationInfo, SNotificationItem::CS_None);
-	
-	TSharedPtr<SNotificationItem> NotificationItem = SetupNotificationItem(NotificationInfo);
-}
-
-void FSlateAsyncTaskNotificationImpl::AddPromptButton(FNotificationInfo &NotificationInfo)
-{
-	if(GetPromptButtonVisibility() == EVisibility::Visible)
 	{
 		FNotificationButtonInfo PromptButtonInfo(
-		PromptText,
-		FText::GetEmpty(),
-		FSimpleDelegate::CreateSP(this, &FSlateAsyncTaskNotificationImpl::OnPromptButtonClicked),
-		SNotificationItem::CS_None
+			TAttribute<FText>::CreateSP(this, &FSlateAsyncTaskNotificationImpl::GetPromptButtonText),
+			FText::GetEmpty(),
+			FSimpleDelegate::CreateSP(this, &FSlateAsyncTaskNotificationImpl::OnPromptButtonClicked),
+			FNotificationButtonInfo::FVisibilityDelegate::CreateSP(this, &FSlateAsyncTaskNotificationImpl::GetPromptButtonVisibility),
+			FNotificationButtonInfo::FIsEnabledDelegate::CreateSP(this, &FSlateAsyncTaskNotificationImpl::IsPromptButtonEnabled)
 		);
-	
 		NotificationInfo.ButtonDetails.Add(PromptButtonInfo);
 	}
-}
-
 
-void FSlateAsyncTaskNotificationImpl::AddCancelButton(FNotificationInfo &NotificationInfo, SNotificationItem::ECompletionState VisibleInState)
-{
-	if(GetCancelButtonVisibility() == EVisibility::Visible)
 	{
 		FNotificationButtonInfo CancelButtonInfo(
-		LOCTEXT("CancelButton", "Cancel"),
-		FText::GetEmpty(),
-		FSimpleDelegate::CreateSP(this, &FSlateAsyncTaskNotificationImpl::OnCancelButtonClicked),
-		VisibleInState
+			LOCTEXT("CancelButton", "Cancel"),
+			FText::GetEmpty(),
+			FSimpleDelegate::CreateSP(this, &FSlateAsyncTaskNotificationImpl::OnCancelButtonClicked),
+			FNotificationButtonInfo::FVisibilityDelegate::CreateSP(this, &FSlateAsyncTaskNotificationImpl::GetCancelButtonVisibility),
+			FNotificationButtonInfo::FIsEnabledDelegate::CreateSP(this, &FSlateAsyncTaskNotificationImpl::IsCancelButtonEnabled)
 		);
-	
 		NotificationInfo.ButtonDetails.Add(CancelButtonInfo);
 	}
-}
 
-void FSlateAsyncTaskNotificationImpl::AddCloseButton(FNotificationInfo &NotificationInfo)
-{
-	if(GetCloseButtonVisibility() == EVisibility::Visible)
 	{
 		FNotificationButtonInfo CloseButtonInfo(
-		LOCTEXT("CloseButton", "Close"),
-		FText::GetEmpty(),
-		FSimpleDelegate::CreateSP(this, &FSlateAsyncTaskNotificationImpl::OnCloseButtonClicked)
+			LOCTEXT("CloseButton", "Close"),
+			FText::GetEmpty(),
+			FSimpleDelegate::CreateSP(this, &FSlateAsyncTaskNotificationImpl::OnCloseButtonClicked),
+			FNotificationButtonInfo::FVisibilityDelegate::CreateSP(this, &FSlateAsyncTaskNotificationImpl::GetCloseButtonVisibility)
 		);
-
-		CloseButtonInfo.VisibilityOnSuccess = EVisibility::Visible;
-		CloseButtonInfo.VisibilityOnFail = EVisibility::Visible;
-		
 		NotificationInfo.ButtonDetails.Add(CloseButtonInfo);
 	}
+
+	OwningNotification = FSlateNotificationManager::Get().AddNotification(NotificationInfo);
+	check(OwningNotification);
 }
 
 void FSlateAsyncTaskNotificationImpl::SyncAttributes()
@@ -221,28 +151,12 @@ void FSlateAsyncTaskNotificationImpl::SyncAttributes()
 	bKeepOpenOnFailure = bKeepOpenOnFailureAttr.Get(false);
 }
 
-void FSlateAsyncTaskNotificationImpl::OnSetCompletionState(SNotificationItem::ECompletionState InState)
-{
-	check(InState == GetNotificationCompletionState());
-
-	// If we completed and we aren't keeping the notification open (which will show the Close button), then expire the notification immediately
-	if ((InState == SNotificationItem::CS_Success || InState == SNotificationItem::CS_Fail) && GetCloseButtonVisibility() == EVisibility::Collapsed)
-	{
-		DestroyCurrentNotification();
-
-		FSlateNotificationManager::Get().UnregisterStagedNotification(AsShared());
-	}
-
-	// Reset the `PromptAction` state when changing completion state
-	PromptAction = FApp::IsUnattended() ? EAsyncTaskNotificationPromptAction::Unattended : EAsyncTaskNotificationPromptAction::None;
-}
-
-void FSlateAsyncTaskNotificationImpl::SetPendingCompletionState(const EAsyncTaskNotificationState InPendingCompletionState)
+void FSlateAsyncTaskNotificationImpl::SetPendingNotificationState(const EAsyncTaskNotificationState InPendingNotificationState)
 {
 	FScopeLock Lock(&CompletionCS);
 
 	// Set the completion state
-	PendingCompletionState = InPendingCompletionState;
+	PendingNotificationState = InPendingNotificationState;
 }
 
 void FSlateAsyncTaskNotificationImpl::SetCanCancel(const TAttribute<bool>& InCanCancel)
@@ -275,14 +189,14 @@ void FSlateAsyncTaskNotificationImpl::SetKeepOpenOnFailure(const TAttribute<bool
 	}
 }
 
-bool FSlateAsyncTaskNotificationImpl::IsCancelButtonEnabled() const
+bool FSlateAsyncTaskNotificationImpl::IsCancelButtonEnabled(SNotificationItem::ECompletionState InState) const
 {
 	return bCanCancel && PromptAction == EAsyncTaskNotificationPromptAction::None;
 }
 
-EVisibility FSlateAsyncTaskNotificationImpl::GetCancelButtonVisibility() const
+EVisibility FSlateAsyncTaskNotificationImpl::GetCancelButtonVisibility(SNotificationItem::ECompletionState InState) const
 {
-	return (bCanCancel && (State == EAsyncTaskNotificationState::Pending || State == EAsyncTaskNotificationState::Prompt))
+	return (bCanCancel && (CurrentNotificationState == EAsyncTaskNotificationState::Pending || CurrentNotificationState == EAsyncTaskNotificationState::Prompt))
 		? EVisibility::Visible
 		: EVisibility::Collapsed;
 }
@@ -292,14 +206,14 @@ void FSlateAsyncTaskNotificationImpl::OnCancelButtonClicked()
 	PromptAction = EAsyncTaskNotificationPromptAction::Cancel;
 }
 
-bool FSlateAsyncTaskNotificationImpl::IsPromptButtonEnabled() const
+bool FSlateAsyncTaskNotificationImpl::IsPromptButtonEnabled(SNotificationItem::ECompletionState InState) const
 {
 	return PromptAction == EAsyncTaskNotificationPromptAction::None;
 }
 
-EVisibility FSlateAsyncTaskNotificationImpl::GetPromptButtonVisibility() const
+EVisibility FSlateAsyncTaskNotificationImpl::GetPromptButtonVisibility(SNotificationItem::ECompletionState InState) const
 {
-	return (!FApp::IsUnattended() && State == EAsyncTaskNotificationState::Prompt)
+	return (!FApp::IsUnattended() && CurrentNotificationState == EAsyncTaskNotificationState::Prompt)
 		? EVisibility::Visible
 		: EVisibility::Collapsed;
 }
@@ -311,12 +225,12 @@ void FSlateAsyncTaskNotificationImpl::OnPromptButtonClicked()
 
 FText FSlateAsyncTaskNotificationImpl::GetPromptButtonText() const
 {
-	return PromptText;
+	return PromptButtonText;
 }
 
-EVisibility FSlateAsyncTaskNotificationImpl::GetCloseButtonVisibility() const
+EVisibility FSlateAsyncTaskNotificationImpl::GetCloseButtonVisibility(SNotificationItem::ECompletionState InState) const
 {
-	return (!FApp::IsUnattended() && ((bKeepOpenOnSuccess && State == EAsyncTaskNotificationState::Success) || (bKeepOpenOnFailure && State == EAsyncTaskNotificationState::Failure)))
+	return (!FApp::IsUnattended() && ((bKeepOpenOnSuccess && CurrentNotificationState == EAsyncTaskNotificationState::Success) || (bKeepOpenOnFailure && CurrentNotificationState == EAsyncTaskNotificationState::Failure)))
 		? EVisibility::Visible
 		: EVisibility::Collapsed;
 }
@@ -333,59 +247,57 @@ void FSlateAsyncTaskNotificationImpl::OnCloseButtonClicked()
 		// Release our reference to our owner so that everything can be destroyed
 		OwningNotification.Reset();
 
+		// Unregister our ticker now that we're closing
+		if (TickerHandle.IsValid())
+		{
+			FTSTicker::GetCoreTicker().RemoveTicker(TickerHandle);
+			TickerHandle.Reset();
+		}
+
 		// Unregister the Staged Notification to complete the cleanup
 		FSlateNotificationManager::Get().UnregisterStagedNotification(AsShared());
 	}
 }
 
-void FSlateAsyncTaskNotificationImpl::OnHyperlinkClicked() const
-{
-	Hyperlink.ExecuteIfBound();
-}
-
-FText FSlateAsyncTaskNotificationImpl::GetHyperlinkText() const
-{
-	return HyperlinkText;
-}
-
-EVisibility FSlateAsyncTaskNotificationImpl::GetHyperlinkVisibility() const
-{
-	return Hyperlink.IsBound() ? EVisibility::Visible : EVisibility::Collapsed;
-}
-
-SNotificationItem::ECompletionState FSlateAsyncTaskNotificationImpl::GetNotificationCompletionState() const
-{
-	if (OwningNotification)
-	{
-		return OwningNotification->GetCompletionState();
-	}
-	return SNotificationItem::CS_None;
-}
-
 void FSlateAsyncTaskNotificationImpl::UpdateNotification()
 {
 	FCoreAsyncTaskNotificationImpl::UpdateNotification();
 	
 	if (!NotificationConfig.bIsHeadless)
 	{
-		// Update the notification UI only if the state hasn't changed (i.e this notification will not be deleted)
-		if(OwningNotification && State == PreviousCompletionState)
+		// Update the notification UI
+		if (OwningNotification)
 		{
-			/* Slate requries the notification to be updated from the main thread, so we add a one frame ticker for it */
-			FTSTicker::GetCoreTicker().AddTicker(FTickerDelegate::CreateStatic(&FSlateAsyncTaskNotificationImpl::UpdateNotificationDeferred, OwningNotification, TitleText, ProgressText, Hyperlink, HyperlinkText));
+			if (bInitializedNotification)
+			{
+				// Slate requires the notification to be updated from the game thread, so we add a one frame ticker for it using the values captured from whichever thread is calling UpdateNotification
+				// Note: We also capture OwningNotification as transitioning to a success/fail state can reset this->OwningNotification before UpdateNotificationDeferred runs, which would cause the deferred update to fail if using this->OwningNotification
+				FTSTicker::GetCoreTicker().AddTicker(FTickerDelegate::CreateStatic(&FSlateAsyncTaskNotificationImpl::UpdateNotificationDeferred, AsWeak(), OwningNotification, TitleText, ProgressText, PromptText, Hyperlink, HyperlinkText));
+			}
+			else
+			{
+				// This is the UpdateNotification call made during Initialize
+				// We're on the game thread here so can push the initial state directly into the notification
+				FSlateAsyncTaskNotificationImpl::UpdateNotificationDeferred(0.0f, AsWeak(), OwningNotification, TitleText, ProgressText, PromptText, Hyperlink, HyperlinkText);
+			}
 		}
 
-		// Set the Pending Completion State in case the notification has to change
-		SetPendingCompletionState(State);
+		// Set the pending state in case the notification has to change
+		SetPendingNotificationState(State);
 	}
 }
 
-bool FSlateAsyncTaskNotificationImpl::UpdateNotificationDeferred(float InDeltaTime, TSharedPtr<SNotificationItem> OwningNotification, FText TitleText, FText ProgressText, FSimpleDelegate Hyperlink, FText HyperlinkText)
+bool FSlateAsyncTaskNotificationImpl::UpdateNotificationDeferred(float InDeltaTime, TWeakPtr<FSlateAsyncTaskNotificationImpl> WeakThis, TSharedPtr<SNotificationItem> OwningNotification, FText TitleText, FText ProgressText, FText PromptText, FSimpleDelegate Hyperlink, FText HyperlinkText)
 {
 	OwningNotification->SetText(TitleText);
 	OwningNotification->SetSubText(ProgressText);
 	OwningNotification->SetHyperlink(Hyperlink, HyperlinkText);
 
+	if (TSharedPtr<FSlateAsyncTaskNotificationImpl> This = WeakThis.Pin())
+	{
+		This->PromptButtonText = PromptText;
+	}
+
 	// We only want this function to tick once
 	return false;
 }
@@ -403,61 +315,54 @@ bool FSlateAsyncTaskNotificationImpl::TickNotification(float InDeltaTime)
 {
 	SyncAttributes();
 	
-	EAsyncTaskNotificationState CompletionStateToApply = EAsyncTaskNotificationState::None;
+	TOptional<EAsyncTaskNotificationState> NextNotificationState;
 	{
 		FScopeLock Lock(&CompletionCS);
 
-		if (PendingCompletionState.IsSet())
-		{
-			CompletionStateToApply = PendingCompletionState.GetValue();
-			PendingCompletionState.Reset();
-		}
+		NextNotificationState = MoveTemp(PendingNotificationState);
+		PendingNotificationState.Reset();
 	}
 
-	// Create a new notification if the state changed to a valid state
-	if (PreviousCompletionState != CompletionStateToApply && CompletionStateToApply != EAsyncTaskNotificationState::None)
+	// Update the notification UI state if the task state changed
+	if (NextNotificationState.IsSet() && CurrentNotificationState != NextNotificationState.GetValue())
 	{
-		// Reset the State of the previous notification if it was 'Pending', to make any misleading buttons disappear
-		if(OwningNotification && PreviousCompletionState == EAsyncTaskNotificationState::Pending)
-		{
-			OwningNotification->SetCompletionState(SNotificationItem::CS_None);
-		}
-		
-		PreviousCompletionState = CompletionStateToApply;
-
-		// Create a new notification based on the new state
-		CreateNewNotificationItem(CompletionStateToApply);
+		CurrentNotificationState = NextNotificationState.GetValue();
 
-		if(OwningNotification)
+		if (OwningNotification)
 		{
-			SNotificationItem::ECompletionState OwningCompletionState = SNotificationItem::CS_None;
-			switch (CompletionStateToApply)
+			SNotificationItem::ECompletionState OwningNotificationState = SNotificationItem::CS_None;
+			switch (CurrentNotificationState)
 			{
 			case EAsyncTaskNotificationState::Pending:
-				OwningCompletionState = SNotificationItem::CS_Pending;
+				OwningNotificationState = SNotificationItem::CS_Pending;
 				break;
 			case EAsyncTaskNotificationState::Failure:
-				OwningCompletionState = SNotificationItem::CS_Fail;
+				OwningNotificationState = SNotificationItem::CS_Fail;
 				break;
 			case EAsyncTaskNotificationState::Success:
-				OwningCompletionState = SNotificationItem::CS_Success;
+				OwningNotificationState = SNotificationItem::CS_Success;
 				break;
 			case EAsyncTaskNotificationState::Prompt:
 				OwningNotification->Pulse(FLinearColor(0.f, 0.f, 1.f));
-
 				break;
 			}
-			if (OwningCompletionState != SNotificationItem::CS_None && OwningCompletionState != OwningNotification->GetCompletionState())
-			{
-				OwningNotification->SetCompletionState(OwningCompletionState);
-				OnSetCompletionState(OwningCompletionState);
-
-				// We don't need the ticker anymore if the notification is complete
-				if(OwningCompletionState == SNotificationItem::CS_Success || OwningCompletionState == SNotificationItem::CS_Fail)
-				{
-					return false;
-				}
-			}
+			OwningNotification->SetCompletionState(OwningNotificationState);
+		}
+		
+		// Reset the `PromptAction` state when changing notification state
+		PromptAction = FApp::IsUnattended() ? EAsyncTaskNotificationPromptAction::Unattended : EAsyncTaskNotificationPromptAction::None;
+	}
+
+	// If we completed and we aren't keeping the notification open (which will show the Close button), then expire the notification immediately
+	{
+		const SNotificationItem::ECompletionState OwningNotificationState = OwningNotification ? OwningNotification->GetCompletionState() : SNotificationItem::CS_None;
+		if ((CurrentNotificationState == EAsyncTaskNotificationState::Success || CurrentNotificationState == EAsyncTaskNotificationState::Failure) && GetCloseButtonVisibility(OwningNotificationState) == EVisibility::Collapsed)
+		{
+			DestroyNotification();
+			TickerHandle.Reset(); // Reset this before potentially destroying 'this' when calling UnregisterStagedNotification
+
+			FSlateNotificationManager::Get().UnregisterStagedNotification(AsShared());
+			return false; // No longer need to Tick
 		}
 	}
 
diff --git a/Engine/Source/Runtime/Slate/Private/Framework/Notifications/SlateAsyncTaskNotificationImpl.h b/Engine/Source/Runtime/Slate/Private/Framework/Notifications/SlateAsyncTaskNotificationImpl.h
--- a/Engine/Source/Runtime/Slate/Private/Framework/Notifications/SlateAsyncTaskNotificationImpl.h
+++ b/Engine/Source/Runtime/Slate/Private/Framework/Notifications/SlateAsyncTaskNotificationImpl.h
@@ -32,72 +32,35 @@ private:
 	/** Run every frame on the game thread to update the notification */
 	bool TickNotification(float InDeltaTime);
 
-	/** Set the notification item that owns this widget */
-	void SetOwner(TSharedPtr<SNotificationItem> InOwningNotification);
-
 	/** Sync attribute bindings with the cached values (once per-frame from the game thread) */
 	void SyncAttributes();
 
-	/** Set the pending completion state of the notification (applied during the next Tick) */
-	void SetPendingCompletionState(const EAsyncTaskNotificationState InPendingCompletionState);
-
-	virtual void OnSetCompletionState(SNotificationItem::ECompletionState State);
+	/** Set the pending state of the notification UI (applied during the next Tick) */
+	void SetPendingNotificationState(const EAsyncTaskNotificationState InPendingCompletionState);
 
 	/** Cancel button */
-	bool IsCancelButtonEnabled() const;
-	EVisibility GetCancelButtonVisibility() const;
+	bool IsCancelButtonEnabled(SNotificationItem::ECompletionState InState) const;
+	EVisibility GetCancelButtonVisibility(SNotificationItem::ECompletionState InState) const;
 	void OnCancelButtonClicked();
 
 	/** Prompt button */
-	bool IsPromptButtonEnabled() const;
-	EVisibility GetPromptButtonVisibility() const;
+	bool IsPromptButtonEnabled(SNotificationItem::ECompletionState InState) const;
+	EVisibility GetPromptButtonVisibility(SNotificationItem::ECompletionState InState) const;
 	void OnPromptButtonClicked();
 	FText GetPromptButtonText() const;
 
 	/** Close button */
-	EVisibility GetCloseButtonVisibility() const;
+	EVisibility GetCloseButtonVisibility(SNotificationItem::ECompletionState InState) const;
 	void OnCloseButtonClicked();
 
-	/** Hyperlink */
-	void OnHyperlinkClicked() const;
-	FText GetHyperlinkText() const;
-	EVisibility GetHyperlinkVisibility() const;
-
-	/** Get the current completion state from the parent notification */
-	SNotificationItem::ECompletionState GetNotificationCompletionState() const;
-
-	/* Create a new notification item when the state changes */
-	void CreateNewNotificationItem(EAsyncTaskNotificationState NewNotificationState);
-
-	/* Create a notification for the Pending state */
-	void CreatePendingNotification();
-
-	/* Create a notification for the Success state */
-	void CreateSuccessNotification();
-
-	/* Create a notification for the Failure state */
-	void CreateFailureNotification();
+	/* Create the notification, if not already created */
+	void CreateNotification();
 
-	/* Create a notification for the Prompt state */
-	void CreatePromptNotification();
-
-	/* Helper function to add the "Cancel" button to a notification */
-	void AddCancelButton(FNotificationInfo &NotificationInfo, SNotificationItem::ECompletionState VisibleInState);
-
-	/* Helper function to add the "Close" button to a notification */
-	void AddCloseButton(FNotificationInfo &NotificationInfo);
-
-	/* Helper function to add a Prompt button to a notification */
-	void AddPromptButton(FNotificationInfo &NotificationInfo);
-
-	/* Helper function to setup the information common to notifications in all states */
-	TSharedPtr<SNotificationItem> SetupNotificationItem(FNotificationInfo& NotificationInfo);
-
-	/* Cleanly destroy the current notification */
-	void DestroyCurrentNotification();
+	/* Cleanly destroy the current notification, if any */
+	void DestroyNotification();
 
 	/* Static function to update the notification from the main thread */
-	static bool UpdateNotificationDeferred(float InDeltaTime, TSharedPtr<SNotificationItem> OwningNotification, FText TitleText, FText ProgressText, FSimpleDelegate Hyperlink, FText HyperlinkText);
+	static bool UpdateNotificationDeferred(float InDeltaTime, TWeakPtr<FSlateAsyncTaskNotificationImpl> WeakThis, TSharedPtr<SNotificationItem> OwningNotification, FText TitleText, FText ProgressText, FText PromptText, FSimpleDelegate Hyperlink, FText HyperlinkText);
 
 private:
 
@@ -110,6 +73,9 @@ private:
 	/** Action taken for the task, resets to none on notification state change. */
 	TAtomic<EAsyncTaskNotificationPromptAction> PromptAction;
 
+	/** The text used by the notification prompt button (if any; UI should query this rather than PromptText) */
+	FText PromptButtonText;
+
 	/** Can this task be canceled? Will show a cancel button for in-progress tasks */
 	TAttribute<bool> bCanCancelAttr = false;
 	bool bCanCancel = false;
@@ -122,10 +88,14 @@ private:
 	TAttribute<bool> bKeepOpenOnFailureAttr = false;
 	bool bKeepOpenOnFailure = false;
 
-	/** The pending completion state of the notification (if any, applied during the next Tick) */
-	TOptional<EAsyncTaskNotificationState> PendingCompletionState;
+	/** Have we finished initializing the notification? */
+	bool bInitializedNotification = false;
+
+	/** The current state of the notification (UI should query this rather than State) */
+	EAsyncTaskNotificationState CurrentNotificationState = EAsyncTaskNotificationState::None;
 
-	EAsyncTaskNotificationState PreviousCompletionState;
+	/** The pending state of the notification (if any, applied during the next Tick) */
+	TOptional<EAsyncTaskNotificationState> PendingNotificationState;
 
 	/** Pointer to the notification item that owns this widget (this is a deliberate reference cycle as we need this object alive until we choose to expire it, at which point we release our reference to allow everything to be destroyed) */
 	TSharedPtr<SNotificationItem> OwningNotification;
