Slate: Add SScrollBox::GetSlot(int32) to access slot at a given index

diff --git a/Engine/Source/Runtime/Slate/Private/Widgets/Layout/SScrollBox.cpp b/Engine/Source/Runtime/Slate/Private/Widgets/Layout/SScrollBox.cpp
--- a/Engine/Source/Runtime/Slate/Private/Widgets/Layout/SScrollBox.cpp
+++ b/Engine/Source/Runtime/Slate/Private/Widgets/Layout/SScrollBox.cpp
@@ -318,7 +318,18 @@ SScrollBox::FScopedWidgetSlotArguments SScrollBox::InsertSlot(int32 Index)
 	return FScopedWidgetSlotArguments{ MakeUnique<FSlot>(), ScrollPanel->Children, Index };
 }
 
-/** Removes a slot at the specified location */
+const SScrollBox::FSlot& SScrollBox::GetSlot(int32 SlotIndex) const
+{
+	check(ScrollPanel->Children.IsValidIndex(SlotIndex));
+	const FSlotBase& BaseSlot = static_cast<const FSlotBase&>(ScrollPanel->Children[SlotIndex]);
+	return static_cast<const FSlot&>(BaseSlot);
+}
+
+SScrollBox::FSlot& SScrollBox::GetSlot(int32 SlotIndex)
+{
+	return const_cast<FSlot&>(const_cast<const SScrollBox*>(this)->GetSlot(SlotIndex));
+}
+
 void SScrollBox::RemoveSlot( const TSharedRef<SWidget>& WidgetToRemove )
 {
 	ScrollPanel->Children.Remove(WidgetToRemove);
diff --git a/Engine/Source/Runtime/Slate/Public/Widgets/Layout/SScrollBox.h b/Engine/Source/Runtime/Slate/Public/Widgets/Layout/SScrollBox.h
--- a/Engine/Source/Runtime/Slate/Public/Widgets/Layout/SScrollBox.h
+++ b/Engine/Source/Runtime/Slate/Public/Widgets/Layout/SScrollBox.h
@@ -358,8 +358,12 @@ public:
 	/** Insert a slot at a given position. */
 	SLATE_API FScopedWidgetSlotArguments InsertSlot(int32 Index);
 
-	/** Removes a slot at the specified location */
-	SLATE_API void RemoveSlot( const TSharedRef<SWidget>& WidgetToRemove );
+	/** Returns the slot at the given index. */
+	SLATE_API const FSlot& GetSlot(int32 SlotIndex) const;
+	SLATE_API FSlot& GetSlot(int32 SlotIndex);
+
+	/** Removes the corresponding widget from the set of slots if it exists. */
+	SLATE_API void RemoveSlot(const TSharedRef<SWidget>& WidgetToRemove);
 
 	/** @return the number of slots. */
 	SLATE_API int32 NumSlots() const;
