Skip drawing Slate windows if running in NullRHI (UE-222271)

- Measured to shave off up to 2ms from headless game's frame (a significant part of it since it is generally <10ms).

diff --git a/Engine/Source/Runtime/Slate/Private/Framework/Application/SlateApplication.cpp b/Engine/Source/Runtime/Slate/Private/Framework/Application/SlateApplication.cpp
--- a/Engine/Source/Runtime/Slate/Private/Framework/Application/SlateApplication.cpp
+++ b/Engine/Source/Runtime/Slate/Private/Framework/Application/SlateApplication.cpp
@@ -123,6 +123,14 @@ static FAutoConsoleVariableRef CVarSlateUseFixedDeltaTime(
 	GSlateUseFixedDeltaTime,
 	TEXT("True means we use a constant delta time on every widget tick.")
 );
+
+static bool GSlateSkipWidgetDrawingInHeadlessMode = true;
+static FAutoConsoleVariableRef CVarSlateSkipWidgetDrawingInHeadlessMode(
+	TEXT("Slate.SkipWidgetDrawingInHeadlessMode"),
+	GSlateSkipWidgetDrawingInHeadlessMode,
+	TEXT("Skip drawing the widgets when running without rendering (e.g. -nullrhi). On by default, disable if there's non-visual logic in widget's Tick function (as this also skips ticking them).")
+);
+
 //////////////////////////////////////////////////////////////////////////
 
 /** 
@@ -1335,6 +1343,12 @@ TArray<SWindow*> GatherAllDescendants(const TArray< TSharedRef<SWindow> >& InWin
 
 void FSlateApplication::PrivateDrawWindows( TSharedPtr<SWindow> DrawOnlyThisWindow )
 {
+	if (GSlateSkipWidgetDrawingInHeadlessMode && !FApp::CanEverRender() && !bAnyActiveTimersPending)
+	{
+		// early out, as window "drawing" can take 1-2ms of a -nullrhi PC game
+		return;
+	}
+
 	check(Renderer.IsValid());
 
 	// Grab a scope lock around access to the resource proxy map, just to ensure we never cross over
@@ -1680,7 +1694,7 @@ void FSlateApplication::TickAndDrawWidgets(float DeltaTime)
 		const double TimeSinceMouseMove = LastTickTime - LastMouseMoveTime;
 	
 		const bool bIsUserIdle = (TimeSinceInput > SleepThreshold) && (TimeSinceMouseMove > SleepThreshold);
-		const bool bAnyActiveTimersPending = AnyActiveTimersArePending();
+		UpdateAnyActiveTimersArePending();
 
 		// skip tick/draw if we are idle and there are no active timers registered that we need to drive slate for.
 		// This effectively means the slate application is totally idle and we don't need to update the UI.
diff --git a/Engine/Source/Runtime/SlateCore/Private/Application/SlateApplicationBase.cpp b/Engine/Source/Runtime/SlateCore/Private/Application/SlateApplicationBase.cpp
--- a/Engine/Source/Runtime/SlateCore/Private/Application/SlateApplicationBase.cpp
+++ b/Engine/Source/Runtime/SlateCore/Private/Application/SlateApplicationBase.cpp
@@ -117,7 +117,7 @@ void FSlateApplicationBase::UnRegisterActiveTimer( const TSharedRef<FActiveTimer
 	ActiveTimerHandles.RemoveSingleSwap(ActiveTimerHandle);
 }
 
-bool FSlateApplicationBase::AnyActiveTimersArePending()
+void FSlateApplicationBase::UpdateAnyActiveTimersArePending()
 {
 	FScopeLock ActiveTimerLock(&ActiveTimerCS);
 
@@ -151,7 +151,7 @@ bool FSlateApplicationBase::AnyActiveTimersArePending()
 		}
 	}
 
-	return bAnyTickReady;
+	bAnyActiveTimersPending = bAnyTickReady;
 }
 
 bool FSlateApplicationBase::IsSlateAsleep()
diff --git a/Engine/Source/Runtime/SlateCore/Public/Application/SlateApplicationBase.h b/Engine/Source/Runtime/SlateCore/Public/Application/SlateApplicationBase.h
--- a/Engine/Source/Runtime/SlateCore/Public/Application/SlateApplicationBase.h
+++ b/Engine/Source/Runtime/SlateCore/Public/Application/SlateApplicationBase.h
@@ -512,7 +512,7 @@ protected:
 	 * Used to determine if any active timer handles are ready to fire.
 	 * Means we need to tick slate even if no user interaction.
 	 */
-	SLATECORE_API bool AnyActiveTimersArePending();
+	SLATECORE_API void UpdateAnyActiveTimersArePending();
 
 public:
 	SLATECORE_API const static uint32 CursorPointerIndex;
@@ -666,5 +666,7 @@ protected:
 
 	/** Safe Zone ratio to override platform settings */
 	FMargin CustomSafeZoneRatio;
-};
 
+	/** Whether any active timers are pending for this Slate tick. Updated once a tick in AnyActiveTimersArePending(). */
+	bool bAnyActiveTimersPending = false;
+};
