[MenuBuilder & BasicFilterBar] making sure menu entry respect menu builder's CloseSelfOnly setting
+basic filter bar not close entire menu stack in case the filter bar is used in another menu

diff --git a/Engine/Source/Developer/ToolMenus/Private/ToolMenus.cpp b/Engine/Source/Developer/ToolMenus/Private/ToolMenus.cpp
--- a/Engine/Source/Developer/ToolMenus/Private/ToolMenus.cpp
+++ b/Engine/Source/Developer/ToolMenus/Private/ToolMenus.cpp
@@ -269,17 +269,16 @@ public:
 			}
 			else
 			{
-				FMenuEntryParams MenuEntryParams;
-				MenuEntryParams.LabelOverride = LabelToDisplay;
-				MenuEntryParams.ToolTipOverride = Block.ToolTip;
-				MenuEntryParams.IconOverride = Block.Icon.Get();
-				MenuEntryParams.DirectActions = UIAction;
-				MenuEntryParams.ExtensionHook = BlockNameOverride;
-				MenuEntryParams.UserInterfaceActionType = Block.UserInterfaceActionType;
-				MenuEntryParams.TutorialHighlightName = Block.TutorialHighlightName;
-				MenuEntryParams.InputBindingOverride = Block.InputBindingLabel;
-				
-				MenuBuilder.AddMenuEntry(MenuEntryParams);
+				MenuBuilder.AddMenuEntry(
+					LabelToDisplay,
+					Block.ToolTip,
+					Block.Icon.Get(),
+					UIAction,
+					BlockNameOverride,
+					Block.UserInterfaceActionType,
+					Block.TutorialHighlightName,
+					Block.InputBindingLabel
+				);
 			}
 		}
 	}
diff --git a/Engine/Source/Developer/ToolWidgets/Public/Filters/SBasicFilterBar.h b/Engine/Source/Developer/ToolWidgets/Public/Filters/SBasicFilterBar.h
--- a/Engine/Source/Developer/ToolWidgets/Public/Filters/SBasicFilterBar.h
+++ b/Engine/Source/Developer/ToolWidgets/Public/Filters/SBasicFilterBar.h
@@ -588,7 +588,10 @@ protected:
 		/** Handler to create a right click menu */
 		TSharedRef<SWidget> GetRightClickMenuContent()
 		{
-			FMenuBuilder MenuBuilder(/*bInShouldCloseWindowAfterMenuSelection=*/true, NULL);
+			constexpr bool bShouldCloseWindowAfterMenuSelection = true;
+			// CloseSelfOnly in case the FilterBar is part of another menu
+			constexpr bool bCloseSelfOnly = true;
+			FMenuBuilder MenuBuilder(bShouldCloseWindowAfterMenuSelection, nullptr, {}, bCloseSelfOnly);
 
 			MenuBuilder.BeginSection("FilterOptions", LOCTEXT("FilterContextHeading", "Filter Options"));
 			{
diff --git a/Engine/Source/Runtime/Slate/Private/Framework/MultiBox/MultiBoxBuilder.cpp b/Engine/Source/Runtime/Slate/Private/Framework/MultiBox/MultiBoxBuilder.cpp
--- a/Engine/Source/Runtime/Slate/Private/Framework/MultiBox/MultiBoxBuilder.cpp
+++ b/Engine/Source/Runtime/Slate/Private/Framework/MultiBox/MultiBoxBuilder.cpp
@@ -164,13 +164,14 @@ void FBaseMenuBuilder::AddMenuEntry( const TSharedPtr< const FUICommandInfo > In
 	ApplyHook(InExtensionHook, EExtensionHook::After);
 }
 
-void FBaseMenuBuilder::AddMenuEntry( const TAttribute<FText>& InLabel, const TAttribute<FText>& InToolTip, const FSlateIcon& InIcon, const FUIAction& InAction, FName InExtensionHook, const EUserInterfaceActionType UserInterfaceActionType, FName InTutorialHighlightName )
+void FBaseMenuBuilder::AddMenuEntry( const TAttribute<FText>& InLabel, const TAttribute<FText>& InToolTip, const FSlateIcon& InIcon, const FUIAction& InAction, FName InExtensionHook, const EUserInterfaceActionType UserInterfaceActionType, FName InTutorialHighlightName, const TAttribute<FText>& InInputBindingOverride)
 {
 	ApplySectionBeginning();
 
 	ApplyHook(InExtensionHook, EExtensionHook::Before);
 	
-	TSharedRef< FMenuEntryBlock > NewMenuEntryBlock = MakeShared<FMenuEntryBlock>( InExtensionHook, InLabel, InToolTip, InIcon, InAction, UserInterfaceActionType, bCloseSelfOnly );
+	constexpr bool bShouldCloseWindowAfterMenuSelection = true;
+	TSharedRef< FMenuEntryBlock > NewMenuEntryBlock = MakeShared<FMenuEntryBlock>(InExtensionHook, InLabel, InToolTip, InIcon, InAction, UserInterfaceActionType, bCloseSelfOnly, bShouldCloseWindowAfterMenuSelection, InInputBindingOverride);
 	NewMenuEntryBlock->SetTutorialHighlightName(GenerateTutorialIdentifierName(TutorialHighlightName, InTutorialHighlightName, nullptr, MultiBox->GetBlocks().Num()));
 	MultiBox->AddMultiBlock( NewMenuEntryBlock );
 	
diff --git a/Engine/Source/Runtime/Slate/Private/Framework/MultiBox/SMenuEntryBlock.cpp b/Engine/Source/Runtime/Slate/Private/Framework/MultiBox/SMenuEntryBlock.cpp
--- a/Engine/Source/Runtime/Slate/Private/Framework/MultiBox/SMenuEntryBlock.cpp
+++ b/Engine/Source/Runtime/Slate/Private/Framework/MultiBox/SMenuEntryBlock.cpp
@@ -48,10 +48,11 @@ FMenuEntryBlock::FMenuEntryBlock( const FName& InExtensionHook, const TAttribute
 }
 
 
-FMenuEntryBlock::FMenuEntryBlock( const FName& InExtensionHook, const TAttribute<FText>& InLabel, const TAttribute<FText>& InToolTip, const FSlateIcon& InIcon, const FUIAction& UIAction, const EUserInterfaceActionType InUserInterfaceActionType, bool bInCloseSelfOnly, bool bInShouldCloseWindowAfterMenuSelection)
+FMenuEntryBlock::FMenuEntryBlock( const FName& InExtensionHook, const TAttribute<FText>& InLabel, const TAttribute<FText>& InToolTip, const FSlateIcon& InIcon, const FUIAction& UIAction, const EUserInterfaceActionType InUserInterfaceActionType, bool bInCloseSelfOnly, bool bInShouldCloseWindowAfterMenuSelection, const TAttribute<FText>& InInputBindingOverride)
 	: FMultiBlock( UIAction, InExtensionHook, EMultiBlockType::MenuEntry )
 	, LabelOverride( InLabel )
 	, ToolTipOverride( InToolTip )
+	, InputBindingOverride( InInputBindingOverride)
 	, IconOverride( InIcon )
 	, bIsSubMenu( false )
 	, bIsRecursivelySearchable( true )
diff --git a/Engine/Source/Runtime/Slate/Private/Framework/MultiBox/SMenuEntryBlock.h b/Engine/Source/Runtime/Slate/Private/Framework/MultiBox/SMenuEntryBlock.h
--- a/Engine/Source/Runtime/Slate/Private/Framework/MultiBox/SMenuEntryBlock.h
+++ b/Engine/Source/Runtime/Slate/Private/Framework/MultiBox/SMenuEntryBlock.h
@@ -104,7 +104,7 @@ public:
 	 */
 	FMenuEntryBlock( const FName& InExtensionHook, const TAttribute<FText>& InLabel, const TAttribute<FText>& InToolTip, const FNewMenuDelegate& InEntryBuilder, TSharedPtr<FExtender> InExtender, bool bInSubMenu, bool bInSubMenuOnClick, const FSlateIcon& InIcon, const FUIAction& InUIAction, const EUserInterfaceActionType InUserInterfaceActionType, bool bInCloseSelfOnly, bool bInShouldCloseWindowAfterMenuSelection = true, TSharedPtr< const FUICommandList > InCommandList = nullptr );
 
-	FMenuEntryBlock( const FName& InExtensionHook, const TAttribute<FText>& InLabel, const TAttribute<FText>& InToolTip, const FSlateIcon& InIcon, const FUIAction& InUIAction, const EUserInterfaceActionType InUserInterfaceActionType, bool bInCloseSelfOnly, bool bInShouldCloseWindowAfterMenuSelection = true);
+	FMenuEntryBlock( const FName& InExtensionHook, const TAttribute<FText>& InLabel, const TAttribute<FText>& InToolTip, const FSlateIcon& InIcon, const FUIAction& InUIAction, const EUserInterfaceActionType InUserInterfaceActionType, bool bInCloseSelfOnly, bool bInShouldCloseWindowAfterMenuSelection = true, const TAttribute<FText>& InInputBindingOverride = TAttribute<FText>());
 
 	FMenuEntryBlock( const FName& InExtensionHook, const FUIAction& UIAction, const TSharedRef< SWidget > Contents, const TAttribute<FText>& InToolTip, const EUserInterfaceActionType InUserInterfaceActionType, bool bInCloseSelfOnly, bool bInShouldCloseWindowAfterMenuSelection = true);
 
diff --git a/Engine/Source/Runtime/Slate/Public/Framework/MultiBox/MultiBoxBuilder.h b/Engine/Source/Runtime/Slate/Public/Framework/MultiBox/MultiBoxBuilder.h
--- a/Engine/Source/Runtime/Slate/Public/Framework/MultiBox/MultiBoxBuilder.h
+++ b/Engine/Source/Runtime/Slate/Public/Framework/MultiBox/MultiBoxBuilder.h
@@ -272,8 +272,9 @@ public:
 	 * @param	InExtensionHook			The section hook. Can be NAME_None
 	 * @param	UserInterfaceActionType	Type of interface action
 	 * @param	InTutorialHighlightName	Optional name to identify this widget and highlight during tutorials
+	 * @param	InInputBindingOverride  Optional overridden input binding text for this menu entry.  If not set, then the UI action's binding will be used if available.
 	 */
-	SLATE_API void AddMenuEntry( const TAttribute<FText>& InLabel, const TAttribute<FText>& InToolTip, const FSlateIcon& InIcon, const FUIAction& UIAction, FName InExtensionHook = NAME_None, const EUserInterfaceActionType UserInterfaceActionType = EUserInterfaceActionType::Button, FName InTutorialHighlightName = NAME_None );
+	SLATE_API void AddMenuEntry( const TAttribute<FText>& InLabel, const TAttribute<FText>& InToolTip, const FSlateIcon& InIcon, const FUIAction& UIAction, FName InExtensionHook = NAME_None, const EUserInterfaceActionType UserInterfaceActionType = EUserInterfaceActionType::Button, FName InTutorialHighlightName = NAME_None, const TAttribute<FText>& InInputBindingOverride = TAttribute<FText>());
 	
 	/**
 	 * Adds a menu entry with a custom widget
