Allow overriding the VAlign of the checkbox inner image

It still defaults to center, but can now be overridden if you need the image to align differently against the contents of the checkbox (eg, to top-align the image against the content)

#rnx

diff --git a/Engine/Source/Runtime/Slate/Private/Widgets/Input/SCheckBox.cpp b/Engine/Source/Runtime/Slate/Private/Widgets/Input/SCheckBox.cpp
--- a/Engine/Source/Runtime/Slate/Private/Widgets/Input/SCheckBox.cpp
+++ b/Engine/Source/Runtime/Slate/Private/Widgets/Input/SCheckBox.cpp
@@ -54,6 +54,7 @@ void SCheckBox::Construct( const SCheckBox::FArguments& InArgs )
 	CheckBoxTypeOverride = InArgs._Type;
 
 	HorizontalAlignment = InArgs._HAlign;
+	CheckBoxImageVAlign = InArgs._CheckBoxImageVAlign;
 	bCheckBoxContentUsesAutoWidth = InArgs._CheckBoxContentUsesAutoWidth;
 
 	bIsPressed = false;
@@ -520,7 +521,7 @@ void SCheckBox::BuildCheckBox(TSharedRef<SWidget> InContent)
 			SNew(SHorizontalBox)
 			+ SHorizontalBox::Slot()
 			.AutoWidth()
-			.VAlign(VAlign_Center)
+			.VAlign(CheckBoxImageVAlign)
 			.HAlign(HAlign_Center)
 			[
 				SNew(SOverlay)
diff --git a/Engine/Source/Runtime/Slate/Public/Widgets/Input/SCheckBox.h b/Engine/Source/Runtime/Slate/Public/Widgets/Input/SCheckBox.h
--- a/Engine/Source/Runtime/Slate/Public/Widgets/Input/SCheckBox.h
+++ b/Engine/Source/Runtime/Slate/Public/Widgets/Input/SCheckBox.h
@@ -37,6 +37,7 @@ public:
 		, _OnCheckStateChanged()
 		, _IsChecked( ECheckBoxState::Unchecked )
 		, _HAlign( HAlign_Fill )
+		, _CheckBoxImageVAlign( VAlign_Center )
 		, _CheckBoxContentUsesAutoWidth(true)
 		, _Padding()
 		, _ClickMethod( EButtonClickMethod::DownAndUp )
@@ -75,9 +76,12 @@ public:
 		/** Whether the check box is currently in a checked state */
 		SLATE_ATTRIBUTE( ECheckBoxState, IsChecked )
 
-		/** How the content of the toggle button should align within the given space*/
+		/** How the content of the toggle button should align within the given space */
 		SLATE_ARGUMENT( EHorizontalAlignment, HAlign )
 
+		/** How the image of the checkbox should align vertically within the given space */
+		SLATE_ARGUMENT( EVerticalAlignment, CheckBoxImageVAlign )
+
 		/** Whether or not the content portion of the checkbox should layout using auto-width. When true the content will always be arranged at its desired size as opposed to resizing to the available space. */
 		SLATE_ARGUMENT(bool, CheckBoxContentUsesAutoWidth)
 
@@ -326,9 +330,12 @@ protected:
 	/** Overrides checkbox type in the widget style, if set */
 	TOptional<ESlateCheckBoxType::Type> CheckBoxTypeOverride;
 
-	/** Horiz align setting if in togglebox mode */
+	/** How the content of the toggle button should align within the given space */
 	EHorizontalAlignment HorizontalAlignment;
 
+	/** How the image of the checkbox should align vertically within the given space */
+	EVerticalAlignment CheckBoxImageVAlign;
+
 	/** Whether or not the checkbox content is arranged using auto-width when in checkbox mode. */
 	bool bCheckBoxContentUsesAutoWidth;
 
