TabManager : In RestoreArea_Helper add the NewWindow before restoring content because the NewWindow can become the ParentWindow during recursion.

Destroy the NewWindow as we unwind in the case that it is not needed.  This fixes a crash while restoring layouts for the case where a nomad tab in a floating window has one of it's panel tab children also as a floating window ( i.e. a 3 level hierarchy: RootWindow -> NomadTab Window -> PanelTab Window )

diff --git a/Engine/Source/Runtime/Slate/Private/Framework/Docking/TabManager.cpp b/Engine/Source/Runtime/Slate/Private/Framework/Docking/TabManager.cpp
--- a/Engine/Source/Runtime/Slate/Private/Framework/Docking/TabManager.cpp
+++ b/Engine/Source/Runtime/Slate/Private/Framework/Docking/TabManager.cpp
@@ -1786,6 +1786,18 @@ TSharedPtr<SDockingNode> FTabManager::RestoreArea_Helper(const TSharedRef<FLayou
 			// Set a default title; restoring the splitter content may override this if it activates a tab
 			NewWindow->SetTitle(FGlobalTabmanager::Get()->GetApplicationTitle());
 
+			// We need to add the new window now before we recursively restore any content.
+			// The reason for this is that NewWindow may become the ParentWindow for another window as we restore content.
+			// We destroy the new window as we unwind if it ends up being extraneous
+			if (bIsChildWindow)
+			{
+				FSlateApplication::Get().AddWindowAsNativeChild(NewWindow, ParentWindow.ToSharedRef());
+			}
+			else
+			{
+				FSlateApplication::Get().AddWindow(NewWindow);
+			}
+
 			TArray<TSharedRef<SDockingNode>> DockingNodes;
 			if (CanRestoreSplitterContent(DockingNodes, NodeAsArea.ToSharedRef(), NewWindow, OutSidebarTabs, OutputCanBeNullptr))
 			{
@@ -1803,18 +1815,19 @@ TSharedPtr<SDockingNode> FTabManager::RestoreArea_Helper(const TSharedRef<FLayou
 
 				if (bIsChildWindow)
 				{
-					// Recursively check to see how many actually spawned tabs there are in this dock area. If there are none we will not spawn a useless window
+					// Recursively check to see how many actually spawned tabs there are in this dock area.					
 					const int32 TotalNumTabs = NewDockAreaWidget->GetNumTabs();
 
-					if (TotalNumTabs > 0 || bForceOpenWindowIfNeeded)
+					// If there are none and we aren't requested to force open the window, then destroy the window.
+					if (TotalNumTabs == 0 && !bForceOpenWindowIfNeeded)
 					{
-						FSlateApplication::Get().AddWindowAsNativeChild(NewWindow, ParentWindow.ToSharedRef());
+						NewWindow->RequestDestroyWindow();
 					}
 				}
-				else
-				{
-					FSlateApplication::Get().AddWindow(NewWindow);
-				}
+			}
+			else
+			{
+				NewWindow->RequestDestroyWindow();
 			}
 		}
 		else
