Slate: when running headlessly, ignore pending timers when deciding whether to paint widgets.

- It looks like this isn't a good gate, it is true pretty much all the time and defeats the purpose of the change.
- We'll just use the Slate.SkipWidgetDrawingInHeadlessMode CVar to handle potentially important widget ticks. We should discourage putting load-bearing code in UI drawing anyway.

diff --git a/Engine/Source/Runtime/Slate/Private/Framework/Application/SlateApplication.cpp b/Engine/Source/Runtime/Slate/Private/Framework/Application/SlateApplication.cpp
--- a/Engine/Source/Runtime/Slate/Private/Framework/Application/SlateApplication.cpp
+++ b/Engine/Source/Runtime/Slate/Private/Framework/Application/SlateApplication.cpp
@@ -1343,7 +1343,7 @@ TArray<SWindow*> GatherAllDescendants(const TArray< TSharedRef<SWindow> >& InWin
 
 void FSlateApplication::PrivateDrawWindows( TSharedPtr<SWindow> DrawOnlyThisWindow )
 {
-	if (GSlateSkipWidgetDrawingInHeadlessMode && !FApp::CanEverRender() && !bAnyActiveTimersPending)
+	if (GSlateSkipWidgetDrawingInHeadlessMode && !FApp::CanEverRender())
 	{
 		// early out, as window "drawing" can take 1-2ms of a -nullrhi PC game
 		return;
