[ContentBrowser] CB does not always scroll icons correctly into focus.

diff --git a/Engine/Source/Editor/ContentBrowser/Private/SAssetView.cpp b/Engine/Source/Editor/ContentBrowser/Private/SAssetView.cpp
--- a/Engine/Source/Editor/ContentBrowser/Private/SAssetView.cpp
+++ b/Engine/Source/Editor/ContentBrowser/Private/SAssetView.cpp
@@ -4801,7 +4801,7 @@ void SAssetView::RequestScrollIntoView(const TSharedPtr<FAssetViewItem>& Item)
 	switch ( GetCurrentViewType() )
 	{
 		case EAssetViewType::List: ListView->RequestScrollIntoView(Item); break;
-		case EAssetViewType::Tile: TileView->RequestScrollIntoView(Item); break;
+		case EAssetViewType::Tile: TileView->RequestNavigateToItem(Item); break;
 		case EAssetViewType::Column: ColumnView->RequestScrollIntoView(Item); break;
 		case EAssetViewType::Custom: ViewExtender->RequestScrollIntoView(Item); break;
 	}
diff --git a/Engine/Source/Runtime/Slate/Public/Widgets/Views/STileView.h b/Engine/Source/Runtime/Slate/Public/Widgets/Views/STileView.h
--- a/Engine/Source/Runtime/Slate/Public/Widgets/Views/STileView.h
+++ b/Engine/Source/Runtime/Slate/Public/Widgets/Views/STileView.h
@@ -466,12 +466,12 @@ protected:
 				const double MinDisplayedLine = this->bNavigateOnScrollIntoView ? FMath::FloorToDouble(ScrollLineOffset) : FMath::CeilToDouble(ScrollLineOffset);
 				const double MaxDisplayedLine = this->bNavigateOnScrollIntoView ? FMath::CeilToDouble(ScrollLineOffset + NumFullLinesInView) : FMath::FloorToDouble(ScrollLineOffset + NumFullLinesInView);
 
-				if (LineOfItem < MinDisplayedLine || LineOfItem > MaxDisplayedLine)
+				if (LineOfItem < MinDisplayedLine || LineOfItem > MaxDisplayedLine || MaxDisplayedLine - MinDisplayedLine <= 1)
 				{
 					// Set the line with the item at the beginning of the view area
 					float NewLineOffset = (float)LineOfItem;
 					// Center the line in the view area
-					NewLineOffset -= NumLinesInView * 0.5f;
+					NewLineOffset -= (NumLinesInView * 0.5f) - 0.5f;
 					// Convert the line offset into an item offset
 					double NewScrollOffset = NewLineOffset * (double)NumItemsPerLine;
 					// And clamp the scroll offset within the allowed limits
