FGlobalTabmanager : Add CanSavePersistentLayouts / SetCanSavePersistentLayouts to replace UnrealEdMisc::IsSavingLayoutOnClosedAllowed /  AllowSavingLayoutOnClose which are now UE_DEPRECATED.

This allows for FOnPersistLayout callbacks to check this boolean condition without the need to add a dependency on the UnrealEdMisc module.

This is also a step towards being able to check this boolean condition internally from FTabManager::SavePersistentLayout so that FOnPersistLayout callbacks need not do it themselves.

diff --git a/Engine/Source/Developer/SlateReflector/Private/Widgets/SWidgetReflector.cpp b/Engine/Source/Developer/SlateReflector/Private/Widgets/SWidgetReflector.cpp
--- a/Engine/Source/Developer/SlateReflector/Private/Widgets/SWidgetReflector.cpp
+++ b/Engine/Source/Developer/SlateReflector/Private/Widgets/SWidgetReflector.cpp
@@ -62,7 +62,6 @@
 #if WITH_EDITOR
 #include "Framework/Docking/LayoutService.h"
 #include "PropertyEditorModule.h"
-#include "UnrealEdMisc.h"
 #endif
 
 #define LOCTEXT_NAMESPACE "SWidgetReflector"
@@ -1066,7 +1065,7 @@ TSharedRef<SDockTab> SWidgetReflector::SpawnWidgetList(const FSpawnTabArgs& Args
 void SWidgetReflector::HandleTabManagerPersistLayout(const TSharedRef<FTabManager::FLayout>& LayoutToSave)
 {
 #if WITH_EDITOR
-	if (FUnrealEdMisc::Get().IsSavingLayoutOnClosedAllowed())
+	if (FGlobalTabmanager::Get()->CanSavePersistentLayouts())
 	{
 		FLayoutSaveRestore::SaveToConfig(GEditorLayoutIni, LayoutToSave);
 	}
diff --git a/Engine/Source/Editor/LevelEditor/Public/LevelEditor.h b/Engine/Source/Editor/LevelEditor/Public/LevelEditor.h
--- a/Engine/Source/Editor/LevelEditor/Public/LevelEditor.h
+++ b/Engine/Source/Editor/LevelEditor/Public/LevelEditor.h
@@ -432,7 +432,7 @@ private:
 	// Callback for persisting the Level Editor's layout.
 	void HandleTabManagerPersistLayout( const TSharedRef<FTabManager::FLayout>& LayoutToSave )
 	{
-		if (FUnrealEdMisc::Get().IsSavingLayoutOnClosedAllowed())
+		if (FGlobalTabmanager::Get()->CanSavePersistentLayouts())
 		{
 			FLayoutSaveRestore::SaveToConfig(GEditorLayoutIni, LayoutToSave);
 		}
diff --git a/Engine/Source/Editor/MainFrame/Private/Frame/MainFrameHandler.cpp b/Engine/Source/Editor/MainFrame/Private/Frame/MainFrameHandler.cpp
--- a/Engine/Source/Editor/MainFrame/Private/Frame/MainFrameHandler.cpp
+++ b/Engine/Source/Editor/MainFrame/Private/Frame/MainFrameHandler.cpp
@@ -137,7 +137,7 @@ void FMainFrameHandler::ShutDownEditor()
 	// Save the visual state of the editor before we even
 	// ask whether we can shut down.
 	TSharedRef<FGlobalTabmanager> GlobalTabManager = FGlobalTabmanager::Get();
-	if (FUnrealEdMisc::Get().IsSavingLayoutOnClosedAllowed())
+	if (GlobalTabManager->CanSavePersistentLayouts())
 	{
 		GlobalTabManager->SaveAllVisualState();
 	}
diff --git a/Engine/Source/Editor/MainFrame/Private/Frame/MainFrameHandler.h b/Engine/Source/Editor/MainFrame/Private/Frame/MainFrameHandler.h
--- a/Engine/Source/Editor/MainFrame/Private/Frame/MainFrameHandler.h
+++ b/Engine/Source/Editor/MainFrame/Private/Frame/MainFrameHandler.h
@@ -169,7 +169,7 @@ public:
 				if ( !bOkToExit && !bPromptedForExit )
 				{
 					FUnrealEdMisc::Get().ClearPendingProjectName();
-					FUnrealEdMisc::Get().AllowSavingLayoutOnClose(true);
+					FGlobalTabmanager::Get()->SetCanSavePersistentLayouts(true);
 					FUnrealEdMisc::Get().ForceDeletePreferences(false);
 					FUnrealEdMisc::Get().ClearConfigRestoreFilenames();
 				}
@@ -258,7 +258,7 @@ private:
 	// Callback for persisting the Level Editor's layout.
 	void HandleTabManagerPersistLayout( const TSharedRef<FTabManager::FLayout>& LayoutToSave )
 	{
-		if (FUnrealEdMisc::Get().IsSavingLayoutOnClosedAllowed())
+		if (FGlobalTabmanager::Get()->CanSavePersistentLayouts())
 		{
 			FLayoutSaveRestore::SaveToConfig(GEditorLayoutIni, LayoutToSave);
 		}
diff --git a/Engine/Source/Editor/MainFrame/Private/Menus/LayoutsMenu.cpp b/Engine/Source/Editor/MainFrame/Private/Menus/LayoutsMenu.cpp
--- a/Engine/Source/Editor/MainFrame/Private/Menus/LayoutsMenu.cpp
+++ b/Engine/Source/Editor/MainFrame/Private/Menus/LayoutsMenu.cpp
@@ -797,11 +797,11 @@ void FLayoutsMenuLoad::ReloadCurrentLayout()
 		bAutoRestoreAndDisableSavingOverride = AssetEditorSubsystem->GetAutoRestoreAndDisableSavingOverride();
 		AssetEditorSubsystem->SetAutoRestoreAndDisableSavingOverride(true);
 	}
-	FUnrealEdMisc::Get().AllowSavingLayoutOnClose(false);
+	FGlobalTabmanager::Get()->SetCanSavePersistentLayouts(false);
 	// Editor is reset on-the-fly
 	EditorReinit();
 	// Reenable config saving
-	FUnrealEdMisc::Get().AllowSavingLayoutOnClose(true);
+	FGlobalTabmanager::Get()->SetCanSavePersistentLayouts(true);
 	if (bAreAssetEditorOpened)
 	{
 		check(AssetEditorSubsystem);
diff --git a/Engine/Source/Editor/MassEntityDebugger/Private/SMassDebugger.cpp b/Engine/Source/Editor/MassEntityDebugger/Private/SMassDebugger.cpp
--- a/Engine/Source/Editor/MassEntityDebugger/Private/SMassDebugger.cpp
+++ b/Engine/Source/Editor/MassEntityDebugger/Private/SMassDebugger.cpp
@@ -12,7 +12,6 @@
 #include "Widgets/Input/SComboBox.h"
 #if WITH_EDITOR
 #include "Editor.h"
-#include "UnrealEdMisc.h"
 #endif // WITH_EDITOR
 
 #define LOCTEXT_NAMESPACE "SMassDebugger"
@@ -140,16 +139,13 @@ void SMassDebugger::Construct(const FArguments& InArgs, const TSharedRef<SDockTa
 		FTabManager::FOnPersistLayout::CreateStatic(
 			[](const TSharedRef<FTabManager::FLayout>& InLayout)
 			{
-#if WITH_EDITOR
-				if (FUnrealEdMisc::Get().IsSavingLayoutOnClosedAllowed() == false)
-				{
-					return;
-				}
-#endif // WITH_EDITOR
-				if (InLayout->GetPrimaryArea().Pin().IsValid())
+				if (FGlobalTabmanager::Get()->CanSavePersistentLayouts())
 				{
-					FLayoutSaveRestore::SaveToConfig(GEditorLayoutIni, InLayout);
-				}
+					if (InLayout->GetPrimaryArea().Pin().IsValid())
+					{
+						FLayoutSaveRestore::SaveToConfig(GEditorLayoutIni, InLayout);
+					}
+				}				
 			}
 		)
 	);
diff --git a/Engine/Source/Editor/UnrealEd/Private/ApplicationMode.cpp b/Engine/Source/Editor/UnrealEd/Private/ApplicationMode.cpp
--- a/Engine/Source/Editor/UnrealEd/Private/ApplicationMode.cpp
+++ b/Engine/Source/Editor/UnrealEd/Private/ApplicationMode.cpp
@@ -30,7 +30,7 @@ void FApplicationMode::DeactivateMode(TSharedPtr<FTabManager> InTabManager)
 {
 	// Save the layout to INI
 	check(InTabManager.IsValid());
-	if (FUnrealEdMisc::Get().IsSavingLayoutOnClosedAllowed())
+	if (FGlobalTabmanager::Get()->CanSavePersistentLayouts())
 	{
 		FLayoutSaveRestore::SaveToConfig(GEditorLayoutIni, InTabManager->PersistLayout());
 	}
diff --git a/Engine/Source/Editor/UnrealEd/Private/UnrealEdMisc.cpp b/Engine/Source/Editor/UnrealEd/Private/UnrealEdMisc.cpp
--- a/Engine/Source/Editor/UnrealEd/Private/UnrealEdMisc.cpp
+++ b/Engine/Source/Editor/UnrealEd/Private/UnrealEdMisc.cpp
@@ -211,7 +211,6 @@ namespace UnrealEdMiscDefs
 FUnrealEdMisc::FUnrealEdMisc() :
 	bCancelBuild( false ),
 	bInitialized( false ),
-	bSaveLayoutOnClose( true ),
 	bDeletePreferences( false ),
 	bIsAssetAnalyticsPending( false ),
 	PerformanceAnalyticsStats(new FPerformanceAnalyticsStats()),
@@ -1775,6 +1774,16 @@ void FUnrealEdMisc::LogAssetUpdate(UObject* UpdatedAsset, FObjectPreSaveContext
 	}
 }
 
+void FUnrealEdMisc::AllowSavingLayoutOnClose(bool bIsEnabled)
+{
+	FGlobalTabmanager::Get()->SetCanSavePersistentLayouts(bIsEnabled);
+}
+
+bool FUnrealEdMisc::IsSavingLayoutOnClosedAllowed()
+{
+	return FGlobalTabmanager::Get()->CanSavePersistentLayouts();
+}
+
 void FUnrealEdMisc::SwitchProject(const FString& GameOrProjectFileName, bool bWarn, const TOptional<FString>& NewCommandLine)
 {
 	if (GUnrealEd->WarnIfLightingBuildIsCurrentlyRunning())
diff --git a/Engine/Source/Editor/UnrealEd/Public/Toolkits/AssetEditorToolkit.h b/Engine/Source/Editor/UnrealEd/Public/Toolkits/AssetEditorToolkit.h
--- a/Engine/Source/Editor/UnrealEd/Public/Toolkits/AssetEditorToolkit.h
+++ b/Engine/Source/Editor/UnrealEd/Public/Toolkits/AssetEditorToolkit.h
@@ -446,7 +446,7 @@ private:
 	// Callback for persisting the Asset Editor's layout.
 	void HandleTabManagerPersistLayout( const TSharedRef<FTabManager::FLayout>& LayoutToSave )
 	{
-		if (FUnrealEdMisc::Get().IsSavingLayoutOnClosedAllowed())
+		if (FGlobalTabmanager::Get()->CanSavePersistentLayouts())
 		{
 			FLayoutSaveRestore::SaveToConfig(GEditorLayoutIni, LayoutToSave);
 		}
diff --git a/Engine/Source/Editor/UnrealEd/Public/UnrealEdMisc.h b/Engine/Source/Editor/UnrealEd/Public/UnrealEdMisc.h
--- a/Engine/Source/Editor/UnrealEd/Public/UnrealEdMisc.h
+++ b/Engine/Source/Editor/UnrealEd/Public/UnrealEdMisc.h
@@ -102,16 +102,12 @@ public:
 	 *
 	 * @param bIsEnabled	true if saving on close is allowed.
 	 */
-	void AllowSavingLayoutOnClose(bool bIsEnabled)
-	{
-		bSaveLayoutOnClose = bIsEnabled;
-	}
+	UE_DEPRECATED(5.6, "Please use FGlobalTabmanager::Get()->SetCanSavePersistentLayouts instead.")
+	UNREALED_API void AllowSavingLayoutOnClose(bool bIsEnabled);
 
 	/** Returns true if saving layout on close is allowed. */
-	bool IsSavingLayoutOnClosedAllowed()
-	{
-		return bSaveLayoutOnClose;
-	}
+	UE_DEPRECATED(5.6, "Please use FGlobalTabmanager::Get()->CanSavePersistentLayouts instead.")
+	UNREALED_API bool IsSavingLayoutOnClosedAllowed();
 
 	/**
 	 * Sets the config file to use for restoring Config files.
@@ -308,9 +304,6 @@ private:
 	/** Map of config->backup filenames to restore on shutdown of the editor */
 	TMap<FString, FString> RestoreConfigFiles;
 
-	/** true if the layout should be saved when closing the editor. */
-	bool bSaveLayoutOnClose;
-
 	/** true if the preferences config file should be deleted. */
 	bool bDeletePreferences;
 
diff --git a/Engine/Source/Runtime/Slate/Private/Framework/Docking/TabManager.cpp b/Engine/Source/Runtime/Slate/Private/Framework/Docking/TabManager.cpp
--- a/Engine/Source/Runtime/Slate/Private/Framework/Docking/TabManager.cpp
+++ b/Engine/Source/Runtime/Slate/Private/Framework/Docking/TabManager.cpp
@@ -2973,6 +2973,16 @@ void FGlobalTabmanager::FinishRestore()
 	}
 }
 
+void FGlobalTabmanager::SetCanSavePersistentLayouts(bool bInCanSavePersistentLayouts)
+{
+	bCanSavePersistentLayouts = bInCanSavePersistentLayouts;
+}
+
+bool FGlobalTabmanager::CanSavePersistentLayouts() const
+{
+	return bCanSavePersistentLayouts;
+}
+
 void FGlobalTabmanager::SetProxyTabManager(TSharedPtr<FProxyTabmanager> InProxyTabManager)
 {
 	ProxyTabManager = InProxyTabManager;
diff --git a/Engine/Source/Runtime/Slate/Public/Framework/Docking/TabManager.h b/Engine/Source/Runtime/Slate/Public/Framework/Docking/TabManager.h
--- a/Engine/Source/Runtime/Slate/Public/Framework/Docking/TabManager.h
+++ b/Engine/Source/Runtime/Slate/Public/Framework/Docking/TabManager.h
@@ -1372,6 +1372,12 @@ public:
 	/** Returns the highest number of parent windows that were open simultaneously during this session */
 	int32 GetMaximumWindowCount() const { return AllAreasWindowMaxCount; }
 
+	/** Sets whether or not any layouts ( across all tab managers ) can be saved at this time. */
+	SLATE_API void SetCanSavePersistentLayouts(bool bInCanSavePersistentLayouts);
+
+	/** Returns true if any layouts ( across all tab managers ) can be saved at this time. */
+	SLATE_API bool CanSavePersistentLayouts() const;
+
 	SLATE_API void SetProxyTabManager(TSharedPtr<FProxyTabmanager> InProxyTabManager);
 
 	DECLARE_DELEGATE(FOnOverrideDockableAreaRestore);
@@ -1454,6 +1460,7 @@ private:
 	: FTabManager( TSharedPtr<SDockTab>(), MakeShareable( new FTabSpawner() ) )
 		, AllTabsMaxCount(0)
 		, AllAreasWindowMaxCount(0)
+		, bCanSavePersistentLayouts(true)
 	{
 	}
 
@@ -1480,6 +1487,9 @@ private:
 	/** Keeps track of the running-maximum number of unique parent windows in all dock areas and sub-managers during this session */
 	int32 AllAreasWindowMaxCount;
 
+	/** true if any layouts ( across all tab managers ) can be saved at this time. */
+	bool bCanSavePersistentLayouts;
+
 	/**  */
 	TSharedPtr<FProxyTabmanager> ProxyTabManager;
 
