Apply widget scale to corner radius in SBackgroundBlur widget. This ensures the corners stay the same relative to other widgets being scaled.

[RN]

diff --git a/Engine/Source/Runtime/Slate/Private/Widgets/Layout/SBackgroundBlur.cpp b/Engine/Source/Runtime/Slate/Private/Widgets/Layout/SBackgroundBlur.cpp
--- a/Engine/Source/Runtime/Slate/Private/Widgets/Layout/SBackgroundBlur.cpp
+++ b/Engine/Source/Runtime/Slate/Private/Widgets/Layout/SBackgroundBlur.cpp
@@ -146,7 +146,7 @@ int32 SBackgroundBlur::OnPaint(const FPaintArgs& Args, const FGeometry& Allotted
 				{
 					OutDrawElements.PushClip(FSlateClippingZone(AllottedGeometry));
 
-					FSlateDrawElement::MakePostProcessBlur(OutDrawElements, LayerId, PaintGeometry, FVector4f((float)KernelSize, ComputedStrength, (float)RenderTargetWidth, (float)RenderTargetHeight), DownsampleAmount, FVector4f(CornerRadiusAttribute.Get()));
+					FSlateDrawElement::MakePostProcessBlur(OutDrawElements, LayerId, PaintGeometry, FVector4f((float)KernelSize, ComputedStrength, (float)RenderTargetWidth, (float)RenderTargetHeight), DownsampleAmount, FVector4f(CornerRadiusAttribute.Get() * AllottedGeometry.Scale));
 
 					OutDrawElements.PopClip();
 				}
