[Viewport Toolbar] Add toolbar combo button menu placement override

Adds the possibility to override toolbar combo box menu content placement, allowing for custom menu alignment.

diff --git a/Engine/Source/Developer/ToolMenus/Private/ToolMenus.cpp b/Engine/Source/Developer/ToolMenus/Private/ToolMenus.cpp
--- a/Engine/Source/Developer/ToolMenus/Private/ToolMenus.cpp
+++ b/Engine/Source/Developer/ToolMenus/Private/ToolMenus.cpp
@@ -1770,7 +1770,8 @@ void UToolMenus::PopulateToolBarBuilderWithEntry(
 				true,
 				Block.TutorialHighlightName,
 				VisibilityOverride,
-				ToolbarLabelOverride
+				ToolbarLabelOverride,
+				Block.ToolBarData.PlacementOverride
 			);
 		}
 	}
@@ -1790,7 +1791,8 @@ void UToolMenus::PopulateToolBarBuilderWithEntry(
 				Block.ToolBarData.bSimpleComboBox,
 				Block.TutorialHighlightName,
 				VisibilityOverride,
-				ToolbarLabelOverride
+				ToolbarLabelOverride,
+				Block.ToolBarData.PlacementOverride
 			);
 		}
 		else
@@ -1807,7 +1809,8 @@ void UToolMenus::PopulateToolBarBuilderWithEntry(
 				Block.ToolBarData.bSimpleComboBox,
 				Block.TutorialHighlightName,
 				VisibilityOverride,
-				ToolbarLabelOverride
+				ToolbarLabelOverride,
+				Block.ToolBarData.PlacementOverride
 			);
 
 			// Also add any top-level flagged children to the toolbar.
diff --git a/Engine/Source/Developer/ToolMenus/Public/ToolMenuEntry.h b/Engine/Source/Developer/ToolMenus/Public/ToolMenuEntry.h
--- a/Engine/Source/Developer/ToolMenus/Public/ToolMenuEntry.h
+++ b/Engine/Source/Developer/ToolMenus/Public/ToolMenuEntry.h
@@ -66,6 +66,9 @@ public:
 
 	/** Whether this toolbar should always use small icons, regardless of the current settings. */
 	bool bForceSmallIcons;
+
+	/** Optional override placement to customize menu placement once opened via e.g. a toolbar menu button */
+	TAttribute<EMenuPlacement> PlacementOverride;
 };
 
 
diff --git a/Engine/Source/Runtime/Slate/Private/Framework/MultiBox/MultiBoxBuilder.cpp b/Engine/Source/Runtime/Slate/Private/Framework/MultiBox/MultiBoxBuilder.cpp
--- a/Engine/Source/Runtime/Slate/Private/Framework/MultiBox/MultiBoxBuilder.cpp
+++ b/Engine/Source/Runtime/Slate/Private/Framework/MultiBox/MultiBoxBuilder.cpp
@@ -585,13 +585,14 @@ void FToolBarBuilder::AddComboButton(
 	bool bInSimpleComboBox,
 	FName InTutorialHighlightName,
 	TAttribute<EVisibility> InVisibilityOverride,
-	TAttribute<FText> InToolbarLabelOverride
+	TAttribute<FText> InToolbarLabelOverride,
+	TAttribute<EMenuPlacement> InPlacementOverride
 )
 {
 	ApplySectionBeginning();
 
 	TSharedRef<FToolBarComboButtonBlock> NewToolBarComboButtonBlock(new FToolBarComboButtonBlock(
-		InAction, InMenuContentGenerator, InLabelOverride, InToolTipOverride, InIconOverride, bInSimpleComboBox, InToolbarLabelOverride
+		InAction, InMenuContentGenerator, InLabelOverride, InToolTipOverride, InIconOverride, bInSimpleComboBox, InToolbarLabelOverride, InPlacementOverride
 	));
 
 	if ( LabelVisibility.IsSet() )
diff --git a/Engine/Source/Runtime/Slate/Private/Framework/MultiBox/SToolBarComboButtonBlock.cpp b/Engine/Source/Runtime/Slate/Private/Framework/MultiBox/SToolBarComboButtonBlock.cpp
--- a/Engine/Source/Runtime/Slate/Private/Framework/MultiBox/SToolBarComboButtonBlock.cpp
+++ b/Engine/Source/Runtime/Slate/Private/Framework/MultiBox/SToolBarComboButtonBlock.cpp
@@ -16,7 +16,8 @@ FToolBarComboButtonBlock::FToolBarComboButtonBlock(
 	const TAttribute<FText>& InToolTip,
 	const TAttribute<FSlateIcon>& InIcon,
 	bool bInSimpleComboBox,
-	TAttribute<FText> InToolbarLabelOverride
+	TAttribute<FText> InToolbarLabelOverride,
+	TAttribute<EMenuPlacement> InPlacementOverride
 )
 	: FMultiBlock(InAction, NAME_None, EMultiBlockType::ToolBarComboButton)
 	, MenuContentGenerator(InMenuContentGenerator)
@@ -24,6 +25,7 @@ FToolBarComboButtonBlock::FToolBarComboButtonBlock(
 	, ToolbarLabelOverride(InToolbarLabelOverride)
 	, ToolTip(InToolTip)
 	, Icon(InIcon)
+	, PlacementOverride(InPlacementOverride)
 	, LabelVisibility()
 	, bSimpleComboBox(bInSimpleComboBox)
 	, bForceSmallIcons(false)
@@ -32,7 +34,6 @@ FToolBarComboButtonBlock::FToolBarComboButtonBlock(
 
 void FToolBarComboButtonBlock::CreateMenuEntry(FMenuBuilder& MenuBuilder) const
 {
-	FName IconName;
 	FText EntryLabel = Label.Get();
 	if ( EntryLabel.IsEmpty() )
 	{
@@ -218,6 +219,7 @@ void SToolBarComboButtonBlock::BuildMultiBlockWidget(const ISlateStyle* StyleSet
 		.ComboButtonStyle(ComboStyle)
 		.ButtonStyle(ButtonStyle)
 		.ToolTipText(ToolBarComboButtonBlock->ToolTip)
+		.MenuPlacement(ToolBarComboButtonBlock->PlacementOverride)
 		.ForegroundColor(this, &SToolBarComboButtonBlock::OnGetForegroundColor)
 		// Route the content generator event
 		.OnGetMenuContent(this, &SToolBarComboButtonBlock::OnGetMenuContent)
@@ -227,7 +229,6 @@ void SToolBarComboButtonBlock::BuildMultiBlockWidget(const ISlateStyle* StyleSet
 		]
 	];
 
-
 	FMargin Padding = ToolBarStyle.ComboButtonPadding;
 	if (ToolBarComboButtonBlock->bSimpleComboBox)
 	{
diff --git a/Engine/Source/Runtime/Slate/Public/Framework/MultiBox/MultiBoxBuilder.h b/Engine/Source/Runtime/Slate/Public/Framework/MultiBox/MultiBoxBuilder.h
--- a/Engine/Source/Runtime/Slate/Public/Framework/MultiBox/MultiBoxBuilder.h
+++ b/Engine/Source/Runtime/Slate/Public/Framework/MultiBox/MultiBoxBuilder.h
@@ -622,7 +622,8 @@ public:
 		bool bInSimpleComboBox = false,
 		FName InTutorialHighlightName = NAME_None,
 		TAttribute<EVisibility> InVisibilityOverride = TAttribute<EVisibility>(),
-		TAttribute<FText> InToolbarLabelOverride = TAttribute<FText>() );
+		TAttribute<FText> InToolbarLabelOverride = TAttribute<FText>(),
+		TAttribute<EMenuPlacement> InPlacementOverride = TAttribute<EMenuPlacement>());
 
 	/**
 	 * Adds a tool bar stack button
diff --git a/Engine/Source/Runtime/Slate/Public/Framework/MultiBox/SToolBarComboButtonBlock.h b/Engine/Source/Runtime/Slate/Public/Framework/MultiBox/SToolBarComboButtonBlock.h
--- a/Engine/Source/Runtime/Slate/Public/Framework/MultiBox/SToolBarComboButtonBlock.h
+++ b/Engine/Source/Runtime/Slate/Public/Framework/MultiBox/SToolBarComboButtonBlock.h
@@ -40,7 +40,8 @@ public:
 		const TAttribute<FText>& InToolTip = TAttribute<FText>(),
 		const TAttribute<FSlateIcon>& InIcon = TAttribute<FSlateIcon>(),
 		bool bInSimpleComboBox = false,
-		TAttribute<FText> InToolbarLabelOverride = TAttribute<FText>()
+		TAttribute<FText> InToolbarLabelOverride = TAttribute<FText>(),
+		TAttribute<EMenuPlacement> InPlacementOverride = MenuPlacement_ComboBox
 	);
 
 	/** FMultiBlock interface */
@@ -88,6 +89,9 @@ private:
 	/** Optional overridden icon for this tool bar button.  IF not set, then the action's icon will be used instead. */
 	TAttribute<FSlateIcon> Icon;
 
+	/** Optional overriden setting for handling placement */
+	TAttribute<EMenuPlacement> PlacementOverride;
+
 	/** Controls the Labels visibility, defaults to GetIconVisibility if no override is provided */
 	TOptional< EVisibility > LabelVisibility;
 
