PR #12276: Changed SInputKeySelector to be able to assign multiple modifier keys in it

diff --git a/Engine/Source/Runtime/Slate/Private/Widgets/Input/SInputKeySelector.cpp b/Engine/Source/Runtime/Slate/Private/Widgets/Input/SInputKeySelector.cpp
--- a/Engine/Source/Runtime/Slate/Private/Widgets/Input/SInputKeySelector.cpp
+++ b/Engine/Source/Runtime/Slate/Private/Widgets/Input/SInputKeySelector.cpp
@@ -187,10 +187,10 @@ FReply SInputKeySelector::OnKeyUp( const FGeometry& MyGeometry, const FKeyEvent&
 
 		SelectKey(
 			KeyUp,
-			ModifierKey == EModifierKey::Shift,
-			ModifierKey == EModifierKey::Control,
-			ModifierKey == EModifierKey::Alt,
-			ModifierKey == EModifierKey::Command);
+			(ModifierKey & EModifierKey::Shift) != EModifierKey::None,
+			(ModifierKey & EModifierKey::Control) != EModifierKey::None,
+			(ModifierKey & EModifierKey::Alt) != EModifierKey::None,
+			(ModifierKey & EModifierKey::Command) != EModifierKey::None);
 		return FReply::Handled();
 	}
 	else if (!bIsSelectingKey && Button.IsValid())
