Initialize SSpinBox slate arguments

diff --git a/Engine/Source/Runtime/Slate/Public/Widgets/Input/SSpinBox.h b/Engine/Source/Runtime/Slate/Public/Widgets/Input/SSpinBox.h
--- a/Engine/Source/Runtime/Slate/Public/Widgets/Input/SSpinBox.h
+++ b/Engine/Source/Runtime/Slate/Public/Widgets/Input/SSpinBox.h
@@ -90,6 +90,9 @@ public:
 		, _ClearKeyboardFocusOnCommit(false)
 		, _SelectAllTextOnCommit(true)
 		, _MinDesiredWidth(0.0f)
+		, _Justification(ETextJustify::Left)
+		, _KeyboardType(Keyboard_Default)
+		, _PreventThrottling(true)
 	{}
 
 	/** The style used to draw this spinbox */
