[Slate] VAlign for menu entries defaults to centered for new overloads accepting FMenuEntryStyleParams. Previous overloads use previous behavior.

#rnx

diff --git a/Engine/Source/Runtime/Slate/Private/Framework/MultiBox/MultiBoxBuilder.cpp b/Engine/Source/Runtime/Slate/Private/Framework/MultiBox/MultiBoxBuilder.cpp
--- a/Engine/Source/Runtime/Slate/Private/Framework/MultiBox/MultiBoxBuilder.cpp
+++ b/Engine/Source/Runtime/Slate/Private/Framework/MultiBox/MultiBoxBuilder.cpp
@@ -689,7 +689,7 @@ void FToolBarBuilder::AddWidget(
 	FMenuEntryStyleParams StyleParams;
 	StyleParams.HorizontalAlignment = Alignment;
 
-	AddWidget(InWidget, StyleParams, InTutorialHighlightName, bSearchable, InCustomMenuDelegate, InVisibilityOverride);
+	AddWidgetInternal(InWidget, StyleParams, InTutorialHighlightName, bSearchable, InCustomMenuDelegate, InVisibilityOverride);
 }
 
 void FToolBarBuilder::AddWidget(
@@ -700,6 +700,21 @@ void FToolBarBuilder::AddWidget(
 	FNewMenuDelegate InCustomMenuDelegate,
 	TAttribute<EVisibility> InVisibilityOverride
 )
+{
+	FMenuEntryStyleParams StyleParams = InStyleParams;
+	StyleParams.VerticalAlignment = VAlign_Center; // Default to centered, versions prior to 5.5 decide alignment based on label presence
+
+	AddWidgetInternal(InWidget, InStyleParams, InTutorialHighlightName, bInSearchable, InCustomMenuDelegate, InVisibilityOverride);
+}
+
+void FToolBarBuilder::AddWidgetInternal(
+	TSharedRef<SWidget> InWidget,
+	const FMenuEntryStyleParams& InStyleParams,
+	FName InTutorialHighlightName,
+	bool bInSearchable,
+	FNewMenuDelegate InCustomMenuDelegate,
+	TAttribute<EVisibility> InVisibilityOverride
+)
 {
 	ApplySectionBeginning();
 
diff --git a/Engine/Source/Runtime/Slate/Public/Framework/MultiBox/MultiBoxBuilder.h b/Engine/Source/Runtime/Slate/Public/Framework/MultiBox/MultiBoxBuilder.h
--- a/Engine/Source/Runtime/Slate/Public/Framework/MultiBox/MultiBoxBuilder.h
+++ b/Engine/Source/Runtime/Slate/Public/Framework/MultiBox/MultiBoxBuilder.h
@@ -414,10 +414,6 @@ protected:
 	SLATE_API virtual void ApplySectionBeginning() override;
 
 public:
-	// These classes need access to the AddWrapperSubMenu() methods
-	//friend class FWidgetBlock;
-	//friend class FToolBarComboButtonBlock;
-
 	/**
 	 * Adds a sub-menu which is a menu within a menu
 	 * 
@@ -741,6 +737,16 @@ protected:
 
 	SLATE_API void InitializeToolBarButtonBlock(TSharedPtr<FToolBarButtonBlock> ButtonRowBlock, const FButtonArgs& ButtonArgs);
 
+	/** Handles AddWidget calls */
+	void AddWidgetInternal(
+		TSharedRef<SWidget> InWidget,
+		const FMenuEntryStyleParams& InStyleParams,
+		FName InTutorialHighlightName = NAME_None,
+		bool bInSearchable = true,
+		FNewMenuDelegate InCustomMenuDelegate = FNewMenuDelegate(),
+		TAttribute<EVisibility> InVisibilityOverride = TAttribute<EVisibility>()
+	);
+
 private:
 	/** Current extension hook name for sections to determine where sections begin and end */
 	FName CurrentSectionExtensionHook;
