Fix menu search bug where only the first character typed would be part of the search.

diff --git a/Engine/Source/Runtime/Slate/Private/Framework/MultiBox/MultiBox.cpp b/Engine/Source/Runtime/Slate/Private/Framework/MultiBox/MultiBox.cpp
--- a/Engine/Source/Runtime/Slate/Private/Framework/MultiBox/MultiBox.cpp
+++ b/Engine/Source/Runtime/Slate/Private/Framework/MultiBox/MultiBox.cpp
@@ -958,9 +958,13 @@ void SMultiBoxWidget::CreateSearchTextWidget()
 	TSharedRef<SBox> SearchBox =
 		SNew(SBox)
 			.Padding(FMargin(8, 0, 8, 0))
-			.Visibility_Lambda([this]() -> EVisibility {
-				return ShouldShowMenuSearchField() ? EVisibility::Visible : EVisibility::Collapsed;
-			})
+			.Visibility_Lambda(
+				[this]() -> EVisibility
+				{
+					const bool bShow = !SearchText.IsEmpty() || ShouldShowMenuSearchField();
+					return bShow ? EVisibility::Visible : EVisibility::Collapsed;
+				}
+			)
 			[
 				SearchTextWidget.ToSharedRef()
 			];
@@ -977,10 +981,15 @@ void SMultiBoxWidget::CreateSearchTextWidget()
 /** Called when the SearchText changes */
 void SMultiBoxWidget::OnFilterTextChanged(const FText& InFilterText)
 {
+	const bool bInitialSearch = SearchText.IsEmpty();
+	// Set SearchText here, so the SBox wrapping our SearchTextWidget will be visible when we attempt to focus the
+	// SearchTextWidget. Otherwise, setting focus will fail because the parent is invisible and therefore the
+	// SearchTextWidget is too, and Slate refuses to focus invisible widgets.
+	SearchText = InFilterText;
+
 	// Activate the searchbox if it was empty and we are putting text in it for the first time.
 	// This is for IME keyboards only because they don't go through the OnKeyChar route.
-	if (bSearchable && SearchText.IsEmpty() 
-		&& !InFilterText.IsEmpty() 
+	if (bSearchable && bInitialSearch && !InFilterText.IsEmpty()
 		&& !FSlateApplication::Get().HasUserFocusedDescendants(SearchTextWidget.ToSharedRef(), 0))
 	{
 		if (SearchTextWidget.IsValid() && SearchBlockWidget.IsValid())
@@ -991,8 +1000,6 @@ void SMultiBoxWidget::OnFilterTextChanged(const FText& InFilterText)
 		}
 	}
 
-	SearchText = InFilterText;
-
 	FilterMultiBoxEntries();
 }
 
