Add support for toggling the checked state of checkbox submenus

A checkbox submenu is a ToggleButton or RadioButton type ToolMenu entry that can also open a submenu, it is both a checkbox and a submenu.

Checkbox submenus can be toggled when they appear in menus by clicking their checkbox and their submenus can be opened by hovering anywhere over the checkbox submenu entry. However, the ToolMenu system supports raising entries from combo button submenus to the top-level of the toolbar. When a checkbox submenu is raised, it appears in the toolbar as a combo button, but traditional combo buttons only support an icon, label, and opening a submenu. This change adds support for toggling the checked state of such checkbox submenus when they appear in toolbars by clicking on their icon, or a checkbox when no icon is set.

diff --git a/Engine/Source/Developer/ToolMenus/Private/ToolMenus.cpp b/Engine/Source/Developer/ToolMenus/Private/ToolMenus.cpp
--- a/Engine/Source/Developer/ToolMenus/Private/ToolMenus.cpp
+++ b/Engine/Source/Developer/ToolMenus/Private/ToolMenus.cpp
@@ -1771,7 +1771,8 @@ void UToolMenus::PopulateToolBarBuilderWithEntry(
 				Block.TutorialHighlightName,
 				VisibilityOverride,
 				ToolbarLabelOverride,
-				Block.ToolBarData.PlacementOverride
+				Block.ToolBarData.PlacementOverride,
+				Block.UserInterfaceActionType
 			);
 		}
 	}
@@ -1792,7 +1793,8 @@ void UToolMenus::PopulateToolBarBuilderWithEntry(
 				Block.TutorialHighlightName,
 				VisibilityOverride,
 				ToolbarLabelOverride,
-				Block.ToolBarData.PlacementOverride
+				Block.ToolBarData.PlacementOverride,
+				Block.UserInterfaceActionType
 			);
 		}
 		else
@@ -1810,7 +1812,8 @@ void UToolMenus::PopulateToolBarBuilderWithEntry(
 				Block.TutorialHighlightName,
 				VisibilityOverride,
 				ToolbarLabelOverride,
-				Block.ToolBarData.PlacementOverride
+				Block.ToolBarData.PlacementOverride,
+				Block.UserInterfaceActionType
 			);
 
 			// Also add any top-level flagged children to the toolbar.
diff --git a/Engine/Source/Runtime/Slate/Private/Framework/MultiBox/MultiBoxBuilder.cpp b/Engine/Source/Runtime/Slate/Private/Framework/MultiBox/MultiBoxBuilder.cpp
--- a/Engine/Source/Runtime/Slate/Private/Framework/MultiBox/MultiBoxBuilder.cpp
+++ b/Engine/Source/Runtime/Slate/Private/Framework/MultiBox/MultiBoxBuilder.cpp
@@ -586,13 +586,14 @@ void FToolBarBuilder::AddComboButton(
 	FName InTutorialHighlightName,
 	TAttribute<EVisibility> InVisibilityOverride,
 	TAttribute<FText> InToolbarLabelOverride,
-	TAttribute<EMenuPlacement> InPlacementOverride
+	TAttribute<EMenuPlacement> InPlacementOverride,
+	const EUserInterfaceActionType InUserInterfaceActionType
 )
 {
 	ApplySectionBeginning();
 
 	TSharedRef<FToolBarComboButtonBlock> NewToolBarComboButtonBlock(new FToolBarComboButtonBlock(
-		InAction, InMenuContentGenerator, InLabelOverride, InToolTipOverride, InIconOverride, bInSimpleComboBox, InToolbarLabelOverride, InPlacementOverride
+		InAction, InMenuContentGenerator, InLabelOverride, InToolTipOverride, InIconOverride, bInSimpleComboBox, InToolbarLabelOverride, InPlacementOverride, InUserInterfaceActionType
 	));
 
 	if ( LabelVisibility.IsSet() )
diff --git a/Engine/Source/Runtime/Slate/Private/Framework/MultiBox/SToolBarComboButtonBlock.cpp b/Engine/Source/Runtime/Slate/Private/Framework/MultiBox/SToolBarComboButtonBlock.cpp
--- a/Engine/Source/Runtime/Slate/Private/Framework/MultiBox/SToolBarComboButtonBlock.cpp
+++ b/Engine/Source/Runtime/Slate/Private/Framework/MultiBox/SToolBarComboButtonBlock.cpp
@@ -1,13 +1,17 @@
 // Copyright Epic Games, Inc. All Rights Reserved.
 
 #include "Framework/MultiBox/SToolBarComboButtonBlock.h"
+
+#include "Framework/Application/SlateApplication.h"
 #include "Widgets/SBoxPanel.h"
 #include "Widgets/Images/SImage.h"
 #include "Framework/MultiBox/MultiBoxBuilder.h"
+#include "Widgets/Input/SCheckBox.h"
 #include "Widgets/Input/SComboButton.h"
 #include "Widgets/Layout/SBox.h"
 #include "Styling/ToolBarStyle.h"
 #include "Widgets/Images/SLayeredImage.h"
+#include "Widgets/SToolTip.h"
 
 FToolBarComboButtonBlock::FToolBarComboButtonBlock(
 	const FUIAction& InAction,
@@ -17,7 +21,8 @@ FToolBarComboButtonBlock::FToolBarComboButtonBlock(
 	const TAttribute<FSlateIcon>& InIcon,
 	bool bInSimpleComboBox,
 	TAttribute<FText> InToolbarLabelOverride,
-	TAttribute<EMenuPlacement> InPlacementOverride
+	TAttribute<EMenuPlacement> InPlacementOverride,
+	const EUserInterfaceActionType InUserInterfaceActionType
 )
 	: FMultiBlock(InAction, NAME_None, EMultiBlockType::ToolBarComboButton)
 	, MenuContentGenerator(InMenuContentGenerator)
@@ -27,6 +32,7 @@ FToolBarComboButtonBlock::FToolBarComboButtonBlock(
 	, Icon(InIcon)
 	, PlacementOverride(InPlacementOverride)
 	, LabelVisibility()
+	, UserInterfaceActionType(InUserInterfaceActionType)
 	, bSimpleComboBox(bInSimpleComboBox)
 	, bForceSmallIcons(false)
 {
@@ -71,7 +77,7 @@ void SToolBarComboButtonBlock::BuildMultiBlockWidget(const ISlateStyle* StyleSet
 	
 	TSharedRef< const FToolBarComboButtonBlock > ToolBarComboButtonBlock = StaticCastSharedRef< const FToolBarComboButtonBlock >( MultiBlock.ToSharedRef() );
 
-	//TSharedPtr< const FUICommandInfo > UICommand = ToolBarComboButtonBlock->GetAction();
+	TSharedPtr<const FUICommandInfo> UICommand = ToolBarComboButtonBlock->GetAction();
 
 	TAttribute<FText> Label;
 
@@ -196,12 +202,10 @@ void SToolBarComboButtonBlock::BuildMultiBlockWidget(const ISlateStyle* StyleSet
 				.TextStyle(&ToolBarStyle.LabelStyle)
 			];
 	}
-	
-		
+
 	EMultiBlockLocation::Type BlockLocation = GetMultiBlockLocation();
 	FName BlockStyle = EMultiBlockLocation::ToName(ISlateStyle::Join(StyleName, ".Button"), BlockLocation);
 	const FButtonStyle* ButtonStyle = BlockLocation == EMultiBlockLocation::None ? &ToolBarStyle.ButtonStyle : &StyleSet->GetWidgetStyle<FButtonStyle>(BlockStyle);
-
 	const FComboButtonStyle* ComboStyle = &ToolBarStyle.ComboButtonStyle;
 	if (ToolBarComboButtonBlock->bSimpleComboBox)
 	{
@@ -211,23 +215,130 @@ void SToolBarComboButtonBlock::BuildMultiBlockWidget(const ISlateStyle* StyleSet
 
 	OpenForegroundColor = ButtonStyle->HoveredForeground;
 
-	ChildSlot
-	[
-		SAssignNew(ComboButtonWidget, SComboButton)
-		.AddMetaData<FTagMetaData>(FTagMetaData(TagName))
-		.ContentPadding(0.f)
-		.ComboButtonStyle(ComboStyle)
-		.ButtonStyle(ButtonStyle)
-		.ToolTipText(ToolBarComboButtonBlock->ToolTip)
-		.MenuPlacement(ToolBarComboButtonBlock->PlacementOverride)
-		.ForegroundColor(this, &SToolBarComboButtonBlock::OnGetForegroundColor)
-		// Route the content generator event
-		.OnGetMenuContent(this, &SToolBarComboButtonBlock::OnGetMenuContent)
-		.ButtonContent()
+	TAttribute<FText> ActualToolTip;
+	if (ToolBarComboButtonBlock->ToolTip.IsSet())
+	{
+		ActualToolTip = ToolBarComboButtonBlock->ToolTip;
+	}
+	else
+	{
+		ActualToolTip = UICommand.IsValid() ? UICommand->GetDescription() : FText::GetEmpty();
+	}
+
+	EUserInterfaceActionType UserInterfaceType = ToolBarComboButtonBlock->UserInterfaceActionType;
+	if (UICommand.IsValid())
+	{
+		// If we have a UICommand, then this is specified in the command.
+		UserInterfaceType = UICommand->GetUserInterfaceType();
+	}
+
+	if (UserInterfaceType == EUserInterfaceActionType::Button)
+	{
+		// clang-format off
+		ChildSlot
 		[
-			ButtonContent
-		]
-	];
+			SAssignNew(ComboButtonWidget, SComboButton)
+			.AddMetaData<FTagMetaData>(FTagMetaData(TagName))
+			.ContentPadding(0.f)
+			.ComboButtonStyle(ComboStyle)
+			.ButtonStyle(ButtonStyle)
+			.ToolTipText(ToolBarComboButtonBlock->ToolTip)
+			.MenuPlacement(ToolBarComboButtonBlock->PlacementOverride)
+			.ForegroundColor(this, &SToolBarComboButtonBlock::OnGetForegroundColor)
+			// Route the content generator event
+			.OnGetMenuContent(this, &SToolBarComboButtonBlock::OnGetMenuContent)
+			.ButtonContent()
+			[
+				ButtonContent
+			]
+		];
+		// clang-format on
+	}
+
+	else if (ensure(
+				 UserInterfaceType == EUserInterfaceActionType::ToggleButton
+				 || UserInterfaceType == EUserInterfaceActionType::RadioButton
+			 ))
+	{
+		FName ToggleBlockStyle = EMultiBlockLocation::ToName(ISlateStyle::Join(StyleName, ".ToggleButton"), BlockLocation);
+
+		const FCheckBoxStyle* CheckStyle;
+		if (!Icon.IsSet())
+		{
+			CheckStyle = &FCoreStyle::Get().GetWidgetStyle<FCheckBoxStyle>("Checkbox");
+			IconWidget = SNullWidget::NullWidget;
+		}
+		else if (BlockLocation == EMultiBlockLocation::None)
+		{
+			CheckStyle = &ToolBarStyle.ToggleButton;
+		}
+		else
+		{
+			CheckStyle = &StyleSet->GetWidgetStyle<FCheckBoxStyle>(ToggleBlockStyle);
+		}
+
+		if (OptionsBlockWidget.IsValid())
+		{
+			CheckStyle = &ToolBarStyle.SettingsToggleButton;
+		}
+
+		// clang-format off
+		const TSharedRef<SCheckBox> CheckBox = SNew(SCheckBox)
+			// Use the tool bar style for this check box
+			.Style(CheckStyle)
+			.CheckBoxContentUsesAutoWidth(false)
+			.OnCheckStateChanged(this, &SToolBarComboButtonBlock::OnCheckStateChanged)
+			.OnGetMenuContent(this, &SToolBarComboButtonBlock::OnGetMenuContent)
+			.IsChecked(this, &SToolBarComboButtonBlock::GetCheckState)
+			.IsEnabled(this, &SToolBarComboButtonBlock::IsEnabled)
+			[
+				IconWidget
+			];
+		// clang-format on
+
+		// clang-format off
+		const TSharedRef<SComboButton> ComboButton = SAssignNew(ComboButtonWidget, SComboButton)
+			.AddMetaData<FTagMetaData>(FTagMetaData(TagName))
+			.ContentPadding(0.f)
+			.ComboButtonStyle(ComboStyle)
+			.ButtonStyle(ButtonStyle)
+			.MenuPlacement(ToolBarComboButtonBlock->PlacementOverride)
+			.ForegroundColor(this, &SToolBarComboButtonBlock::OnGetForegroundColor)
+			// Route the content generator event
+			.OnGetMenuContent(this, &SToolBarComboButtonBlock::OnGetMenuContent)
+			.ButtonContent()
+			[
+				SNew(STextBlock)
+				.Visibility(ToolBarComboButtonBlock->bSimpleComboBox ? EVisibility::Collapsed : LabelVisibility)
+				.Text(Label)
+				// Smaller font for tool tip labels
+				.TextStyle(&ToolBarStyle.LabelStyle)
+				.Margin(0)
+			];
+		// clang-format on
+
+		// clang-format off
+		ChildSlot
+		[
+			SNew(SHorizontalBox)
+			.ToolTip(FMultiBoxSettings::ToolTipConstructor.Execute(
+				ActualToolTip, nullptr, UICommand, /*ShowActionShortcut=*/true
+			))
+			+ SHorizontalBox::Slot()
+			.AutoWidth()
+			.VAlign(VAlign_Center)
+			[
+				CheckBox
+			]
+			+ SHorizontalBox::Slot()
+			.AutoWidth()
+			.VAlign(VAlign_Center)
+			[
+				ComboButton
+			]
+		];
+		// clang-format on
+	}
 
 	FMargin Padding = ToolBarStyle.ComboButtonPadding;
 	if (ToolBarComboButtonBlock->bSimpleComboBox)
@@ -236,6 +347,7 @@ void SToolBarComboButtonBlock::BuildMultiBlockWidget(const ISlateStyle* StyleSet
 	}
 
 	ChildSlot.Padding(Padding);
+
 	// Bind our widget's enabled state to whether or not our action can execute
 	SetEnabled(TAttribute<bool>( this, &SToolBarComboButtonBlock::IsEnabled));
 
@@ -249,6 +361,60 @@ TSharedRef<SWidget> SToolBarComboButtonBlock::OnGetMenuContent()
 	return ToolBarButtonComboBlock->MenuContentGenerator.Execute();
 }
 
+FReply SToolBarComboButtonBlock::OnClicked()
+{
+	// Button was clicked, so trigger the action!
+	TSharedPtr<const FUICommandList> ActionList = MultiBlock->GetActionList();
+	TSharedPtr<const FUICommandInfo> Action = MultiBlock->GetAction();
+	const FUIAction& DirectActions = MultiBlock->GetDirectActions();
+
+	if (ActionList.IsValid() && Action.IsValid())
+	{
+		ActionList->ExecuteAction(Action.ToSharedRef());
+	}
+	else
+	{
+		// There is no action list or action associated with this block via a UI command.  Execute any direct action we have
+		MultiBlock->GetDirectActions().Execute();
+	}
+
+	TSharedRef<const FMultiBox> MultiBox(OwnerMultiBoxWidget.Pin()->GetMultiBox());
+
+	// If this is a context menu, then we'll also dismiss the window after the user clicked on the item
+	const bool bClosingMenu = MultiBox->ShouldCloseWindowAfterMenuSelection();
+	if (bClosingMenu)
+	{
+		FSlateApplication::Get().DismissMenuByWidget(AsShared());
+	}
+
+	return FReply::Handled();
+}
+
+void SToolBarComboButtonBlock::OnCheckStateChanged(const ECheckBoxState NewCheckedState)
+{
+	OnClicked();
+}
+
+ECheckBoxState SToolBarComboButtonBlock::GetCheckState() const
+{
+	TSharedPtr<const FUICommandList> ActionList = MultiBlock->GetActionList();
+	TSharedPtr<const FUICommandInfo> Action = MultiBlock->GetAction();
+	const FUIAction& DirectActions = MultiBlock->GetDirectActions();
+
+	ECheckBoxState CheckState = ECheckBoxState::Unchecked;
+	if (ActionList.IsValid() && Action.IsValid())
+	{
+		CheckState = ActionList->GetCheckState(Action.ToSharedRef());
+	}
+	else
+	{
+		// There is no action list or action associated with this block via a UI command.  Execute any direct action we have
+		CheckState = DirectActions.GetCheckState();
+	}
+
+	return CheckState;
+}
+
 bool SToolBarComboButtonBlock::IsEnabled() const
 {
 	const FUIAction& UIAction = MultiBlock->GetDirectActions();
diff --git a/Engine/Source/Runtime/Slate/Public/Framework/MultiBox/MultiBoxBuilder.h b/Engine/Source/Runtime/Slate/Public/Framework/MultiBox/MultiBoxBuilder.h
--- a/Engine/Source/Runtime/Slate/Public/Framework/MultiBox/MultiBoxBuilder.h
+++ b/Engine/Source/Runtime/Slate/Public/Framework/MultiBox/MultiBoxBuilder.h
@@ -558,7 +558,7 @@ public:
 	 * @param	InTutorialHighlightName	Name to identify this widget and highlight during tutorials
 	 * @param	InCustomMenuDelegate	Optional custom menu delegate for cases where the toolbar is compressed into a menu
 	 * @param	InVisibilityOverride	Optional visibility override which, if set and returning anything but Visible, overrides the visibility of the toolbar button.
-	 * @param	InToolbarLabelOverride	Optional label override for wh.  If omitted, then the action's label will be used instead.
+	 * @param	InToolbarLabelOverride	Optional label override for when this block appears in a toolbar. If omitted, then the action's label will be used instead.
 	 */
 	SLATE_API void AddToolBarButton(
 		const TSharedPtr<const FUICommandInfo> InCommand,
@@ -583,7 +583,7 @@ public:
 	 * @param	UserInterfaceActionType	Type of interface action
 	 * @param	InTutorialHighlightName	Name to identify this widget and highlight during tutorials
 	 * @param	InVisibilityOverride	Optional visibility override which, if set and returning anything but Visible, overrides the visibility of the toolbar button.
-	 * @param	InToolbarLabelOverride	Optional label override for wh.  If omitted, then the action's label will be used instead.
+	 * @param	InToolbarLabelOverride	Optional label override for when this block appears in a toolbar. If omitted, then the action's label will be used instead.
 	 */
 	SLATE_API void AddToolBarButton(
 		const FUIAction& InAction,
@@ -607,7 +607,9 @@ public:
 	 * @param	bInSimpleComboBox			If true, the icon and label won't be displayed
 	 * @param	InTutorialHighlightName		Name to identify this widget and highlight during tutorials
 	 * @param	InVisibilityOverride		Optional visibility override which, if set and returning anything but Visible, overrides the visibility of the combo button.
-	 * @param	InToolbarLabelOverride		Optional label override for wh.  If omitted, then the action's label will be used instead.
+	 * @param	InToolbarLabelOverride		Optional label override for when this block appears in a toolbar.  If omitted, then the action's label will be used instead.
+	 * @param	InPlacementOverride			Optional override placement to customize menu placement once opened via e.g. a toolbar menu button.
+	 * @param	InUserInterfaceActionType	Optional user interface action type to use if a command isn't bound. If a command is bound, we use the action type associated with the command instead.
 	 */
 	SLATE_API void AddComboButton(
 		const FUIAction& InAction,
@@ -619,7 +621,9 @@ public:
 		FName InTutorialHighlightName = NAME_None,
 		TAttribute<EVisibility> InVisibilityOverride = TAttribute<EVisibility>(),
 		TAttribute<FText> InToolbarLabelOverride = TAttribute<FText>(),
-		TAttribute<EMenuPlacement> InPlacementOverride = TAttribute<EMenuPlacement>());
+		TAttribute<EMenuPlacement> InPlacementOverride = TAttribute<EMenuPlacement>(),
+		const EUserInterfaceActionType InUserInterfaceActionType = EUserInterfaceActionType::Button
+	);
 
 	/**
 	 * Adds a tool bar stack button
diff --git a/Engine/Source/Runtime/Slate/Public/Framework/MultiBox/SToolBarComboButtonBlock.h b/Engine/Source/Runtime/Slate/Public/Framework/MultiBox/SToolBarComboButtonBlock.h
--- a/Engine/Source/Runtime/Slate/Public/Framework/MultiBox/SToolBarComboButtonBlock.h
+++ b/Engine/Source/Runtime/Slate/Public/Framework/MultiBox/SToolBarComboButtonBlock.h
@@ -23,15 +23,17 @@ public:
 	/**
 	 * Constructor
 	 *
-	 * @param	InAction				UI action that sets the enabled state for this combo button
-	 * @param	InMenuContentGenerator	Delegate that generates a widget for this combo button's menu content.
+	 * @param	InAction					UI action that sets the enabled state for this combo button
+	 * @param	InMenuContentGenerator		Delegate that generates a widget for this combo button's menu content.
 	 * Called when the menu is summoned.
-	 * @param	InLabel					Optional label for this combo button.
-	 * @param	InToolTip				Tool tip string (required!)
-	 * @param	InIcon					Optional icon to use for the tool bar image
-	 * @param	bInSimpleComboBox		If true, the icon and label won't be displayed
-	 * @param	InToolbarLabelOverride	Optional label to use when the block appears in a toolbar. If omitted, then the
+	 * @param	InLabel						Optional label for this combo button.
+	 * @param	InToolTip					Tool tip string (required!)
+	 * @param	InIcon						Optional icon to use for the tool bar image
+	 * @param	bInSimpleComboBox			If true, the icon and label won't be displayed
+	 * @param	InToolbarLabelOverride		Optional label to use when the block appears in a toolbar. If omitted, then the
 	 * label override or command name will be used instead.
+	 * @param	InPlacementOverride			Optional override placement to customize menu placement once opened via e.g. a toolbar menu button.
+	 * @param	InUserInterfaceActionType	Optional type of user interface action to use. If a command is bound, the user interface action type associated with the command will be used instead.
 	 */
 	FToolBarComboButtonBlock(
 		const FUIAction& InAction,
@@ -41,7 +43,8 @@ public:
 		const TAttribute<FSlateIcon>& InIcon = TAttribute<FSlateIcon>(),
 		bool bInSimpleComboBox = false,
 		TAttribute<FText> InToolbarLabelOverride = TAttribute<FText>(),
-		TAttribute<EMenuPlacement> InPlacementOverride = MenuPlacement_ComboBox
+		TAttribute<EMenuPlacement> InPlacementOverride = MenuPlacement_ComboBox,
+		const EUserInterfaceActionType InUserInterfaceActionType = EUserInterfaceActionType::Button
 	);
 
 	/** FMultiBlock interface */
@@ -95,6 +98,10 @@ private:
 	/** Controls the Labels visibility, defaults to GetIconVisibility if no override is provided */
 	TOptional< EVisibility > LabelVisibility;
 
+	/** In the case where a command is not bound, the user interface action type to use. If a command is bound, we
+		simply use the action type associated with that command. */
+	EUserInterfaceActionType UserInterfaceActionType;
+
 	/** If true, the icon and label won't be displayed */
 	bool bSimpleComboBox;
 
@@ -151,6 +158,22 @@ protected:
 	 */
 	SLATE_API TSharedRef<SWidget> OnGetMenuContent();
 
+	/**
+	 * Called by Slate when this tool bar button's button is clicked
+	 */
+	SLATE_API FReply OnClicked();
+
+	/**
+	 * Called by Slate when this tool bar check box button is toggled
+	 */
+	SLATE_API void OnCheckStateChanged(const ECheckBoxState NewCheckedState);
+
+	/**
+	 * Called by slate to determine if this button should appear checked
+	 *
+	 * @return ECheckBoxState::Checked if it should be checked, ECheckBoxState::Unchecked if not.
+	 */
+	SLATE_API ECheckBoxState GetCheckState() const;
 
 	/**
 	 * Called by Slate to determine if this button is enabled
