Remove empty space for combo buttons that don't have an icon set

diff --git a/Engine/Source/Runtime/Slate/Private/Framework/MultiBox/SToolBarComboButtonBlock.cpp b/Engine/Source/Runtime/Slate/Private/Framework/MultiBox/SToolBarComboButtonBlock.cpp
--- a/Engine/Source/Runtime/Slate/Private/Framework/MultiBox/SToolBarComboButtonBlock.cpp
+++ b/Engine/Source/Runtime/Slate/Private/Framework/MultiBox/SToolBarComboButtonBlock.cpp
@@ -101,29 +101,26 @@ void SToolBarComboButtonBlock::BuildMultiBlockWidget(const ISlateStyle* StyleSet
 	TSharedRef<SWidget> IconWidget = SNullWidget::NullWidget;
 	if (!ToolBarComboButtonBlock->bSimpleComboBox)
 	{
-		TSharedRef<SLayeredImage> ActualIconWidget =
-			SNew(SLayeredImage)
-			.ColorAndOpacity(this, &SToolBarComboButtonBlock::GetIconForegroundColor)
-			.Image(this, &SToolBarComboButtonBlock::GetIconBrush);
-
-		ActualIconWidget->AddLayer(TAttribute<const FSlateBrush*>(this, &SToolBarComboButtonBlock::GetOverlayIconBrush));
-
-		if (MultiBox->GetType() == EMultiBoxType::SlimHorizontalToolBar
-		|| MultiBox->GetType() == EMultiBoxType::SlimHorizontalUniformToolBar)
-		{
-			const FVector2f IconSize = ToolBarStyle.IconSize;
-
-			IconWidget =
-				SNew(SBox)
-				.WidthOverride(IconSize.X)
-				.HeightOverride(IconSize.Y)
-				[
-					ActualIconWidget
-				];
-		}
-		else
+		if (MultiBlock->HasIcon())
 		{
-			IconWidget = ActualIconWidget;
+			TSharedRef<SLayeredImage> ActualIconWidget =
+				SNew(SLayeredImage)
+					.ColorAndOpacity(this, &SToolBarComboButtonBlock::GetIconForegroundColor)
+					.Image(this, &SToolBarComboButtonBlock::GetIconBrush);
+
+			ActualIconWidget->AddLayer(TAttribute<const FSlateBrush*>(this, &SToolBarComboButtonBlock::GetOverlayIconBrush));
+
+			if (MultiBox->GetType() == EMultiBoxType::SlimHorizontalToolBar
+				|| MultiBox->GetType() == EMultiBoxType::SlimHorizontalUniformToolBar)
+			{
+				const FVector2f IconSize = ToolBarStyle.IconSize;
+
+				IconWidget = SNew(SBox).WidthOverride(IconSize.X).HeightOverride(IconSize.Y)[ActualIconWidget];
+			}
+			else
+			{
+				IconWidget = ActualIconWidget;
+			}
 		}
 
 		if (ToolBarComboButtonBlock->ToolbarLabelOverride.IsSet())
@@ -181,6 +178,7 @@ void SToolBarComboButtonBlock::BuildMultiBlockWidget(const ISlateStyle* StyleSet
 	}
 	else
 	{
+		// clang-format off
 		ButtonContent =
 			SNew(SVerticalBox)
 			// Icon image
@@ -191,9 +189,9 @@ void SToolBarComboButtonBlock::BuildMultiBlockWidget(const ISlateStyle* StyleSet
 				IconWidget
 			]
 			// Label text
-			+ SVerticalBox::Slot().AutoHeight()
-			.Padding(ToolBarStyle.LabelPadding)
+			+ SVerticalBox::Slot()
 			.AutoHeight()
+			.Padding(ToolBarStyle.LabelPadding)
 			.HAlign(HAlign_Center)	// Center the label text horizontally
 			[
 				SNew(STextBlock)
@@ -201,6 +199,7 @@ void SToolBarComboButtonBlock::BuildMultiBlockWidget(const ISlateStyle* StyleSet
 				.Text(Label)
 				.TextStyle(&ToolBarStyle.LabelStyle)
 			];
+		// clang-format on
 	}
 
 	EMultiBlockLocation::Type BlockLocation = GetMultiBlockLocation();
@@ -263,10 +262,9 @@ void SToolBarComboButtonBlock::BuildMultiBlockWidget(const ISlateStyle* StyleSet
 		FName ToggleBlockStyle = EMultiBlockLocation::ToName(ISlateStyle::Join(StyleName, ".ToggleButton"), BlockLocation);
 
 		const FCheckBoxStyle* CheckStyle;
-		if (!Icon.IsSet())
+		if (!MultiBlock->HasIcon())
 		{
 			CheckStyle = &FCoreStyle::Get().GetWidgetStyle<FCheckBoxStyle>("Checkbox");
-			IconWidget = SNullWidget::NullWidget;
 		}
 		else if (BlockLocation == EMultiBlockLocation::None)
 		{
