Add bool to allow or dissallow the faking of input touches
Gate the faking of touches based on the current input type
Enable / Disable fake touches when input type changes

#tests PIE, Client
[RN] Fixes conflicts with gamepad when using the UseMouseForTouch feature

diff --git a/Engine/Plugins/Runtime/CommonUI/Source/CommonInput/Private/CommonInputSubsystem.cpp b/Engine/Plugins/Runtime/CommonUI/Source/CommonInput/Private/CommonInputSubsystem.cpp
--- a/Engine/Plugins/Runtime/CommonUI/Source/CommonInput/Private/CommonInputSubsystem.cpp
+++ b/Engine/Plugins/Runtime/CommonUI/Source/CommonInput/Private/CommonInputSubsystem.cpp
@@ -318,9 +318,12 @@ void UCommonInputSubsystem::SetCurrentInputType(ECommonInputType NewInputType)
 					{
 						SlateApplication.UsePlatformCursorForCursorUser(bEnableGamepadPlatformCursor);
 					}
+					SlateApplication.SetGameAllowsFakingTouchEvents(false);
 					break;
 				case ECommonInputType::Touch:
 					UE_LOG(LogCommonInput, Log, TEXT("UCommonInputSubsystem::SetCurrentInputType(): Using Touch"));
+					SlateApplication.SetGameAllowsFakingTouchEvents(true);
+					SlateApplication.SetGameIsFakingTouchEvents(LocalPlayer && LocalPlayer->ViewportClient && LocalPlayer->ViewportClient->GetUseMouseForTouch());
 					break;
 				case ECommonInputType::MouseAndKeyboard:
 				default:				
@@ -329,6 +332,7 @@ void UCommonInputSubsystem::SetCurrentInputType(ECommonInputType NewInputType)
 					{
 						SlateApplication.UsePlatformCursorForCursorUser(true);
 					}
+					SlateApplication.SetGameAllowsFakingTouchEvents(false);
 					break;
 				}
 
diff --git a/Engine/Source/Runtime/Slate/Private/Framework/Application/SlateApplication.cpp b/Engine/Source/Runtime/Slate/Private/Framework/Application/SlateApplication.cpp
--- a/Engine/Source/Runtime/Slate/Private/Framework/Application/SlateApplication.cpp
+++ b/Engine/Source/Runtime/Slate/Private/Framework/Application/SlateApplication.cpp
@@ -824,6 +824,7 @@ FSlateApplication::FSlateApplication()
 	, bIsFakingTouch(FParse::Param(FCommandLine::Get(), TEXT("simmobile")) || FParse::Param(FCommandLine::Get(), TEXT("faketouches")))
 	, bIsGameFakingTouch( false )
 	, bIsFakingTouched( false )
+	, bAllowFakingTouch( true )
 	, bHandleDeviceInputWhenApplicationNotActive(false)
 	, bTouchFallbackToMouse( true )
 	, bSoftwareCursorAvailable( false )	
@@ -5003,7 +5004,7 @@ void FSlateApplication::SetGameIsFakingTouchEvents(const bool bIsFaking, FVector
 	// the only place this is not guarded is in FPIEPreviewDeviceModule::OnWindowReady()
 	if ( bIsGameFakingTouch != bIsFaking )
 	{
-		if (bIsFakingTouched && !bIsFaking && bIsGameFakingTouch && !bIsFakingTouch)
+		if (bAllowFakingTouch && bIsFakingTouched && !bIsFaking && bIsGameFakingTouch && !bIsFakingTouch)
 		{
 			OnTouchEnded((CursorLocation ? *CursorLocation : PlatformApplication->Cursor->GetPosition()), 0, FSlateApplicationBase::SlateAppPrimaryPlatformUser, IPlatformInputDeviceMapper::Get().GetDefaultInputDevice());
 		}
@@ -5012,9 +5013,18 @@ void FSlateApplication::SetGameIsFakingTouchEvents(const bool bIsFaking, FVector
 	}
 }
 
+void FSlateApplication::SetGameAllowsFakingTouchEvents(const bool bAllowFaking)
+{
+	bAllowFakingTouch = bAllowFaking;
+	if(!bAllowFaking && IsFakingTouchEvents())
+	{
+		SetGameIsFakingTouchEvents(bAllowFaking);
+	}
+}
+
 bool FSlateApplication::IsFakingTouchEvents() const
 {
-	return bIsFakingTouch || bIsGameFakingTouch;
+	return bAllowFakingTouch && (bIsFakingTouch || bIsGameFakingTouch);
 }
 
 bool FSlateApplication::OnMouseDown(const TSharedPtr< FGenericWindow >& PlatformWindow, const EMouseButtons::Type Button)
diff --git a/Engine/Source/Runtime/Slate/Public/Framework/Application/SlateApplication.h b/Engine/Source/Runtime/Slate/Public/Framework/Application/SlateApplication.h
--- a/Engine/Source/Runtime/Slate/Public/Framework/Application/SlateApplication.h
+++ b/Engine/Source/Runtime/Slate/Public/Framework/Application/SlateApplication.h
@@ -1444,6 +1444,9 @@ public:
 	/** Sets whether the application is treating mouse events as imitating touch events.  Optional CursorLocation can be supplied to override the platform's belief of where the cursor is */
 	SLATE_API void SetGameIsFakingTouchEvents(const bool bIsFaking, FVector2D* CursorLocation = nullptr);
 
+	/** Sets whether the application is allowed to fake touches or not */
+	SLATE_API void SetGameAllowsFakingTouchEvents(const bool bAllowFaking);
+
 	/** Sets the handler for otherwise unhandled key down events. This is used by the editor to provide a global action list, if the key was not consumed by any widget. */
 	SLATE_API void SetUnhandledKeyDownEventHandler( const FOnKeyEvent& NewHandler );
 
@@ -2040,6 +2043,9 @@ private:
 	/**For desktop platforms that the touch move event be called when this variable is true */
 	bool bIsFakingTouched;
 
+	/** Restricts or allows the faking of touches. */
+	bool bAllowFakingTouch;
+
 	/** Force Mouse Pointer Capture to always occur even when the application is not active or widget is not a virtual window */
 	bool bHandleDeviceInputWhenApplicationNotActive;
 
