TabManager : fix a bug where we weren't setting the parent window when spawning a tab leading to a crash.

When opening a nomad tab window ( e.g. Mass Debugger ) with one of it's panel tabs ( e.g. Archetypes ) in a floating window, the parent window of the panel tab ( Archetypes ) was not being set.  This led to 2 issues:
1) Clicking on either the main UE window or the Mass Debugger would not bring the Archetypes window in front of other applications
2) Closing UE would lead to a crash where it erroneously tried to activate the Archetypes window during shutdown.

diff --git a/Engine/Source/Runtime/Slate/Private/Framework/Docking/TabManager.cpp b/Engine/Source/Runtime/Slate/Private/Framework/Docking/TabManager.cpp
--- a/Engine/Source/Runtime/Slate/Private/Framework/Docking/TabManager.cpp
+++ b/Engine/Source/Runtime/Slate/Private/Framework/Docking/TabManager.cpp
@@ -1449,7 +1449,7 @@ TSharedPtr<SDockTab> FTabManager::InvokeTab_Internal(const FTabId& TabId, bool b
 
 	if (StackToSpawnIn.IsValid())
 	{
-		const TSharedPtr<SDockTab> NewTab = SpawnTab(TabId, TSharedPtr<SWindow>());
+		const TSharedPtr<SDockTab> NewTab = SpawnTab(TabId, GetPrivateApi().GetParentWindow());
 
 		if (NewTab.IsValid())
 		{
