Move template implementation of SSpinBox into separate object file
* Primary motivator is to improve developer iteration while working on SSpinBox.h
* Code bloat/compile time reduction is a bonus
* No functional changes

diff --git a/Engine/Source/Runtime/Slate/Private/Widgets/Input/SSpinBox.cpp b/Engine/Source/Runtime/Slate/Private/Widgets/Input/SSpinBox.cpp
--- a/Engine/Source/Runtime/Slate/Private/Widgets/Input/SSpinBox.cpp
+++ b/Engine/Source/Runtime/Slate/Private/Widgets/Input/SSpinBox.cpp
@@ -18,4 +18,1071 @@ float SpinBoxComputeExponentSliderFraction(float FractionFilled, float StartFrac
 	return FractionFilled;
 }
 
+template <typename NumericType>
+SSpinBox<NumericType>::SSpinBox()
+{
+}
+
+template <typename NumericType>
+SSpinBox<NumericType>::~SSpinBox()
+{
+	if (bDragging || PointerDraggingSliderIndex != INDEX_NONE)
+	{
+		CancelMouseCapture();
+	}
+}
+
+template <typename NumericType>
+void SSpinBox<NumericType>::Construct(const FArguments& InArgs)
+{
+	check(InArgs._Style);
+
+	Style = InArgs._Style;
+
+	SetForegroundColor(InArgs._Style->ForegroundColor);
+	Interface = InArgs._TypeInterface.IsValid() ? InArgs._TypeInterface : MakeShareable(new TDefaultNumericTypeInterface<NumericType>);
+
+	if (InArgs._TypeInterface.IsValid() && Interface->GetOnSettingChanged())
+	{
+		Interface->GetOnSettingChanged()->AddSP(this, &SSpinBox::ResetCachedValueString);
+	}
+
+	ValueAttribute = InArgs._Value;
+	OnValueChanged = InArgs._OnValueChanged;
+	OnValueCommitted = InArgs._OnValueCommitted;
+	OnBeginSliderMovement = InArgs._OnBeginSliderMovement;
+	OnEndSliderMovement = InArgs._OnEndSliderMovement;
+	MinDesiredWidth = InArgs._MinDesiredWidth;
+
+	MinValue = InArgs._MinValue;
+	MaxValue = InArgs._MaxValue;
+	MinSliderValue = (InArgs._MinSliderValue.Get().IsSet()) ? InArgs._MinSliderValue : MinValue;
+	MaxSliderValue = (InArgs._MaxSliderValue.Get().IsSet()) ? InArgs._MaxSliderValue : MaxValue;
+
+	MinFractionalDigits = (InArgs._MinFractionalDigits.Get().IsSet()) ? InArgs._MinFractionalDigits : DefaultMinFractionalDigits;
+	MaxFractionalDigits = (InArgs._MaxFractionalDigits.Get().IsSet()) ? InArgs._MaxFractionalDigits : DefaultMaxFractionalDigits;
+	SetMaxFractionalDigits(MaxFractionalDigits);
+	SetMinFractionalDigits(MinFractionalDigits);
+
+	AlwaysUsesDeltaSnap = InArgs._AlwaysUsesDeltaSnap;
+	EnableSlider = InArgs._EnableSlider;
+
+	SupportDynamicSliderMaxValue = InArgs._SupportDynamicSliderMaxValue;
+	SupportDynamicSliderMinValue = InArgs._SupportDynamicSliderMinValue;
+	OnDynamicSliderMaxValueChanged = InArgs._OnDynamicSliderMaxValueChanged;
+	OnDynamicSliderMinValueChanged = InArgs._OnDynamicSliderMinValueChanged;
+
+	bEnableWheel = InArgs._EnableWheel;
+	bBroadcastValueChangesPerKey = InArgs._BroadcastValueChangesPerKey;
+	WheelStep = InArgs._WheelStep;
+
+	bPreventThrottling = InArgs._PreventThrottling;
+
+	CachedExternalValue = ValueAttribute.Get();
+	CachedValueString = Interface->ToString(CachedExternalValue);
+	bCachedValueStringDirty = false;
+
+	InternalValue = (double)CachedExternalValue;
+
+	if (SupportDynamicSliderMaxValue.Get() && CachedExternalValue > GetMaxSliderValue())
+	{
+		ApplySliderMaxValueChanged(float(CachedExternalValue - GetMaxSliderValue()), true);
+	}
+	else if (SupportDynamicSliderMinValue.Get() && CachedExternalValue < GetMinSliderValue())
+	{
+		ApplySliderMinValueChanged(float(CachedExternalValue - GetMinSliderValue()), true);
+	}
+
+	UpdateIsSpinRangeUnlimited();
+
+	SliderExponent = InArgs._SliderExponent;
+
+	SliderExponentNeutralValue = InArgs._SliderExponentNeutralValue;
+
+	DistanceDragged = 0.0f;
+	PreDragValue = NumericType();
+
+	Delta = InArgs._Delta;
+	ShiftMultiplier = InArgs._ShiftMultiplier;
+	CtrlMultiplier = InArgs._CtrlMultiplier;
+	LinearDeltaSensitivity = InArgs._LinearDeltaSensitivity;
+
+	BackgroundHoveredBrush = &InArgs._Style->HoveredBackgroundBrush;
+	BackgroundBrush = &InArgs._Style->BackgroundBrush;
+	BackgroundActiveBrush = InArgs._Style->ActiveBackgroundBrush.IsSet() ? &InArgs._Style->ActiveBackgroundBrush : BackgroundHoveredBrush;
+
+	ActiveFillBrush = &InArgs._Style->ActiveFillBrush;
+	HoveredFillBrush = InArgs._Style->HoveredFillBrush.IsSet() ? &InArgs._Style->HoveredFillBrush : ActiveFillBrush;
+	InactiveFillBrush = &InArgs._Style->InactiveFillBrush;
+
+	const FMargin& TextMargin = InArgs._Style->TextPadding;
+
+	bDragging = false;
+	PointerDraggingSliderIndex = INDEX_NONE;
+
+	bIsTextChanging = false;
+
+	this->ChildSlot
+	    .Padding(InArgs._ContentPadding)
+	[
+		SNew(SHorizontalBox)
+
+		+ SHorizontalBox::Slot()
+		  .FillWidth(1.0f)
+		  .Padding(TextMargin)
+		  .HAlign(HAlign_Fill)
+		  .VAlign(VAlign_Center)
+		[
+			SAssignNew(TextBlock, STextBlock)
+				.Font(InArgs._Font)
+				.Text(this, &SSpinBox<NumericType>::GetValueAsText)
+				.MinDesiredWidth(this, &SSpinBox<NumericType>::GetTextMinDesiredWidth)
+				.Justification(InArgs._Justification)
+		]
+		+ SHorizontalBox::Slot()
+		  .FillWidth(1.0f)
+		  .Padding(TextMargin)
+		  .HAlign(HAlign_Fill)
+		  .VAlign(VAlign_Center)
+		[
+			SAssignNew(EditableText, SEditableText)
+				.Visibility(EVisibility::Collapsed)
+				.Font(InArgs._Font)
+				.SelectAllTextWhenFocused(true)
+				.Text(this, &SSpinBox<NumericType>::GetValueAsText)
+				.OnIsTypedCharValid(this, &SSpinBox<NumericType>::IsCharacterValid)
+				.OnTextChanged(this, &SSpinBox<NumericType>::TextField_OnTextChanged)
+				.OnTextCommitted(this, &SSpinBox<NumericType>::TextField_OnTextCommitted)
+				.ClearKeyboardFocusOnCommit(InArgs._ClearKeyboardFocusOnCommit)
+				.SelectAllTextOnCommit(InArgs._SelectAllTextOnCommit)
+				.MinDesiredWidth(this, &SSpinBox<NumericType>::GetTextMinDesiredWidth)
+				.VirtualKeyboardType(InArgs._KeyboardType)
+				.Justification(InArgs._Justification)
+				.VirtualKeyboardTrigger(EVirtualKeyboardTrigger::OnAllFocusEvents)
+				.ContextMenuExtender(InArgs._ContextMenuExtender)
+		]
+		+ SHorizontalBox::Slot()
+		  .AutoWidth()
+		  .HAlign(HAlign_Fill)
+		  .VAlign(VAlign_Center)
+		[
+			SNew(SImage)
+				.Image(&InArgs._Style->ArrowsImage)
+				.ColorAndOpacity(FSlateColor::UseForeground())
+		]
+	];
+}
+
+template <typename NumericType>
+int32 SSpinBox<NumericType>::OnPaint(const FPaintArgs& Args, const FGeometry& AllottedGeometry, const FSlateRect& MyCullingRect,
+	FSlateWindowElementList& OutDrawElements, int32 LayerId, const FWidgetStyle& InWidgetStyle, bool bParentEnabled) const
+{
+
+	const bool bActiveFeedback = bDragging || IsInTextMode();
+
+	const FSlateBrush* BackgroundImage = bActiveFeedback ? BackgroundActiveBrush : IsHovered() ?
+		                                     BackgroundHoveredBrush :
+		                                     BackgroundBrush;
+
+	const FSlateBrush* FillImage = bActiveFeedback ?
+		                               ActiveFillBrush : IsHovered() ? HoveredFillBrush : 
+		                               InactiveFillBrush;
+
+	const int32 BackgroundLayer = LayerId;
+
+	const bool bEnabled = ShouldBeEnabled(bParentEnabled);
+	const ESlateDrawEffect DrawEffects = bEnabled ? ESlateDrawEffect::None : ESlateDrawEffect::DisabledEffect;
+
+	FSlateDrawElement::MakeBox(
+		OutDrawElements,
+		BackgroundLayer,
+		AllottedGeometry.ToPaintGeometry(),
+		BackgroundImage,
+		DrawEffects,
+		BackgroundImage->GetTint(InWidgetStyle) * InWidgetStyle.GetColorAndOpacityTint()
+	);
+
+	const int32 FilledLayer = BackgroundLayer + 1;
+
+	//if there is a spin range limit, draw the filler bar
+	if (!bUnlimitedSpinRange)
+	{
+		NumericType Value = ValueAttribute.Get();
+		NumericType CurrentDelta = Delta.Get();
+		if (CurrentDelta != NumericType())
+		{
+			Value = FMath::GridSnap(Value, CurrentDelta); // snap value to nearest Delta
+		}
+
+		float FractionFilled = Fraction((double)Value, (double)GetMinSliderValue(), (double)GetMaxSliderValue());
+		const float CachedSliderExponent = SliderExponent.Get();
+		if (!FMath::IsNearlyEqual(CachedSliderExponent, 1.f))
+		{
+			if (SliderExponentNeutralValue.IsSet() && SliderExponentNeutralValue.Get() > GetMinSliderValue() && SliderExponentNeutralValue.Get() < GetMaxSliderValue())
+			{
+				//Compute a log curve on both side of the neutral value
+				float StartFractionFilled = Fraction((double)SliderExponentNeutralValue.Get(), (double)GetMinSliderValue(), (double)GetMaxSliderValue());
+				FractionFilled = SpinBoxComputeExponentSliderFraction(FractionFilled, StartFractionFilled, CachedSliderExponent);
+			}
+			else
+			{
+				FractionFilled = 1.0f - FMath::Pow(1.0f - FractionFilled, CachedSliderExponent);
+			}
+		}
+		const FVector2D FillSize(AllottedGeometry.GetLocalSize().X * FractionFilled, AllottedGeometry.GetLocalSize().Y);
+
+		if (!IsInTextMode())
+		{
+			FSlateDrawElement::MakeBox(
+				OutDrawElements,
+				FilledLayer,
+				AllottedGeometry.ToPaintGeometry(FillSize - FVector2D(Style->InsetPadding.GetTotalSpaceAlong<Orient_Horizontal>(), Style->InsetPadding.GetTotalSpaceAlong<Orient_Vertical>()), FSlateLayoutTransform(Style->InsetPadding.GetTopLeft())),
+				FillImage,
+				DrawEffects,
+				FillImage->GetTint(InWidgetStyle) * InWidgetStyle.GetColorAndOpacityTint()
+			);
+		}
+	}
+
+	return FMath::Max(FilledLayer, SCompoundWidget::OnPaint(Args, AllottedGeometry, MyCullingRect, OutDrawElements, FilledLayer, InWidgetStyle, bEnabled));
+}
+
+template <typename NumericType>
+const bool SSpinBox<NumericType>::CommitWithMultiplier(const FPointerEvent& MouseEvent)
+{
+	return MouseEvent.IsShiftDown() || MouseEvent.IsControlDown();
+}
+
+template <typename NumericType>
+FReply SSpinBox<NumericType>::OnMouseButtonDown(const FGeometry& MyGeometry, const FPointerEvent& MouseEvent)
+{
+	if (MouseEvent.GetEffectingButton() == EKeys::LeftMouseButton && PointerDraggingSliderIndex == INDEX_NONE)
+	{
+		DistanceDragged = 0.f;
+		PreDragValue = ValueAttribute.Get();
+		InternalValue = (double)PreDragValue;
+		PointerDraggingSliderIndex = MouseEvent.GetPointerIndex();
+		CachedMousePosition = MouseEvent.GetScreenSpacePosition().IntPoint();
+
+		FReply ReturnReply = FReply::Handled().CaptureMouse(SharedThis(this)).UseHighPrecisionMouseMovement(SharedThis(this)).SetUserFocus(SharedThis(this), EFocusCause::Mouse);
+		if (bPreventThrottling)
+		{
+			ReturnReply.PreventThrottling();
+		}
+		return ReturnReply;
+	}
+	else
+	{
+
+		return FReply::Unhandled();
+	}
+}
+
+template <typename NumericType>
+FReply SSpinBox<NumericType>::OnMouseButtonUp(const FGeometry& MyGeometry, const FPointerEvent& MouseEvent)
+{
+	if (MouseEvent.GetEffectingButton() == EKeys::LeftMouseButton && PointerDraggingSliderIndex == MouseEvent.GetPointerIndex())
+	{
+		if (!this->HasMouseCapture())
+		{
+			// Lost Capture - ensure reset
+			bDragging = false;
+			PointerDraggingSliderIndex = INDEX_NONE;
+
+			return FReply::Unhandled();
+		}
+
+		if (bDragging)
+		{
+			NumericType CurrentDelta = Delta.Get();
+			if (CurrentDelta != NumericType() && !CommitWithMultiplier(MouseEvent))
+			{
+				InternalValue = FMath::GridSnap(InternalValue, (double)CurrentDelta);
+			}
+
+			const NumericType CurrentValue = RoundIfIntegerValue(InternalValue);
+			NotifyValueCommitted(CurrentValue);
+		}
+
+		bDragging = false;
+		PointerDraggingSliderIndex = INDEX_NONE;
+
+		FReply Reply = FReply::Handled().ReleaseMouseCapture();
+
+		if (!MouseEvent.IsTouchEvent())
+		{
+			Reply.SetMousePos(CachedMousePosition);
+		}
+
+		if (DistanceDragged < FSlateApplication::Get().GetDragTriggerDistance())
+		{
+			EnterTextMode();
+			Reply.SetUserFocus(EditableText.ToSharedRef(), EFocusCause::SetDirectly);
+		}
+
+		return Reply;
+	}
+
+	return FReply::Unhandled();
+}
+
+template <typename NumericType>
+void SSpinBox<NumericType>::ApplySliderMaxValueChanged(float SliderDeltaToAdd, bool UpdateOnlyIfHigher)
+{
+	check(SupportDynamicSliderMaxValue.Get());
+
+	NumericType NewMaxSliderValue = std::numeric_limits<NumericType>::min();
+
+	if (MaxSliderValue.IsSet() && MaxSliderValue.Get().IsSet())
+	{
+		NewMaxSliderValue = GetMaxSliderValue();
+
+		if ((NewMaxSliderValue + (NumericType)SliderDeltaToAdd > GetMaxSliderValue() && UpdateOnlyIfHigher) || !UpdateOnlyIfHigher)
+		{
+			NewMaxSliderValue += (NumericType)SliderDeltaToAdd;
+
+			if (!MaxSliderValue.IsBound()) // simple value so we can update it without breaking the mechanic otherwise it must be handle by the callback implementer
+			{
+				SetMaxSliderValue(NewMaxSliderValue);
+			}
+		}
+	}
+
+	if (OnDynamicSliderMaxValueChanged.IsBound())
+	{
+		OnDynamicSliderMaxValueChanged.Execute(NewMaxSliderValue, TWeakPtr<SWidget>(AsShared()), true, UpdateOnlyIfHigher);
+	}
+}
+
+template <typename NumericType>
+void SSpinBox<NumericType>::ApplySliderMinValueChanged(float SliderDeltaToAdd, bool UpdateOnlyIfLower)
+{
+	check(SupportDynamicSliderMaxValue.Get());
+
+	NumericType NewMinSliderValue = std::numeric_limits<NumericType>::min();
+
+	if (MinSliderValue.IsSet() && MinSliderValue.Get().IsSet())
+	{
+		NewMinSliderValue = GetMinSliderValue();
+
+		if ((NewMinSliderValue + (NumericType)SliderDeltaToAdd < GetMinSliderValue() && UpdateOnlyIfLower) || !UpdateOnlyIfLower)
+		{
+			NewMinSliderValue += (NumericType)SliderDeltaToAdd;
+
+			if (!MinSliderValue.IsBound()) // simple value so we can update it without breaking the mechanic otherwise it must be handle by the callback implementer
+			{
+				SetMinSliderValue(NewMinSliderValue);
+			}
+		}
+	}
+
+	if (OnDynamicSliderMinValueChanged.IsBound())
+	{
+		OnDynamicSliderMinValueChanged.Execute(NewMinSliderValue, TWeakPtr<SWidget>(AsShared()), true, UpdateOnlyIfLower);
+	}
+}
+
+template <typename NumericType>
+FReply SSpinBox<NumericType>::OnMouseMove(const FGeometry& MyGeometry, const FPointerEvent& MouseEvent)
+{
+	const bool bEnableSlider = GetEnableSlider();
+	if (PointerDraggingSliderIndex == MouseEvent.GetPointerIndex() && bEnableSlider)
+	{
+		if (!this->HasMouseCapture())
+		{
+			// Lost the mouse capture - ensure reset
+			bDragging = false;
+			PointerDraggingSliderIndex = INDEX_NONE;
+
+			return FReply::Unhandled();
+		}
+
+		if (!bDragging)
+		{
+			DistanceDragged += (float)FMath::Abs(MouseEvent.GetCursorDelta().X);
+			if (DistanceDragged > FSlateApplication::Get().GetDragTriggerDistance())
+			{
+				ExitTextMode();
+				bDragging = true;
+				OnBeginSliderMovement.ExecuteIfBound();
+			}
+
+			// Cache the mouse, even if not dragging cache it.
+			CachedMousePosition = MouseEvent.GetScreenSpacePosition().IntPoint();
+		}
+		else
+		{
+			double NewValue = 0.0;
+
+			// Increments the spin based on delta mouse movement.
+
+			// A minimum slider width to use for calculating deltas in the slider-range space
+			const float MinSliderWidth = 100.f;
+			float SliderWidthInSlateUnits = FMath::Max((float)MyGeometry.GetDrawSize().X, MinSliderWidth);
+
+			if (MouseEvent.IsAltDown())
+			{
+				float DeltaToAdd = (float)MouseEvent.GetCursorDelta().X / SliderWidthInSlateUnits;
+
+				if (SupportDynamicSliderMaxValue.Get() && (NumericType)InternalValue == GetMaxSliderValue())
+				{
+					ApplySliderMaxValueChanged(DeltaToAdd, false);
+				}
+				else if (SupportDynamicSliderMinValue.Get() && (NumericType)InternalValue == GetMinSliderValue())
+				{
+					ApplySliderMinValueChanged(DeltaToAdd, false);
+				}
+			}
+
+			ECommitMethod CommitMethod = CommittedViaSpin;
+				
+			const bool bIsSmallStep = (GetMaxSliderValue() - GetMinSliderValue()) <= 10.0;
+			double Step = bIsSmallStep ? 0.1 : 1.0;
+
+			if (MouseEvent.IsControlDown())
+			{
+				Step *= CtrlMultiplier.Get();
+				CommitMethod = CommittedViaSpinMultiplier;
+			}
+			else if (MouseEvent.IsShiftDown())
+			{
+				Step *= ShiftMultiplier.Get();
+				CommitMethod = CommittedViaSpinMultiplier;
+			}
+
+			//if we have a range to draw in
+			if (!bUnlimitedSpinRange)
+			{
+				bool HasValidExponentNeutralValue = SliderExponentNeutralValue.IsSet() && SliderExponentNeutralValue.Get() > GetMinSliderValue() && SliderExponentNeutralValue.Get() < GetMaxSliderValue();
+
+				const float CachedSliderExponent = SliderExponent.Get();
+				// The amount currently filled in the spinbox, needs to be calculated to do deltas correctly.
+				float FractionFilled = Fraction(InternalValue, (double)GetMinSliderValue(), (double)GetMaxSliderValue());
+
+				if (!FMath::IsNearlyEqual(CachedSliderExponent, 1.0f))
+				{
+					if (HasValidExponentNeutralValue)
+					{
+						//Compute a log curve on both side of the neutral value
+						float StartFractionFilled = Fraction((double)SliderExponentNeutralValue.Get(), (double)GetMinSliderValue(), (double)GetMaxSliderValue());
+						FractionFilled = SpinBoxComputeExponentSliderFraction(FractionFilled, StartFractionFilled, CachedSliderExponent);
+					}
+					else
+					{
+						FractionFilled = 1.0f - FMath::Pow(1.0f - FractionFilled, CachedSliderExponent);
+					}
+				}
+				FractionFilled *= SliderWidthInSlateUnits;
+
+				// Now add the delta to the fraction filled, this causes the spin.
+				FractionFilled += (float)(MouseEvent.GetCursorDelta().X * Step);
+
+				// Clamp the fraction to be within the bounds of the geometry.
+				FractionFilled = FMath::Clamp(FractionFilled, 0.0f, SliderWidthInSlateUnits);
+
+				// Convert the fraction filled to a percent.
+				float Percent = FMath::Clamp(FractionFilled / SliderWidthInSlateUnits, 0.0f, 1.0f);
+				if (!FMath::IsNearlyEqual(CachedSliderExponent, 1.0f))
+				{
+					// Have to convert the percent to the proper value due to the exponent component to the spin.
+					if (HasValidExponentNeutralValue)
+					{
+						//Compute a log curve on both side of the neutral value
+						float StartFractionFilled = Fraction(SliderExponentNeutralValue.Get(), GetMinSliderValue(), GetMaxSliderValue());
+						Percent = SpinBoxComputeExponentSliderFraction(Percent, StartFractionFilled, 1.0f / CachedSliderExponent);
+					}
+					else
+					{
+						Percent = 1.0f - FMath::Pow(1.0f - Percent, 1.0f / CachedSliderExponent);
+					}
+
+
+				}
+
+				NewValue = FMath::LerpStable<double>((double)GetMinSliderValue(), (double)GetMaxSliderValue(), Percent);
+			}
+			else
+			{
+				// If this control has a specified delta and sensitivity then we use that instead of the current value for determining how much to change.
+				const double Sign = (MouseEvent.GetCursorDelta().X > 0) ? 1.0 : -1.0;
+
+				if (LinearDeltaSensitivity.IsSet() && LinearDeltaSensitivity.Get() != 0 && Delta.IsSet() && Delta.Get() > 0)
+				{
+					const double MouseDelta = FMath::Abs(MouseEvent.GetCursorDelta().X / (float)LinearDeltaSensitivity.Get());
+					NewValue = InternalValue + (Sign * MouseDelta * FMath::Pow((double)Delta.Get(), (double)SliderExponent.Get())) * Step;
+				}
+				else
+				{
+					const double MouseDelta = FMath::Abs(MouseEvent.GetCursorDelta().X / SliderWidthInSlateUnits);
+					const double CurrentValue = FMath::Clamp<double>(FMath::Abs(InternalValue), 1.0, (double)std::numeric_limits<NumericType>::max());
+					NewValue = InternalValue + (Sign * MouseDelta * FMath::Pow((double)CurrentValue, (double)SliderExponent.Get())) * Step;
+				}
+			}
+
+			NumericType RoundedNewValue = RoundIfIntegerValue(NewValue);
+			CommitValue(RoundedNewValue, NewValue, CommitMethod, ETextCommit::OnEnter);
+		}
+
+		return FReply::Handled();
+	}
+
+	return FReply::Unhandled();
+}
+
+template <typename NumericType>
+FReply SSpinBox<NumericType>::OnMouseWheel(const FGeometry& MyGeometry, const FPointerEvent& MouseEvent)
+{
+	if (bEnableWheel && PointerDraggingSliderIndex == INDEX_NONE && HasKeyboardFocus())
+	{
+		// If there is no WheelStep defined, we use 1.0 (Or 0.1 if slider range is <= 10)
+		constexpr bool bIsIntegral = TIsIntegral<NumericType>::Value;
+		const bool bIsSmallStep = !bIsIntegral && (GetMaxSliderValue() - GetMinSliderValue()) <= 10.0;
+		double Step = WheelStep.IsSet() && WheelStep.Get().IsSet() ? WheelStep.Get().GetValue() : (bIsSmallStep ? 0.1 : 1.0);
+
+		if (MouseEvent.IsControlDown())
+		{
+			// If no value is set for WheelSmallStep, we use the DefaultStep multiplied by the CtrlMultiplier
+			Step *= CtrlMultiplier.Get();
+		}
+		else if (MouseEvent.IsShiftDown())
+		{
+			// If no value is set for WheelBigStep, we use the DefaultStep multiplied by the ShiftMultiplier
+			Step *= ShiftMultiplier.Get();
+		}
+
+		const double Sign = (MouseEvent.GetWheelDelta() > 0) ? 1.0 : -1.0;
+		const double NewValue = InternalValue + (Sign * Step);
+		const NumericType RoundedNewValue = RoundIfIntegerValue(NewValue);
+
+		// First SetText is to update the value before calling CommitValue. Otherwise, when the text lose
+		// focus from the CommitValue, it will override the value we just committed.
+		// The second SetText is to update the text to the InternalValue since it could have been clamped.
+		EditableText->SetText(FText::FromString(Interface->ToString((NumericType)NewValue)));
+		CommitValue(RoundedNewValue, NewValue, CommittedViaSpin, ETextCommit::OnEnter);
+		EditableText->SetText(FText::FromString(Interface->ToString((NumericType)InternalValue)));
+
+		return FReply::Handled();
+	}
+
+	return FReply::Unhandled();
+}
+
+template <typename NumericType>
+FCursorReply SSpinBox<NumericType>::OnCursorQuery(const FGeometry& MyGeometry, const FPointerEvent& CursorEvent) const
+{
+	const bool bEnableSlider = GetEnableSlider();
+
+	if (!bEnableSlider)
+	{
+		return FCursorReply::Cursor(EMouseCursor::Default);
+	}
+
+	return bDragging ?
+		       FCursorReply::Cursor(EMouseCursor::None) :
+		       FCursorReply::Cursor(EMouseCursor::ResizeLeftRight);
+}
+
+template <typename NumericType>
+bool SSpinBox<NumericType>::SupportsKeyboardFocus() const
+{
+	// SSpinBox is focusable.
+	return true;
+}
+
+template <typename NumericType>
+FReply SSpinBox<NumericType>::OnFocusReceived(const FGeometry& MyGeometry, const FFocusEvent& InFocusEvent)
+{
+	if (!bDragging && (InFocusEvent.GetCause() == EFocusCause::Navigation || InFocusEvent.GetCause() == EFocusCause::SetDirectly))
+	{
+		EnterTextMode();
+		return FReply::Handled().SetUserFocus(EditableText.ToSharedRef(), InFocusEvent.GetCause());
+	}
+	else
+	{
+		return FReply::Unhandled();
+	}
+}
+
+template <typename NumericType>
+FReply SSpinBox<NumericType>::OnKeyDown(const FGeometry& MyGeometry, const FKeyEvent& InKeyEvent)
+{
+	const FKey Key = InKeyEvent.GetKey();
+	if (Key == EKeys::Escape && HasMouseCapture())
+	{
+		CancelMouseCapture();
+		return FReply::Handled().ReleaseMouseCapture().SetMousePos(CachedMousePosition);
+	}
+	else if (Key == EKeys::Up || Key == EKeys::Right)
+	{
+		const NumericType LocalValueAttribute = ValueAttribute.Get();
+		const NumericType LocalDelta = Delta.Get();
+		InternalValue = (double)LocalValueAttribute;
+		CommitValue(LocalValueAttribute + LocalDelta, InternalValue + (double)LocalDelta, CommittedViaArrowKey, ETextCommit::OnEnter);
+		ExitTextMode();
+		return FReply::Handled();
+	}
+	else if (Key == EKeys::Down || Key == EKeys::Left)
+	{
+		const NumericType LocalValueAttribute = ValueAttribute.Get();
+		const NumericType LocalDelta = Delta.Get();
+		InternalValue = (double)LocalValueAttribute;
+		CommitValue(LocalValueAttribute - LocalDelta, InternalValue + (double)LocalDelta, CommittedViaArrowKey, ETextCommit::OnEnter);
+		ExitTextMode();
+		return FReply::Handled();
+	}
+	else if (Key == EKeys::Enter)
+	{
+		InternalValue = (double)ValueAttribute.Get();
+		EnterTextMode();
+		return FReply::Handled().SetUserFocus(EditableText.ToSharedRef(), EFocusCause::Navigation);
+	}
+	else
+	{
+		return FReply::Unhandled();
+	}
+}
+
+template <typename NumericType>
+bool SSpinBox<NumericType>::HasKeyboardFocus() const
+{
+	// The spinbox is considered focused when we are typing it text.
+	return SCompoundWidget::HasKeyboardFocus() || (EditableText.IsValid() && EditableText->HasKeyboardFocus());
+}
+
+template <typename NumericType>
+TAttribute<NumericType> SSpinBox<NumericType>::GetValueAttribute() const
+{
+	return ValueAttribute;
+}
+
+template <typename NumericType>
+NumericType SSpinBox<NumericType>::GetValue() const
+{
+	return ValueAttribute.Get();
+}
+
+template <typename NumericType>
+void SSpinBox<NumericType>::SetValue(const TAttribute<NumericType>& InValueAttribute)
+{
+	ValueAttribute = InValueAttribute;
+	const NumericType LocalValueAttribute = ValueAttribute.Get();
+	CommitValue(LocalValueAttribute, (double)LocalValueAttribute, ECommitMethod::CommittedViaCode, ETextCommit::Default);
+}
+
+template <typename NumericType>
+NumericType SSpinBox<NumericType>::GetMinValue() const
+{
+	return MinValue.Get().Get(std::numeric_limits<NumericType>::lowest());
+}
+
+template <typename NumericType>
+void SSpinBox<NumericType>::SetMinValue(const TAttribute<TOptional<NumericType>>& InMinValue)
+{
+	MinValue = InMinValue;
+	UpdateIsSpinRangeUnlimited();
+}
+
+template <typename NumericType>
+NumericType SSpinBox<NumericType>::GetMaxValue() const
+{ return MaxValue.Get().Get(std::numeric_limits<NumericType>::max()); }
+
+template <typename NumericType>
+void SSpinBox<NumericType>::SetMaxValue(const TAttribute<TOptional<NumericType>>& InMaxValue)
+{
+	MaxValue = InMaxValue;
+	UpdateIsSpinRangeUnlimited();
+}
+
+template <typename NumericType>
+bool SSpinBox<NumericType>::IsMinSliderValueBound() const
+{
+	return MinSliderValue.IsBound();
+}
+
+template <typename NumericType>
+NumericType SSpinBox<NumericType>::GetMinSliderValue() const
+{
+	return MinSliderValue.Get().Get(std::numeric_limits<NumericType>::lowest());
+}
+
+template <typename NumericType>
+void SSpinBox<NumericType>::SetMinSliderValue(const TAttribute<TOptional<NumericType>>& InMinSliderValue)
+{
+	MinSliderValue = (InMinSliderValue.Get().IsSet()) ? InMinSliderValue : MinValue;
+	UpdateIsSpinRangeUnlimited();
+}
+
+template <typename NumericType>
+bool SSpinBox<NumericType>::IsMaxSliderValueBound() const
+{
+	return MaxSliderValue.IsBound();
+}
+
+template <typename NumericType>
+NumericType SSpinBox<NumericType>::GetMaxSliderValue() const
+{
+	return MaxSliderValue.Get().Get(std::numeric_limits<NumericType>::max());
+}
+
+template <typename NumericType>
+void SSpinBox<NumericType>::SetMaxSliderValue(const TAttribute<TOptional<NumericType>>& InMaxSliderValue)
+{
+	MaxSliderValue = (InMaxSliderValue.Get().IsSet()) ? InMaxSliderValue : MaxValue;;
+	UpdateIsSpinRangeUnlimited();
+}
+
+template <typename NumericType>
+int32 SSpinBox<NumericType>::GetMinFractionalDigits() const
+{ return Interface->GetMinFractionalDigits(); }
+
+template <typename NumericType>
+void SSpinBox<NumericType>::SetMinFractionalDigits(const TAttribute<TOptional<int32>>& InMinFractionalDigits)
+{
+	Interface->SetMinFractionalDigits((InMinFractionalDigits.Get().IsSet()) ? InMinFractionalDigits.Get() : MinFractionalDigits);
+	bCachedValueStringDirty = true;
+}
+
+template <typename NumericType>
+int32 SSpinBox<NumericType>::GetMaxFractionalDigits() const
+{ return Interface->GetMaxFractionalDigits(); }
+
+template <typename NumericType>
+void SSpinBox<NumericType>::SetMaxFractionalDigits(const TAttribute<TOptional<int32>>& InMaxFractionalDigits)
+{
+	Interface->SetMaxFractionalDigits((InMaxFractionalDigits.Get().IsSet()) ? InMaxFractionalDigits.Get() : MaxFractionalDigits);
+	bCachedValueStringDirty = true;
+}
+
+template <typename NumericType>
+bool SSpinBox<NumericType>::GetAlwaysUsesDeltaSnap() const
+{ return AlwaysUsesDeltaSnap.Get(); }
+
+template <typename NumericType>
+void SSpinBox<NumericType>::SetAlwaysUsesDeltaSnap(bool bNewValue)
+{ AlwaysUsesDeltaSnap.Set(bNewValue); }
+
+template <typename NumericType>
+bool SSpinBox<NumericType>::GetEnableSlider() const
+{ return EnableSlider.Get(); }
+
+template <typename NumericType>
+void SSpinBox<NumericType>::SetEnableSlider(bool bNewValue)
+{ EnableSlider.Set(bNewValue); }
+
+template <typename NumericType>
+NumericType SSpinBox<NumericType>::GetDelta() const
+{ return Delta.Get(); }
+
+template <typename NumericType>
+void SSpinBox<NumericType>::SetDelta(NumericType InDelta)
+{ Delta.Set(InDelta); }
+
+template <typename NumericType>
+float SSpinBox<NumericType>::GetSliderExponent() const
+{ return SliderExponent.Get(); }
+
+template <typename NumericType>
+void SSpinBox<NumericType>::SetSliderExponent(const TAttribute<float>& InSliderExponent)
+{ SliderExponent = InSliderExponent; }
+
+template <typename NumericType>
+float SSpinBox<NumericType>::GetMinDesiredWidth() const
+{ return SliderExponent.Get(); }
+
+template <typename NumericType>
+void SSpinBox<NumericType>::SetMinDesiredWidth(const TAttribute<float>& InMinDesiredWidth)
+{ MinDesiredWidth = InMinDesiredWidth; }
+
+template <typename NumericType>
+const FSpinBoxStyle* SSpinBox<NumericType>::GetWidgetStyle() const
+{ return Style; }
+
+template <typename NumericType>
+void SSpinBox<NumericType>::SetWidgetStyle(const FSpinBoxStyle* InStyle)
+{ Style = InStyle; }
+
+template <typename NumericType>
+void SSpinBox<NumericType>::InvalidateStyle()
+{ Invalidate(EInvalidateWidgetReason::Layout); }
+
+template <typename NumericType>
+void SSpinBox<NumericType>::SetTextBlockFont(FSlateFontInfo InFont)
+{ EditableText->SetFont(InFont); TextBlock->SetFont(InFont); }
+
+template <typename NumericType>
+void SSpinBox<NumericType>::SetTextJustification(ETextJustify::Type InJustification)
+{ EditableText->SetJustification(InJustification); TextBlock->SetJustification(InJustification);  }
+
+template <typename NumericType>
+void SSpinBox<NumericType>::SetTextClearKeyboardFocusOnCommit(bool bNewValue)
+{ EditableText->SetClearKeyboardFocusOnCommit(bNewValue); }
+
+template <typename NumericType>
+void SSpinBox<NumericType>::SetTextSelectAllTextOnCommit(bool bNewValue)
+{ EditableText->SetSelectAllTextOnCommit(bNewValue); }
+
+template <typename NumericType>
+void SSpinBox<NumericType>::EnterTextMode()
+{
+	TextBlock->SetVisibility(EVisibility::Collapsed);
+	EditableText->SetVisibility(EVisibility::Visible);
+}
+
+template <typename NumericType>
+void SSpinBox<NumericType>::ExitTextMode()
+{
+	TextBlock->SetVisibility(EVisibility::Visible);
+	EditableText->SetVisibility(EVisibility::Collapsed);
+}
+
+template <typename NumericType>
+FString SSpinBox<NumericType>::GetValueAsString() const
+{
+	NumericType CurrentValue = ValueAttribute.Get();
+	if (!bCachedValueStringDirty && CurrentValue == CachedExternalValue)
+	{
+		return CachedValueString;
+	}
+
+	bCachedValueStringDirty = false;
+	return Interface->ToString(CurrentValue);
+}
+
+template <typename NumericType>
+FText SSpinBox<NumericType>::GetValueAsText() const
+{
+	return FText::FromString(GetValueAsString());
+}
+
+template <typename NumericType>
+void SSpinBox<NumericType>::TextField_OnTextChanged(const FText& NewText)
+{
+	if (!bIsTextChanging)
+	{
+		TGuardValue<bool> TextChangedGuard(bIsTextChanging, true);
+
+		// Validate the text on change, and only accept text up until the first invalid character
+		FString Data = NewText.ToString();
+		int32 NumValidChars = Data.Len();
+
+		for (int32 Index = 0; Index < Data.Len(); ++Index)
+		{
+			if (!Interface->IsCharacterValid(Data[Index]))
+			{
+				NumValidChars = Index;
+				break;
+			}
+		}
+
+		if (NumValidChars < Data.Len())
+		{
+			FString ValidData = NumValidChars > 0 ? Data.Left(NumValidChars) : GetValueAsString();
+			EditableText->SetText(FText::FromString(ValidData));
+		}
+
+		// we check that the input is numeric, as we don't want to commit the new value on every change when an expression like *= is entered
+		if (bBroadcastValueChangesPerKey && FCString::IsNumeric(*Data))
+		{
+			TOptional<NumericType> NewValue = Interface->FromString(Data, ValueAttribute.Get());
+			if (NewValue.IsSet())
+			{
+				CommitValue(NewValue.GetValue(), static_cast<double>(NewValue.GetValue()), CommittedViaCode, ETextCommit::Default);
+			}
+		}
+	}
+}
+
+template <typename NumericType>
+void SSpinBox<NumericType>::TextField_OnTextCommitted(const FText& NewText, ETextCommit::Type CommitInfo)
+{
+	if (CommitInfo != ETextCommit::OnEnter)
+	{
+		ExitTextMode();
+	}
+
+	TOptional<NumericType> NewValue = Interface->FromString(NewText.ToString(), ValueAttribute.Get());
+	if (NewValue.IsSet())
+	{
+		CommitValue(NewValue.GetValue(), (double)NewValue.GetValue(), CommittedViaTypeIn, CommitInfo);
+	}
+}
+
+template <typename NumericType>
+void SSpinBox<NumericType>::CommitValue(NumericType NewValue, double NewSpinValue, ECommitMethod CommitMethod, ETextCommit::Type OriginalCommitInfo)
+{
+	if (CommitMethod == CommittedViaSpin || CommitMethod == CommittedViaArrowKey)
+	{
+		const NumericType LocalMinSliderValue = GetMinSliderValue();
+		const NumericType LocalMaxSliderValue = GetMaxSliderValue();
+		NewValue = FMath::Clamp<NumericType>(NewValue, LocalMinSliderValue, LocalMaxSliderValue);
+		NewSpinValue = FMath::Clamp<double>(NewSpinValue, (double)LocalMinSliderValue, (double)LocalMaxSliderValue);
+	}
+
+	{
+		const NumericType LocalMinValue = GetMinValue();
+		const NumericType LocalMaxValue = GetMaxValue();
+		NewValue = FMath::Clamp<NumericType>(NewValue, LocalMinValue, LocalMaxValue);
+		NewSpinValue = FMath::Clamp<double>(NewSpinValue, (double)LocalMinValue, (double)LocalMaxValue);
+	}
+
+	if (!ValueAttribute.IsBound())
+	{
+		ValueAttribute.Set(NewValue);
+	}
+
+	// If not in spin mode, there is no need to jump to the value from the external source, continue to use the committed value.
+	if (CommitMethod == CommittedViaSpin)
+	{
+		const NumericType CurrentValue = ValueAttribute.Get();
+		// This will detect if an external force has changed the value. Internally it will abandon the delta calculated this tick and update the internal value instead.
+		if (CurrentValue != CachedExternalValue)
+		{
+			NewValue = CurrentValue;
+			NewSpinValue = (double)CurrentValue;
+		}
+	}
+
+	// Update the internal value, this needs to be done before rounding.
+	InternalValue = NewSpinValue;
+
+	const bool bAlwaysUsesDeltaSnap = GetAlwaysUsesDeltaSnap();
+	// If needed, round this value to the delta. Internally the value is not held to the Delta but externally it appears to be.
+	if (CommitMethod == CommittedViaSpin || CommitMethod == CommittedViaArrowKey || bAlwaysUsesDeltaSnap)
+	{
+		NumericType CurrentDelta = Delta.Get();
+		if (CurrentDelta != NumericType())
+		{
+			NewValue = FMath::GridSnap<NumericType>(NewValue, CurrentDelta); // snap numeric point value to nearest Delta
+		}
+	}
+
+	// Update the max slider value based on the current value if we're in dynamic mode
+	if (SupportDynamicSliderMaxValue.Get() && ValueAttribute.Get() > GetMaxSliderValue())
+	{
+		ApplySliderMaxValueChanged(float(ValueAttribute.Get() - GetMaxSliderValue()), true);
+	}
+	else if (SupportDynamicSliderMinValue.Get() && ValueAttribute.Get() < GetMinSliderValue())
+	{
+		ApplySliderMinValueChanged(float(ValueAttribute.Get() - GetMinSliderValue()), true);
+	}
+
+	if (CommitMethod == CommittedViaTypeIn || CommitMethod == CommittedViaArrowKey)
+	{
+		OnValueCommitted.ExecuteIfBound(NewValue, OriginalCommitInfo);
+	}
+
+	OnValueChanged.ExecuteIfBound(NewValue);
+
+	if (!ValueAttribute.IsBound())
+	{
+		ValueAttribute.Set(NewValue);
+	}
+
+	// Update the cache of the external value to what the user believes the value is now.
+	const NumericType CurrentValue = ValueAttribute.Get();
+	if (CachedExternalValue != CurrentValue || bCachedValueStringDirty)
+	{
+		CachedExternalValue = ValueAttribute.Get();
+		CachedValueString = Interface->ToString(CachedExternalValue);
+		bCachedValueStringDirty = false;
+	}
+
+	// This ensures that dragging is cleared if focus has been removed from this widget in one of the delegate calls, such as when spawning a modal dialog.
+	if (!this->HasMouseCapture())
+	{
+		bDragging = false;
+		PointerDraggingSliderIndex = INDEX_NONE;
+	}
+}
+
+template <typename NumericType>
+void SSpinBox<NumericType>::NotifyValueCommitted(NumericType CurrentValue) const
+{
+	// The internal value will have been clamped and rounded to the delta at this point, but integer values may still need to be rounded
+	// if the delta is 0.
+	OnValueCommitted.ExecuteIfBound(CurrentValue, ETextCommit::OnEnter);
+	OnEndSliderMovement.ExecuteIfBound(CurrentValue);
+}
+
+template <typename NumericType>
+bool SSpinBox<NumericType>::IsInTextMode() const
+{
+	return (EditableText->GetVisibility() == EVisibility::Visible);
+}
+
+template <typename NumericType>
+float SSpinBox<NumericType>::Fraction(double InValue, double InMinValue, double InMaxValue)
+{
+	const double HalfMax = InMaxValue * 0.5;
+	const double HalfMin = InMinValue * 0.5;
+	const double HalfVal = InValue * 0.5;
+
+	return (float)FMath::Clamp((HalfVal - HalfMin) / (HalfMax - HalfMin), 0.0, 1.0);
+}
+
+template <typename NumericType>
+void SSpinBox<NumericType>::UpdateIsSpinRangeUnlimited()
+{
+	bUnlimitedSpinRange = !((MinValue.Get().IsSet() && MaxValue.Get().IsSet()) || (MinSliderValue.Get().IsSet() && MaxSliderValue.Get().IsSet()));
+}
+
+template <typename NumericType>
+float SSpinBox<NumericType>::GetTextMinDesiredWidth() const
+{
+	return FMath::Max(0.0f, MinDesiredWidth.Get() - (float)Style->ArrowsImage.ImageSize.X);
+}
+
+template <typename NumericType>
+bool SSpinBox<NumericType>::IsCharacterValid(TCHAR InChar) const
+{
+	return Interface->IsCharacterValid(InChar);
+}
+
+template <typename NumericType>
+NumericType SSpinBox<NumericType>::RoundIfIntegerValue(double ValueToRound) const
+{
+	constexpr bool bIsIntegral = TIsIntegral<NumericType>::Value;
+	constexpr bool bCanBeRepresentedInDouble = std::numeric_limits<double>::digits >= std::numeric_limits<NumericType>::digits;
+	if (bIsIntegral && !bCanBeRepresentedInDouble)
+	{
+		return (NumericType)FMath::Clamp<double>(FMath::FloorToDouble(ValueToRound + 0.5), -1.0 * (double)(1ll << std::numeric_limits<double>::digits), (double)(1ll << std::numeric_limits<double>::digits));
+	}
+	else if (bIsIntegral)
+	{
+		return (NumericType)FMath::Clamp<double>(FMath::FloorToDouble(ValueToRound + 0.5), (double)std::numeric_limits<NumericType>::lowest(), (double)std::numeric_limits<NumericType>::max());
+	}
+	else
+	{
+		return (NumericType)FMath::Clamp<double>(ValueToRound, (double)std::numeric_limits<NumericType>::lowest(), (double)std::numeric_limits<NumericType>::max());
+	}
+}
+
+template <typename NumericType>
+void SSpinBox<NumericType>::CancelMouseCapture()
+{
+	bDragging = false;
+	PointerDraggingSliderIndex = INDEX_NONE;
+
+	InternalValue = (double)PreDragValue;
+	NotifyValueCommitted(PreDragValue);
+}
+
+template <typename NumericType>
+void SSpinBox<NumericType>::ResetCachedValueString()
+{ 
+	const NumericType CurrentValue = ValueAttribute.Get();
+	CachedExternalValue = CurrentValue;
+	CachedValueString = Interface->ToString(CachedExternalValue);
+}
 
+// Explicit instantiation of the numeric types valid for this template
+template class SSpinBox<double>;
+template class SSpinBox<float>;
+template class SSpinBox<uint64>;
+template class SSpinBox<uint32>;
+template class SSpinBox<uint16>;
+template class SSpinBox<uint8>;
+template class SSpinBox<int64>;
+template class SSpinBox<int32>;
+template class SSpinBox<int16>;
+template class SSpinBox<int8>;
\ No newline at end of file
diff --git a/Engine/Source/Runtime/Slate/Public/Widgets/Input/SSpinBox.h b/Engine/Source/Runtime/Slate/Public/Widgets/Input/SSpinBox.h
--- a/Engine/Source/Runtime/Slate/Public/Widgets/Input/SSpinBox.h
+++ b/Engine/Source/Runtime/Slate/Public/Widgets/Input/SSpinBox.h
@@ -175,239 +175,20 @@ public:
 
 	SLATE_END_ARGS()
 
-	SSpinBox()
-	{
-	}
+	SLATE_API SSpinBox();
 
-	virtual ~SSpinBox()
-	{
-		if (bDragging || PointerDraggingSliderIndex != INDEX_NONE)
-		{
-			CancelMouseCapture();
-		}
-	}
+	SLATE_API virtual ~SSpinBox();
 
 	/**
 	 * Construct the widget
 	 *
 	 * @param InArgs   A declaration from which to construct the widget
 	 */
-	void Construct(const FArguments& InArgs)
-	{
-		check(InArgs._Style);
+	SLATE_API void Construct(const FArguments& InArgs);
 
-		Style = InArgs._Style;
-
-		SetForegroundColor(InArgs._Style->ForegroundColor);
-		Interface = InArgs._TypeInterface.IsValid() ? InArgs._TypeInterface : MakeShareable(new TDefaultNumericTypeInterface<NumericType>);
-
-		if (InArgs._TypeInterface.IsValid() && Interface->GetOnSettingChanged())
-		{
-			Interface->GetOnSettingChanged()->AddSP(this, &SSpinBox::ResetCachedValueString);
-		}
-
-		ValueAttribute = InArgs._Value;
-		OnValueChanged = InArgs._OnValueChanged;
-		OnValueCommitted = InArgs._OnValueCommitted;
-		OnBeginSliderMovement = InArgs._OnBeginSliderMovement;
-		OnEndSliderMovement = InArgs._OnEndSliderMovement;
-		MinDesiredWidth = InArgs._MinDesiredWidth;
-
-		MinValue = InArgs._MinValue;
-		MaxValue = InArgs._MaxValue;
-		MinSliderValue = (InArgs._MinSliderValue.Get().IsSet()) ? InArgs._MinSliderValue : MinValue;
-		MaxSliderValue = (InArgs._MaxSliderValue.Get().IsSet()) ? InArgs._MaxSliderValue : MaxValue;
-
-		MinFractionalDigits = (InArgs._MinFractionalDigits.Get().IsSet()) ? InArgs._MinFractionalDigits : DefaultMinFractionalDigits;
-		MaxFractionalDigits = (InArgs._MaxFractionalDigits.Get().IsSet()) ? InArgs._MaxFractionalDigits : DefaultMaxFractionalDigits;
-		SetMaxFractionalDigits(MaxFractionalDigits);
-		SetMinFractionalDigits(MinFractionalDigits);
-
-		AlwaysUsesDeltaSnap = InArgs._AlwaysUsesDeltaSnap;
-		EnableSlider = InArgs._EnableSlider;
-
-		SupportDynamicSliderMaxValue = InArgs._SupportDynamicSliderMaxValue;
-		SupportDynamicSliderMinValue = InArgs._SupportDynamicSliderMinValue;
-		OnDynamicSliderMaxValueChanged = InArgs._OnDynamicSliderMaxValueChanged;
-		OnDynamicSliderMinValueChanged = InArgs._OnDynamicSliderMinValueChanged;
-
-		bEnableWheel = InArgs._EnableWheel;
-		bBroadcastValueChangesPerKey = InArgs._BroadcastValueChangesPerKey;
-		WheelStep = InArgs._WheelStep;
-
-		bPreventThrottling = InArgs._PreventThrottling;
-
-		CachedExternalValue = ValueAttribute.Get();
-		CachedValueString = Interface->ToString(CachedExternalValue);
-		bCachedValueStringDirty = false;
-
-		InternalValue = (double)CachedExternalValue;
-
-		if (SupportDynamicSliderMaxValue.Get() && CachedExternalValue > GetMaxSliderValue())
-		{
-			ApplySliderMaxValueChanged(float(CachedExternalValue - GetMaxSliderValue()), true);
-		}
-		else if (SupportDynamicSliderMinValue.Get() && CachedExternalValue < GetMinSliderValue())
-		{
-			ApplySliderMinValueChanged(float(CachedExternalValue - GetMinSliderValue()), true);
-		}
-
-		UpdateIsSpinRangeUnlimited();
-
-		SliderExponent = InArgs._SliderExponent;
-
-		SliderExponentNeutralValue = InArgs._SliderExponentNeutralValue;
-
-		DistanceDragged = 0.0f;
-		PreDragValue = NumericType();
-
-		Delta = InArgs._Delta;
-		ShiftMultiplier = InArgs._ShiftMultiplier;
-		CtrlMultiplier = InArgs._CtrlMultiplier;
-		LinearDeltaSensitivity = InArgs._LinearDeltaSensitivity;
-
-		BackgroundHoveredBrush = &InArgs._Style->HoveredBackgroundBrush;
-		BackgroundBrush = &InArgs._Style->BackgroundBrush;
-		BackgroundActiveBrush = InArgs._Style->ActiveBackgroundBrush.IsSet() ? &InArgs._Style->ActiveBackgroundBrush : BackgroundHoveredBrush;
-
-		ActiveFillBrush = &InArgs._Style->ActiveFillBrush;
-		HoveredFillBrush = InArgs._Style->HoveredFillBrush.IsSet() ? &InArgs._Style->HoveredFillBrush : ActiveFillBrush;
-		InactiveFillBrush = &InArgs._Style->InactiveFillBrush;
-
-		const FMargin& TextMargin = InArgs._Style->TextPadding;
-
-		bDragging = false;
-		PointerDraggingSliderIndex = INDEX_NONE;
-
-		bIsTextChanging = false;
-
-		this->ChildSlot
-		.Padding(InArgs._ContentPadding)
-		[
-			SNew(SHorizontalBox)
-
-			+ SHorizontalBox::Slot()
-			.FillWidth(1.0f)
-			.Padding(TextMargin)
-			.HAlign(HAlign_Fill)
-			.VAlign(VAlign_Center)
-			[
-				SAssignNew(TextBlock, STextBlock)
-				.Font(InArgs._Font)
-				.Text(this, &SSpinBox<NumericType>::GetValueAsText)
-				.MinDesiredWidth(this, &SSpinBox<NumericType>::GetTextMinDesiredWidth)
-				.Justification(InArgs._Justification)
-			]
-			+ SHorizontalBox::Slot()
-			.FillWidth(1.0f)
-			.Padding(TextMargin)
-			.HAlign(HAlign_Fill)
-			.VAlign(VAlign_Center)
-			[
-				SAssignNew(EditableText, SEditableText)
-				.Visibility(EVisibility::Collapsed)
-				.Font(InArgs._Font)
-				.SelectAllTextWhenFocused(true)
-				.Text(this, &SSpinBox<NumericType>::GetValueAsText)
-				.OnIsTypedCharValid(this, &SSpinBox<NumericType>::IsCharacterValid)
-				.OnTextChanged(this, &SSpinBox<NumericType>::TextField_OnTextChanged)
-				.OnTextCommitted(this, &SSpinBox<NumericType>::TextField_OnTextCommitted)
-				.ClearKeyboardFocusOnCommit(InArgs._ClearKeyboardFocusOnCommit)
-				.SelectAllTextOnCommit(InArgs._SelectAllTextOnCommit)
-				.MinDesiredWidth(this, &SSpinBox<NumericType>::GetTextMinDesiredWidth)
-				.VirtualKeyboardType(InArgs._KeyboardType)
-				.Justification(InArgs._Justification)
-				.VirtualKeyboardTrigger(EVirtualKeyboardTrigger::OnAllFocusEvents)
-				.ContextMenuExtender(InArgs._ContextMenuExtender)
-			]
-			+ SHorizontalBox::Slot()
-			.AutoWidth()
-			.HAlign(HAlign_Fill)
-			.VAlign(VAlign_Center)
-			[
-				SNew(SImage)
-				.Image(&InArgs._Style->ArrowsImage)
-				.ColorAndOpacity(FSlateColor::UseForeground())
-			]
-		];
-	}
-
-	virtual int32 OnPaint(const FPaintArgs& Args, const FGeometry& AllottedGeometry, const FSlateRect& MyCullingRect, FSlateWindowElementList& OutDrawElements, int32 LayerId, const FWidgetStyle& InWidgetStyle, bool bParentEnabled) const override
-	{
+	SLATE_API virtual int32 OnPaint(const FPaintArgs& Args, const FGeometry& AllottedGeometry, const FSlateRect& MyCullingRect, FSlateWindowElementList& OutDrawElements, int32 LayerId, const FWidgetStyle& InWidgetStyle, bool bParentEnabled) const override;
 
-		const bool bActiveFeedback = bDragging || IsInTextMode();
-
-		const FSlateBrush* BackgroundImage = bActiveFeedback ? BackgroundActiveBrush : IsHovered() ?
-			BackgroundHoveredBrush :
-			BackgroundBrush;
-
-		const FSlateBrush* FillImage = bActiveFeedback ?
-			ActiveFillBrush : IsHovered() ? HoveredFillBrush : 
-			InactiveFillBrush;
-
-		const int32 BackgroundLayer = LayerId;
-
-		const bool bEnabled = ShouldBeEnabled(bParentEnabled);
-		const ESlateDrawEffect DrawEffects = bEnabled ? ESlateDrawEffect::None : ESlateDrawEffect::DisabledEffect;
-
-		FSlateDrawElement::MakeBox(
-			OutDrawElements,
-			BackgroundLayer,
-			AllottedGeometry.ToPaintGeometry(),
-			BackgroundImage,
-			DrawEffects,
-			BackgroundImage->GetTint(InWidgetStyle) * InWidgetStyle.GetColorAndOpacityTint()
-		);
-
-		const int32 FilledLayer = BackgroundLayer + 1;
-
-		//if there is a spin range limit, draw the filler bar
-		if (!bUnlimitedSpinRange)
-		{
-			NumericType Value = ValueAttribute.Get();
-			NumericType CurrentDelta = Delta.Get();
-			if (CurrentDelta != NumericType())
-			{
-				Value = FMath::GridSnap(Value, CurrentDelta); // snap value to nearest Delta
-			}
-
-			float FractionFilled = Fraction((double)Value, (double)GetMinSliderValue(), (double)GetMaxSliderValue());
-			const float CachedSliderExponent = SliderExponent.Get();
-			if (!FMath::IsNearlyEqual(CachedSliderExponent, 1.f))
-			{
-				if (SliderExponentNeutralValue.IsSet() && SliderExponentNeutralValue.Get() > GetMinSliderValue() && SliderExponentNeutralValue.Get() < GetMaxSliderValue())
-				{
-					//Compute a log curve on both side of the neutral value
-					float StartFractionFilled = Fraction((double)SliderExponentNeutralValue.Get(), (double)GetMinSliderValue(), (double)GetMaxSliderValue());
-					FractionFilled = SpinBoxComputeExponentSliderFraction(FractionFilled, StartFractionFilled, CachedSliderExponent);
-				}
-				else
-				{
-					FractionFilled = 1.0f - FMath::Pow(1.0f - FractionFilled, CachedSliderExponent);
-				}
-			}
-			const FVector2D FillSize(AllottedGeometry.GetLocalSize().X * FractionFilled, AllottedGeometry.GetLocalSize().Y);
-
-			if (!IsInTextMode())
-			{
-				FSlateDrawElement::MakeBox(
-					OutDrawElements,
-					FilledLayer,
-					AllottedGeometry.ToPaintGeometry(FillSize - FVector2D(Style->InsetPadding.GetTotalSpaceAlong<Orient_Horizontal>(), Style->InsetPadding.GetTotalSpaceAlong<Orient_Vertical>()), FSlateLayoutTransform(Style->InsetPadding.GetTopLeft())),
-					FillImage,
-					DrawEffects,
-					FillImage->GetTint(InWidgetStyle) * InWidgetStyle.GetColorAndOpacityTint()
-				);
-			}
-		}
-
-		return FMath::Max(FilledLayer, SCompoundWidget::OnPaint(Args, AllottedGeometry, MyCullingRect, OutDrawElements, FilledLayer, InWidgetStyle, bEnabled));
-	}
-
-	const bool CommitWithMultiplier(const FPointerEvent& MouseEvent)
-	{
-		return MouseEvent.IsShiftDown() || MouseEvent.IsControlDown();
-	}
+	SLATE_API const bool CommitWithMultiplier(const FPointerEvent& MouseEvent);
 
 	/**
 	 * The system calls this method to notify the widget that a mouse button was pressed within it. This event is bubbled.
@@ -416,29 +197,7 @@ public:
 	 * @param MouseEvent Information about the input event
 	 * @return Whether the event was handled along with possible requests for the system to take action.
 	 */
-	virtual FReply OnMouseButtonDown(const FGeometry& MyGeometry, const FPointerEvent& MouseEvent) override
-	{
-		if (MouseEvent.GetEffectingButton() == EKeys::LeftMouseButton && PointerDraggingSliderIndex == INDEX_NONE)
-		{
-			DistanceDragged = 0.f;
-			PreDragValue = ValueAttribute.Get();
-			InternalValue = (double)PreDragValue;
-			PointerDraggingSliderIndex = MouseEvent.GetPointerIndex();
-			CachedMousePosition = MouseEvent.GetScreenSpacePosition().IntPoint();
-
-			FReply ReturnReply = FReply::Handled().CaptureMouse(SharedThis(this)).UseHighPrecisionMouseMovement(SharedThis(this)).SetUserFocus(SharedThis(this), EFocusCause::Mouse);
-			if (bPreventThrottling)
-			{
-				ReturnReply.PreventThrottling();
-			}
-			return ReturnReply;
-		}
-		else
-		{
-
-			return FReply::Unhandled();
-		}
-	}
+	SLATE_API virtual FReply OnMouseButtonDown(const FGeometry& MyGeometry, const FPointerEvent& MouseEvent) override;
 
 	/**
 	 * The system calls this method to notify the widget that a mouse button was release within it. This event is bubbled.
@@ -447,106 +206,11 @@ public:
 	 * @param MouseEvent Information about the input event
 	 * @return Whether the event was handled along with possible requests for the system to take action.
 	 */
-	virtual FReply OnMouseButtonUp(const FGeometry& MyGeometry, const FPointerEvent& MouseEvent) override
-	{
-		if (MouseEvent.GetEffectingButton() == EKeys::LeftMouseButton && PointerDraggingSliderIndex == MouseEvent.GetPointerIndex())
-		{
-			if (!this->HasMouseCapture())
-			{
-				// Lost Capture - ensure reset
-				bDragging = false;
-				PointerDraggingSliderIndex = INDEX_NONE;
-
-				return FReply::Unhandled();
-			}
-
-			if (bDragging)
-			{
-				NumericType CurrentDelta = Delta.Get();
-				if (CurrentDelta != NumericType() && !CommitWithMultiplier(MouseEvent))
-				{
-					InternalValue = FMath::GridSnap(InternalValue, (double)CurrentDelta);
-				}
-
-				const NumericType CurrentValue = RoundIfIntegerValue(InternalValue);
-				NotifyValueCommitted(CurrentValue);
-			}
-
-			bDragging = false;
-			PointerDraggingSliderIndex = INDEX_NONE;
-
-			FReply Reply = FReply::Handled().ReleaseMouseCapture();
-
-			if (!MouseEvent.IsTouchEvent())
-			{
-				Reply.SetMousePos(CachedMousePosition);
-			}
-
-			if (DistanceDragged < FSlateApplication::Get().GetDragTriggerDistance())
-			{
-				EnterTextMode();
-				Reply.SetUserFocus(EditableText.ToSharedRef(), EFocusCause::SetDirectly);
-			}
-
-			return Reply;
-		}
-
-		return FReply::Unhandled();
-	}
-
-	void ApplySliderMaxValueChanged(float SliderDeltaToAdd, bool UpdateOnlyIfHigher)
-	{
-		check(SupportDynamicSliderMaxValue.Get());
-
-		NumericType NewMaxSliderValue = std::numeric_limits<NumericType>::min();
-
-		if (MaxSliderValue.IsSet() && MaxSliderValue.Get().IsSet())
-		{
-			NewMaxSliderValue = GetMaxSliderValue();
-
-			if ((NewMaxSliderValue + (NumericType)SliderDeltaToAdd > GetMaxSliderValue() && UpdateOnlyIfHigher) || !UpdateOnlyIfHigher)
-			{
-				NewMaxSliderValue += (NumericType)SliderDeltaToAdd;
-
-				if (!MaxSliderValue.IsBound()) // simple value so we can update it without breaking the mechanic otherwise it must be handle by the callback implementer
-				{
-					SetMaxSliderValue(NewMaxSliderValue);
-				}
-			}
-		}
-
-		if (OnDynamicSliderMaxValueChanged.IsBound())
-		{
-			OnDynamicSliderMaxValueChanged.Execute(NewMaxSliderValue, TWeakPtr<SWidget>(AsShared()), true, UpdateOnlyIfHigher);
-		}
-	}
-
-	void ApplySliderMinValueChanged(float SliderDeltaToAdd, bool UpdateOnlyIfLower)
-	{
-		check(SupportDynamicSliderMaxValue.Get());
-
-		NumericType NewMinSliderValue = std::numeric_limits<NumericType>::min();
+	SLATE_API virtual FReply OnMouseButtonUp(const FGeometry& MyGeometry, const FPointerEvent& MouseEvent) override;
 
-		if (MinSliderValue.IsSet() && MinSliderValue.Get().IsSet())
-		{
-			NewMinSliderValue = GetMinSliderValue();
+	SLATE_API void ApplySliderMaxValueChanged(float SliderDeltaToAdd, bool UpdateOnlyIfHigher);
 
-			if ((NewMinSliderValue + (NumericType)SliderDeltaToAdd < GetMinSliderValue() && UpdateOnlyIfLower) || !UpdateOnlyIfLower)
-			{
-				NewMinSliderValue += (NumericType)SliderDeltaToAdd;
-
-				if (!MinSliderValue.IsBound()) // simple value so we can update it without breaking the mechanic otherwise it must be handle by the callback implementer
-				{
-					SetMinSliderValue(NewMinSliderValue);
-				}
-			}
-		}
-
-		if (OnDynamicSliderMinValueChanged.IsBound())
-		{
-			OnDynamicSliderMinValueChanged.Execute(NewMinSliderValue, TWeakPtr<SWidget>(AsShared()), true, UpdateOnlyIfLower);
-		}
-	}
+	SLATE_API void ApplySliderMinValueChanged(float SliderDeltaToAdd, bool UpdateOnlyIfLower);
 
 	/**
 	 * The system calls this method to notify the widget that a mouse moved within it. This event is bubbled.
@@ -555,151 +219,7 @@ public:
 	 * @param MouseEvent Information about the input event
 	 * @return Whether the event was handled along with possible requests for the system to take action.
 	 */
-	virtual FReply OnMouseMove(const FGeometry& MyGeometry, const FPointerEvent& MouseEvent) override
-	{
-		const bool bEnableSlider = GetEnableSlider();
-		if (PointerDraggingSliderIndex == MouseEvent.GetPointerIndex() && bEnableSlider)
-		{
-			if (!this->HasMouseCapture())
-			{
-				// Lost the mouse capture - ensure reset
-				bDragging = false;
-				PointerDraggingSliderIndex = INDEX_NONE;
-
-				return FReply::Unhandled();
-			}
-
-			if (!bDragging)
-			{
-				DistanceDragged += (float)FMath::Abs(MouseEvent.GetCursorDelta().X);
-				if (DistanceDragged > FSlateApplication::Get().GetDragTriggerDistance())
-				{
-					ExitTextMode();
-					bDragging = true;
-					OnBeginSliderMovement.ExecuteIfBound();
-				}
-
-				// Cache the mouse, even if not dragging cache it.
-				CachedMousePosition = MouseEvent.GetScreenSpacePosition().IntPoint();
-			}
-			else
-			{
-				double NewValue = 0.0;
-
-				// Increments the spin based on delta mouse movement.
-
-				// A minimum slider width to use for calculating deltas in the slider-range space
-				const float MinSliderWidth = 100.f;
-				float SliderWidthInSlateUnits = FMath::Max((float)MyGeometry.GetDrawSize().X, MinSliderWidth);
-
-				if (MouseEvent.IsAltDown())
-				{
-					float DeltaToAdd = (float)MouseEvent.GetCursorDelta().X / SliderWidthInSlateUnits;
-
-					if (SupportDynamicSliderMaxValue.Get() && (NumericType)InternalValue == GetMaxSliderValue())
-					{
-						ApplySliderMaxValueChanged(DeltaToAdd, false);
-					}
-					else if (SupportDynamicSliderMinValue.Get() && (NumericType)InternalValue == GetMinSliderValue())
-					{
-						ApplySliderMinValueChanged(DeltaToAdd, false);
-					}
-				}
-
-				ECommitMethod CommitMethod = CommittedViaSpin;
-				
-				const bool bIsSmallStep = (GetMaxSliderValue() - GetMinSliderValue()) <= 10.0;
-				double Step = bIsSmallStep ? 0.1 : 1.0;
-
-				if (MouseEvent.IsControlDown())
-				{
-					Step *= CtrlMultiplier.Get();
-					CommitMethod = CommittedViaSpinMultiplier;
-				}
-				else if (MouseEvent.IsShiftDown())
-				{
-					Step *= ShiftMultiplier.Get();
-					CommitMethod = CommittedViaSpinMultiplier;
-				}
-
-				//if we have a range to draw in
-				if (!bUnlimitedSpinRange)
-				{
-					bool HasValidExponentNeutralValue = SliderExponentNeutralValue.IsSet() && SliderExponentNeutralValue.Get() > GetMinSliderValue() && SliderExponentNeutralValue.Get() < GetMaxSliderValue();
-
-					const float CachedSliderExponent = SliderExponent.Get();
-					// The amount currently filled in the spinbox, needs to be calculated to do deltas correctly.
-					float FractionFilled = Fraction(InternalValue, (double)GetMinSliderValue(), (double)GetMaxSliderValue());
-
-					if (!FMath::IsNearlyEqual(CachedSliderExponent, 1.0f))
-					{
-						if (HasValidExponentNeutralValue)
-						{
-							//Compute a log curve on both side of the neutral value
-							float StartFractionFilled = Fraction((double)SliderExponentNeutralValue.Get(), (double)GetMinSliderValue(), (double)GetMaxSliderValue());
-							FractionFilled = SpinBoxComputeExponentSliderFraction(FractionFilled, StartFractionFilled, CachedSliderExponent);
-						}
-						else
-						{
-							FractionFilled = 1.0f - FMath::Pow(1.0f - FractionFilled, CachedSliderExponent);
-						}
-					}
-					FractionFilled *= SliderWidthInSlateUnits;
-
-					// Now add the delta to the fraction filled, this causes the spin.
-					FractionFilled += (float)(MouseEvent.GetCursorDelta().X * Step);
-
-					// Clamp the fraction to be within the bounds of the geometry.
-					FractionFilled = FMath::Clamp(FractionFilled, 0.0f, SliderWidthInSlateUnits);
-
-					// Convert the fraction filled to a percent.
-					float Percent = FMath::Clamp(FractionFilled / SliderWidthInSlateUnits, 0.0f, 1.0f);
-					if (!FMath::IsNearlyEqual(CachedSliderExponent, 1.0f))
-					{
-						// Have to convert the percent to the proper value due to the exponent component to the spin.
-						if (HasValidExponentNeutralValue)
-						{
-							//Compute a log curve on both side of the neutral value
-							float StartFractionFilled = Fraction(SliderExponentNeutralValue.Get(), GetMinSliderValue(), GetMaxSliderValue());
-							Percent = SpinBoxComputeExponentSliderFraction(Percent, StartFractionFilled, 1.0f / CachedSliderExponent);
-						}
-						else
-						{
-							Percent = 1.0f - FMath::Pow(1.0f - Percent, 1.0f / CachedSliderExponent);
-						}
-
-
-					}
-
-					NewValue = FMath::LerpStable<double>((double)GetMinSliderValue(), (double)GetMaxSliderValue(), Percent);
-				}
-				else
-				{
-					// If this control has a specified delta and sensitivity then we use that instead of the current value for determining how much to change.
-					const double Sign = (MouseEvent.GetCursorDelta().X > 0) ? 1.0 : -1.0;
-
-					if (LinearDeltaSensitivity.IsSet() && LinearDeltaSensitivity.Get() != 0 && Delta.IsSet() && Delta.Get() > 0)
-					{
-						const double MouseDelta = FMath::Abs(MouseEvent.GetCursorDelta().X / (float)LinearDeltaSensitivity.Get());
-						NewValue = InternalValue + (Sign * MouseDelta * FMath::Pow((double)Delta.Get(), (double)SliderExponent.Get())) * Step;
-					}
-					else
-					{
-						const double MouseDelta = FMath::Abs(MouseEvent.GetCursorDelta().X / SliderWidthInSlateUnits);
-						const double CurrentValue = FMath::Clamp<double>(FMath::Abs(InternalValue), 1.0, (double)std::numeric_limits<NumericType>::max());
-						NewValue = InternalValue + (Sign * MouseDelta * FMath::Pow((double)CurrentValue, (double)SliderExponent.Get())) * Step;
-					}
-				}
-
-				NumericType RoundedNewValue = RoundIfIntegerValue(NewValue);
-				CommitValue(RoundedNewValue, NewValue, CommitMethod, ETextCommit::OnEnter);
-			}
-
-			return FReply::Handled();
-		}
-
-		return FReply::Unhandled();
-	}
+	SLATE_API virtual FReply OnMouseMove(const FGeometry& MyGeometry, const FPointerEvent& MouseEvent) override;
 
 	/**
 	 * Called when the mouse wheel is spun. This event is bubbled.
@@ -707,289 +227,109 @@ public:
 	 * @param  MouseEvent  Mouse event
 	 * @return  Returns whether the event was handled, along with other possible actions
 	 */
-	virtual FReply OnMouseWheel(const FGeometry& MyGeometry, const FPointerEvent& MouseEvent) override
-	{
-		if (bEnableWheel && PointerDraggingSliderIndex == INDEX_NONE && HasKeyboardFocus())
-		{
-			// If there is no WheelStep defined, we use 1.0 (Or 0.1 if slider range is <= 10)
-			constexpr bool bIsIntegral = TIsIntegral<NumericType>::Value;
-			const bool bIsSmallStep = !bIsIntegral && (GetMaxSliderValue() - GetMinSliderValue()) <= 10.0;
-			double Step = WheelStep.IsSet() && WheelStep.Get().IsSet() ? WheelStep.Get().GetValue() : (bIsSmallStep ? 0.1 : 1.0);
-
-			if (MouseEvent.IsControlDown())
-			{
-				// If no value is set for WheelSmallStep, we use the DefaultStep multiplied by the CtrlMultiplier
-				Step *= CtrlMultiplier.Get();
-			}
-			else if (MouseEvent.IsShiftDown())
-			{
-				// If no value is set for WheelBigStep, we use the DefaultStep multiplied by the ShiftMultiplier
-				Step *= ShiftMultiplier.Get();
-			}
-
-			const double Sign = (MouseEvent.GetWheelDelta() > 0) ? 1.0 : -1.0;
-			const double NewValue = InternalValue + (Sign * Step);
-			const NumericType RoundedNewValue = RoundIfIntegerValue(NewValue);
-
-			// First SetText is to update the value before calling CommitValue. Otherwise, when the text lose
-			// focus from the CommitValue, it will override the value we just committed.
-			// The second SetText is to update the text to the InternalValue since it could have been clamped.
-			EditableText->SetText(FText::FromString(Interface->ToString((NumericType)NewValue)));
-			CommitValue(RoundedNewValue, NewValue, CommittedViaSpin, ETextCommit::OnEnter);
-			EditableText->SetText(FText::FromString(Interface->ToString((NumericType)InternalValue)));
-
-			return FReply::Handled();
-		}
-
-		return FReply::Unhandled();
-	}
-
-	virtual FCursorReply OnCursorQuery(const FGeometry& MyGeometry, const FPointerEvent& CursorEvent) const override
-	{
-		const bool bEnableSlider = GetEnableSlider();
+	SLATE_API virtual FReply OnMouseWheel(const FGeometry& MyGeometry, const FPointerEvent& MouseEvent) override;
 
-		if (!bEnableSlider)
-		{
-			return FCursorReply::Cursor(EMouseCursor::Default);
-		}
+	SLATE_API virtual FCursorReply OnCursorQuery(const FGeometry& MyGeometry, const FPointerEvent& CursorEvent) const override;
 
-		return bDragging ?
-			FCursorReply::Cursor(EMouseCursor::None) :
-			FCursorReply::Cursor(EMouseCursor::ResizeLeftRight);
-	}
+	SLATE_API virtual bool SupportsKeyboardFocus() const override;
 
-	virtual bool SupportsKeyboardFocus() const override
-	{
-		// SSpinBox is focusable.
-		return true;
-	}
 
+	SLATE_API virtual FReply OnFocusReceived(const FGeometry& MyGeometry, const FFocusEvent& InFocusEvent) override;
 
-	virtual FReply OnFocusReceived(const FGeometry& MyGeometry, const FFocusEvent& InFocusEvent) override
-	{
-		if (!bDragging && (InFocusEvent.GetCause() == EFocusCause::Navigation || InFocusEvent.GetCause() == EFocusCause::SetDirectly))
-		{
-			EnterTextMode();
-			return FReply::Handled().SetUserFocus(EditableText.ToSharedRef(), InFocusEvent.GetCause());
-		}
-		else
-		{
-			return FReply::Unhandled();
-		}
-	}
-
-	virtual FReply OnKeyDown(const FGeometry& MyGeometry, const FKeyEvent& InKeyEvent) override
-	{
-		const FKey Key = InKeyEvent.GetKey();
-		if (Key == EKeys::Escape && HasMouseCapture())
-		{
-			CancelMouseCapture();
-			return FReply::Handled().ReleaseMouseCapture().SetMousePos(CachedMousePosition);
-		}
-		else if (Key == EKeys::Up || Key == EKeys::Right)
-		{
-			const NumericType LocalValueAttribute = ValueAttribute.Get();
-			const NumericType LocalDelta = Delta.Get();
-			InternalValue = (double)LocalValueAttribute;
-			CommitValue(LocalValueAttribute + LocalDelta, InternalValue + (double)LocalDelta, CommittedViaArrowKey, ETextCommit::OnEnter);
-			ExitTextMode();
-			return FReply::Handled();
-		}
-		else if (Key == EKeys::Down || Key == EKeys::Left)
-		{
-			const NumericType LocalValueAttribute = ValueAttribute.Get();
-			const NumericType LocalDelta = Delta.Get();
-			InternalValue = (double)LocalValueAttribute;
-			CommitValue(LocalValueAttribute - LocalDelta, InternalValue + (double)LocalDelta, CommittedViaArrowKey, ETextCommit::OnEnter);
-			ExitTextMode();
-			return FReply::Handled();
-		}
-		else if (Key == EKeys::Enter)
-		{
-			InternalValue = (double)ValueAttribute.Get();
-			EnterTextMode();
-			return FReply::Handled().SetUserFocus(EditableText.ToSharedRef(), EFocusCause::Navigation);
-		}
-		else
-		{
-			return FReply::Unhandled();
-		}
-	}
-
-	virtual bool HasKeyboardFocus() const override
-	{
-		// The spinbox is considered focused when we are typing it text.
-		return SCompoundWidget::HasKeyboardFocus() || (EditableText.IsValid() && EditableText->HasKeyboardFocus());
-	}
+	SLATE_API virtual FReply OnKeyDown(const FGeometry& MyGeometry, const FKeyEvent& InKeyEvent) override;
+
+	SLATE_API virtual bool HasKeyboardFocus() const override;
 
 	/** Return the Value attribute */
-	TAttribute<NumericType> GetValueAttribute() const { return ValueAttribute; }
+	SLATE_API TAttribute<NumericType> GetValueAttribute() const;
 
 	/** See the Value attribute */
-	NumericType GetValue() const { return ValueAttribute.Get(); }
-	void SetValue(const TAttribute<NumericType>& InValueAttribute)
-	{
-		ValueAttribute = InValueAttribute;
-		const NumericType LocalValueAttribute = ValueAttribute.Get();
-		CommitValue(LocalValueAttribute, (double)LocalValueAttribute, ECommitMethod::CommittedViaCode, ETextCommit::Default);
-	}
+	SLATE_API NumericType GetValue() const;
+
+	SLATE_API void SetValue(const TAttribute<NumericType>& InValueAttribute);
 
 	/** See the MinValue attribute */
-	NumericType GetMinValue() const { return MinValue.Get().Get(std::numeric_limits<NumericType>::lowest()); }
-	void SetMinValue(const TAttribute<TOptional<NumericType>>& InMinValue)
-	{
-		MinValue = InMinValue;
-		UpdateIsSpinRangeUnlimited();
-	}
+	SLATE_API NumericType GetMinValue() const;
+
+	SLATE_API void SetMinValue(const TAttribute<TOptional<NumericType>>& InMinValue);
 
 	/** See the MaxValue attribute */
-	NumericType GetMaxValue() const { return MaxValue.Get().Get(std::numeric_limits<NumericType>::max()); }
-	void SetMaxValue(const TAttribute<TOptional<NumericType>>& InMaxValue)
-	{
-		MaxValue = InMaxValue;
-		UpdateIsSpinRangeUnlimited();
-	}
+	SLATE_API NumericType GetMaxValue() const;
+
+	SLATE_API void SetMaxValue(const TAttribute<TOptional<NumericType>>& InMaxValue);
 
 	/** See the MinSliderValue attribute */
-	bool IsMinSliderValueBound() const { return MinSliderValue.IsBound(); }
+	SLATE_API bool IsMinSliderValueBound() const;
 
-	NumericType GetMinSliderValue() const { return MinSliderValue.Get().Get(std::numeric_limits<NumericType>::lowest()); }
-	void SetMinSliderValue(const TAttribute<TOptional<NumericType>>& InMinSliderValue)
-	{
-		MinSliderValue = (InMinSliderValue.Get().IsSet()) ? InMinSliderValue : MinValue;
-		UpdateIsSpinRangeUnlimited();
-	}
+	SLATE_API NumericType GetMinSliderValue() const;
+
+	SLATE_API void SetMinSliderValue(const TAttribute<TOptional<NumericType>>& InMinSliderValue);
 
 	/** See the MaxSliderValue attribute */
-	bool IsMaxSliderValueBound() const { return MaxSliderValue.IsBound(); }
+	SLATE_API bool IsMaxSliderValueBound() const;
 
-	NumericType GetMaxSliderValue() const { return MaxSliderValue.Get().Get(std::numeric_limits<NumericType>::max()); }
-	void SetMaxSliderValue(const TAttribute<TOptional<NumericType>>& InMaxSliderValue)
-	{
-		MaxSliderValue = (InMaxSliderValue.Get().IsSet()) ? InMaxSliderValue : MaxValue;;
-		UpdateIsSpinRangeUnlimited();
-	}
+	SLATE_API NumericType GetMaxSliderValue() const;
+
+	SLATE_API void SetMaxSliderValue(const TAttribute<TOptional<NumericType>>& InMaxSliderValue);
 
 	/** See the MinFractionalDigits attribute */
-	int32 GetMinFractionalDigits() const { return Interface->GetMinFractionalDigits(); }
-	void SetMinFractionalDigits(const TAttribute<TOptional<int32>>& InMinFractionalDigits)
-	{
-		Interface->SetMinFractionalDigits((InMinFractionalDigits.Get().IsSet()) ? InMinFractionalDigits.Get() : MinFractionalDigits);
-		bCachedValueStringDirty = true;
-	}
+	SLATE_API int32 GetMinFractionalDigits() const;
+
+	SLATE_API void SetMinFractionalDigits(const TAttribute<TOptional<int32>>& InMinFractionalDigits);
 
 	/** See the MaxFractionalDigits attribute */
-	int32 GetMaxFractionalDigits() const { return Interface->GetMaxFractionalDigits(); }
-	void SetMaxFractionalDigits(const TAttribute<TOptional<int32>>& InMaxFractionalDigits)
-	{
-		Interface->SetMaxFractionalDigits((InMaxFractionalDigits.Get().IsSet()) ? InMaxFractionalDigits.Get() : MaxFractionalDigits);
-		bCachedValueStringDirty = true;
-	}
+	SLATE_API int32 GetMaxFractionalDigits() const;
+
+	SLATE_API void SetMaxFractionalDigits(const TAttribute<TOptional<int32>>& InMaxFractionalDigits);
 
 	/** See the AlwaysUsesDeltaSnap attribute */
-	bool GetAlwaysUsesDeltaSnap() const { return AlwaysUsesDeltaSnap.Get(); }
-	void SetAlwaysUsesDeltaSnap(bool bNewValue) { AlwaysUsesDeltaSnap.Set(bNewValue); }
+	SLATE_API bool GetAlwaysUsesDeltaSnap() const;
+	SLATE_API void SetAlwaysUsesDeltaSnap(bool bNewValue);
 
 	/** See the EnableSlider attribute */
-	bool GetEnableSlider() const { return EnableSlider.Get(); }
-	void SetEnableSlider(bool bNewValue) { EnableSlider.Set(bNewValue); }
+	SLATE_API bool GetEnableSlider() const;
+	SLATE_API void SetEnableSlider(bool bNewValue);
 
 	/** See the Delta attribute */
-	NumericType GetDelta() const { return Delta.Get(); }
-	void SetDelta(NumericType InDelta) { Delta.Set(InDelta); }
+	SLATE_API NumericType GetDelta() const;
+	SLATE_API void SetDelta(NumericType InDelta);
 
 	/** See the SliderExponent attribute */
-	float GetSliderExponent() const { return SliderExponent.Get(); }
-	void SetSliderExponent(const TAttribute<float>& InSliderExponent) { SliderExponent = InSliderExponent; }
+	SLATE_API float GetSliderExponent() const;
+	SLATE_API void SetSliderExponent(const TAttribute<float>& InSliderExponent);
 
 	/** See the MinDesiredWidth attribute */
-	float GetMinDesiredWidth() const { return SliderExponent.Get(); }
-	void SetMinDesiredWidth(const TAttribute<float>& InMinDesiredWidth) { MinDesiredWidth = InMinDesiredWidth; }
+	SLATE_API float GetMinDesiredWidth() const;
+	SLATE_API void SetMinDesiredWidth(const TAttribute<float>& InMinDesiredWidth);
 
-	const FSpinBoxStyle* GetWidgetStyle() const { return Style; }
-	void SetWidgetStyle(const FSpinBoxStyle* InStyle) { Style = InStyle; }
-	void InvalidateStyle() { Invalidate(EInvalidateWidgetReason::Layout); }
+	SLATE_API const FSpinBoxStyle* GetWidgetStyle() const;
+	SLATE_API void SetWidgetStyle(const FSpinBoxStyle* InStyle);
+	SLATE_API void InvalidateStyle();
 
-	void SetTextBlockFont(FSlateFontInfo InFont) { EditableText->SetFont(InFont); TextBlock->SetFont(InFont); }
-	void SetTextJustification(ETextJustify::Type InJustification) { EditableText->SetJustification(InJustification); TextBlock->SetJustification(InJustification);  }
-	void SetTextClearKeyboardFocusOnCommit(bool bNewValue) { EditableText->SetClearKeyboardFocusOnCommit(bNewValue); }
-	void SetTextSelectAllTextOnCommit(bool bNewValue) { EditableText->SetSelectAllTextOnCommit(bNewValue); }
+	SLATE_API void SetTextBlockFont(FSlateFontInfo InFont);
+	SLATE_API void SetTextJustification(ETextJustify::Type InJustification);
+	SLATE_API void SetTextClearKeyboardFocusOnCommit(bool bNewValue);
+	SLATE_API void SetTextSelectAllTextOnCommit(bool bNewValue);
 
 protected:
 	/** Make the spinbox switch to keyboard-based input mode. */
-	void EnterTextMode()
-	{
-		TextBlock->SetVisibility(EVisibility::Collapsed);
-		EditableText->SetVisibility(EVisibility::Visible);
-	}
+	SLATE_API void EnterTextMode();
 
 	/** Make the spinbox switch to mouse-based input mode. */
-	void ExitTextMode()
-	{
-		TextBlock->SetVisibility(EVisibility::Visible);
-		EditableText->SetVisibility(EVisibility::Collapsed);
-	}
+	SLATE_API void ExitTextMode();
 
 	/** @return the value being observed by the spinbox as a string */
-	FString GetValueAsString() const
-	{
-		NumericType CurrentValue = ValueAttribute.Get();
-		if (!bCachedValueStringDirty && CurrentValue == CachedExternalValue)
-		{
-			return CachedValueString;
-		}
-
-		bCachedValueStringDirty = false;
-		return Interface->ToString(CurrentValue);
-	}
+	SLATE_API FString GetValueAsString() const;
 
 	/** @return the value being observed by the spinbox as FText - todo: spinbox FText support (reimplement me) */
-	FText GetValueAsText() const
-	{
-		return FText::FromString(GetValueAsString());
-	}
+	SLATE_API FText GetValueAsText() const;
 
 	/**
 	 * Invoked when the text in the text field changes
 	 *
 	 * @param NewText		The value of the text in the text field
 	 */
-	void TextField_OnTextChanged(const FText& NewText)
-	{
-		if (!bIsTextChanging)
-		{
-			TGuardValue<bool> TextChangedGuard(bIsTextChanging, true);
-
-			// Validate the text on change, and only accept text up until the first invalid character
-			FString Data = NewText.ToString();
-			int32 NumValidChars = Data.Len();
-
-			for (int32 Index = 0; Index < Data.Len(); ++Index)
-			{
-				if (!Interface->IsCharacterValid(Data[Index]))
-				{
-					NumValidChars = Index;
-					break;
-				}
-			}
-
-			if (NumValidChars < Data.Len())
-			{
-				FString ValidData = NumValidChars > 0 ? Data.Left(NumValidChars) : GetValueAsString();
-				EditableText->SetText(FText::FromString(ValidData));
-			}
-
-			// we check that the input is numeric, as we don't want to commit the new value on every change when an expression like *= is entered
-			if (bBroadcastValueChangesPerKey && FCString::IsNumeric(*Data))
-			{
-				TOptional<NumericType> NewValue = Interface->FromString(Data, ValueAttribute.Get());
-				if (NewValue.IsSet())
-				{
-					CommitValue(NewValue.GetValue(), static_cast<double>(NewValue.GetValue()), CommittedViaCode, ETextCommit::Default);
-				}
-			}
-		}
-	}
+	SLATE_API void TextField_OnTextChanged(const FText& NewText);
 
 	/**
 	 * Invoked when the text field commits its text.
@@ -997,19 +337,7 @@ protected:
 	 * @param NewText		The value of text coming from the editable text field.
 	 * @param CommitInfo	Information about the source of the commit
 	 */
-	void TextField_OnTextCommitted(const FText& NewText, ETextCommit::Type CommitInfo)
-	{
-		if (CommitInfo != ETextCommit::OnEnter)
-		{
-			ExitTextMode();
-		}
-
-		TOptional<NumericType> NewValue = Interface->FromString(NewText.ToString(), ValueAttribute.Get());
-		if (NewValue.IsSet())
-		{
-			CommitValue(NewValue.GetValue(), (double)NewValue.GetValue(), CommittedViaTypeIn, CommitInfo);
-		}
-	}
+	SLATE_API void TextField_OnTextCommitted(const FText& NewText, ETextCommit::Type CommitInfo);
 
 
 	/** How user changed the value in the spinbox */
@@ -1029,124 +357,23 @@ protected:
 	 * @param CommitMethod           Did the user type in the value or spin to it.
 	 * @param OriginalCommitInfo	 If the user typed in the value, information about the source of the commit
 	 */
-	void CommitValue(NumericType NewValue, double NewSpinValue, ECommitMethod CommitMethod, ETextCommit::Type OriginalCommitInfo)
-	{
-		if (CommitMethod == CommittedViaSpin || CommitMethod == CommittedViaArrowKey)
-		{
-			const NumericType LocalMinSliderValue = GetMinSliderValue();
-			const NumericType LocalMaxSliderValue = GetMaxSliderValue();
-			NewValue = FMath::Clamp<NumericType>(NewValue, LocalMinSliderValue, LocalMaxSliderValue);
-			NewSpinValue = FMath::Clamp<double>(NewSpinValue, (double)LocalMinSliderValue, (double)LocalMaxSliderValue);
-		}
-
-		{
-			const NumericType LocalMinValue = GetMinValue();
-			const NumericType LocalMaxValue = GetMaxValue();
-			NewValue = FMath::Clamp<NumericType>(NewValue, LocalMinValue, LocalMaxValue);
-			NewSpinValue = FMath::Clamp<double>(NewSpinValue, (double)LocalMinValue, (double)LocalMaxValue);
-		}
-
-		if (!ValueAttribute.IsBound())
-		{
-			ValueAttribute.Set(NewValue);
-		}
-
-		// If not in spin mode, there is no need to jump to the value from the external source, continue to use the committed value.
-		if (CommitMethod == CommittedViaSpin)
-		{
-			const NumericType CurrentValue = ValueAttribute.Get();
-			// This will detect if an external force has changed the value. Internally it will abandon the delta calculated this tick and update the internal value instead.
-			if (CurrentValue != CachedExternalValue)
-			{
-				NewValue = CurrentValue;
-				NewSpinValue = (double)CurrentValue;
-			}
-		}
-
-		// Update the internal value, this needs to be done before rounding.
-		InternalValue = NewSpinValue;
-
-		const bool bAlwaysUsesDeltaSnap = GetAlwaysUsesDeltaSnap();
-		// If needed, round this value to the delta. Internally the value is not held to the Delta but externally it appears to be.
-		if (CommitMethod == CommittedViaSpin || CommitMethod == CommittedViaArrowKey || bAlwaysUsesDeltaSnap)
-		{
-			NumericType CurrentDelta = Delta.Get();
-			if (CurrentDelta != NumericType())
-			{
-				NewValue = FMath::GridSnap<NumericType>(NewValue, CurrentDelta); // snap numeric point value to nearest Delta
-			}
-		}
-
-		// Update the max slider value based on the current value if we're in dynamic mode
-		if (SupportDynamicSliderMaxValue.Get() && ValueAttribute.Get() > GetMaxSliderValue())
-		{
-			ApplySliderMaxValueChanged(float(ValueAttribute.Get() - GetMaxSliderValue()), true);
-		}
-		else if (SupportDynamicSliderMinValue.Get() && ValueAttribute.Get() < GetMinSliderValue())
-		{
-			ApplySliderMinValueChanged(float(ValueAttribute.Get() - GetMinSliderValue()), true);
-		}
-
-		if (CommitMethod == CommittedViaTypeIn || CommitMethod == CommittedViaArrowKey)
-		{
-			OnValueCommitted.ExecuteIfBound(NewValue, OriginalCommitInfo);
-		}
-
-		OnValueChanged.ExecuteIfBound(NewValue);
-
-		if (!ValueAttribute.IsBound())
-		{
-			ValueAttribute.Set(NewValue);
-		}
-
-		// Update the cache of the external value to what the user believes the value is now.
-		const NumericType CurrentValue = ValueAttribute.Get();
-		if (CachedExternalValue != CurrentValue || bCachedValueStringDirty)
-		{
-			CachedExternalValue = ValueAttribute.Get();
-			CachedValueString = Interface->ToString(CachedExternalValue);
-			bCachedValueStringDirty = false;
-		}
-
-		// This ensures that dragging is cleared if focus has been removed from this widget in one of the delegate calls, such as when spawning a modal dialog.
-		if (!this->HasMouseCapture())
-		{
-			bDragging = false;
-			PointerDraggingSliderIndex = INDEX_NONE;
-		}
-	}
-
-	void NotifyValueCommitted(NumericType CurrentValue) const
-	{
-		// The internal value will have been clamped and rounded to the delta at this point, but integer values may still need to be rounded
-		// if the delta is 0.
-		OnValueCommitted.ExecuteIfBound(CurrentValue, ETextCommit::OnEnter);
-		OnEndSliderMovement.ExecuteIfBound(CurrentValue);
-	}
+	SLATE_API void CommitValue(NumericType NewValue, double NewSpinValue, ECommitMethod CommitMethod, ETextCommit::Type OriginalCommitInfo);
+
+	SLATE_API void NotifyValueCommitted(NumericType CurrentValue) const;
 
 	/** @return true when we are in keyboard-based input mode; false otherwise */
-	bool IsInTextMode() const
-	{
-		return (EditableText->GetVisibility() == EVisibility::Visible);
-	}
+	SLATE_API bool IsInTextMode() const;
 
 	/** Calculates range fraction. Possible to use on full numeric range  */
-	static float Fraction(double InValue, double InMinValue, double InMaxValue)
-	{
-		const double HalfMax = InMaxValue * 0.5;
-		const double HalfMin = InMinValue * 0.5;
-		const double HalfVal = InValue * 0.5;
-
-		return (float)FMath::Clamp((HalfVal - HalfMin) / (HalfMax - HalfMin), 0.0, 1.0);
-	}
+	SLATE_API static float Fraction(double InValue, double InMinValue, double InMaxValue);
 
 private:
 
 	/** The default minimum fractional digits */
-	static const int32 DefaultMinFractionalDigits;
+	static constexpr int32 DefaultMinFractionalDigits = 1;
 
 	/** The default maximum fractional digits */
-	static const int32 DefaultMaxFractionalDigits;
+	static constexpr int32 DefaultMaxFractionalDigits = 6;
 
 	TAttribute<NumericType> ValueAttribute;
 	FOnValueChanged OnValueChanged;
@@ -1161,10 +388,7 @@ private:
 
 	/** True when no range is specified, spinner can be spun indefinitely */
 	bool bUnlimitedSpinRange;
-	void UpdateIsSpinRangeUnlimited()
-	{
-		bUnlimitedSpinRange = !((MinValue.Get().IsSet() && MaxValue.Get().IsSet()) || (MinSliderValue.Get().IsSet() && MaxSliderValue.Get().IsSet()));
-	}
+	void UpdateIsSpinRangeUnlimited();
 
 	const FSpinBoxStyle* Style;
 
@@ -1198,48 +422,19 @@ private:
 
 	/** Prevents the spinbox from being smaller than desired in certain cases (e.g. when it is empty) */
 	TAttribute<float> MinDesiredWidth;
-	float GetTextMinDesiredWidth() const
-	{
-		return FMath::Max(0.0f, MinDesiredWidth.Get() - (float)Style->ArrowsImage.ImageSize.X);
-	}
+	float GetTextMinDesiredWidth() const;
 
 	/** Check whether a typed character is valid */
-	bool IsCharacterValid(TCHAR InChar) const
-	{
-		return Interface->IsCharacterValid(InChar);
-	}
+	bool IsCharacterValid(TCHAR InChar) const;
 
 	/**
 	 * Rounds the submitted value to the correct value if it's an integer.
 	 * For int64, not all values can be represented by a double. We can only round until we reach that limit.
 	 * This function should only be used when we drag the value. We accept that we can't drag huge numbers.
 	 */
-	NumericType RoundIfIntegerValue(double ValueToRound) const
-	{
-		constexpr bool bIsIntegral = TIsIntegral<NumericType>::Value;
-		constexpr bool bCanBeRepresentedInDouble = std::numeric_limits<double>::digits >= std::numeric_limits<NumericType>::digits;
-		if (bIsIntegral && !bCanBeRepresentedInDouble)
-		{
-			return (NumericType)FMath::Clamp<double>(FMath::FloorToDouble(ValueToRound + 0.5), -1.0 * (double)(1ll << std::numeric_limits<double>::digits), (double)(1ll << std::numeric_limits<double>::digits));
-		}
-		else if (bIsIntegral)
-		{
-			return (NumericType)FMath::Clamp<double>(FMath::FloorToDouble(ValueToRound + 0.5), (double)std::numeric_limits<NumericType>::lowest(), (double)std::numeric_limits<NumericType>::max());
-		}
-		else
-		{
-			return (NumericType)FMath::Clamp<double>(ValueToRound, (double)std::numeric_limits<NumericType>::lowest(), (double)std::numeric_limits<NumericType>::max());
-		}
-	}
-
-	void CancelMouseCapture()
-	{
-		bDragging = false;
-		PointerDraggingSliderIndex = INDEX_NONE;
+	NumericType RoundIfIntegerValue(double ValueToRound) const;
 
-		InternalValue = (double)PreDragValue;
-		NotifyValueCommitted(PreDragValue);
-	}
+	void CancelMouseCapture();
 
 	/** Tracks which cursor is currently dragging the slider (e.g., the mouse cursor or a specific finger) */
 	int32 PointerDraggingSliderIndex;
@@ -1267,12 +462,7 @@ private:
 	FString CachedValueString;
 
 	/** Reset the cached string when the value is the same but the display format changed (through the callback). */
-	void ResetCachedValueString() 
-	{ 
-		const NumericType CurrentValue = ValueAttribute.Get();
-		CachedExternalValue = CurrentValue;
-		CachedValueString = Interface->ToString(CachedExternalValue);
-	}
+	void ResetCachedValueString();
 
 	/** Whetever the interfaced setting changed and the CachedValueString needs to be recomputed. */
 	mutable bool bCachedValueStringDirty;
@@ -1296,8 +486,13 @@ private:
 	bool bBroadcastValueChangesPerKey = false;
 };
 
-template<typename NumericType>
-const int32 SSpinBox<NumericType>::DefaultMinFractionalDigits = 1;
-
-template<typename NumericType>
-const int32 SSpinBox<NumericType>::DefaultMaxFractionalDigits = 6;
+extern template class SSpinBox<double>;
+extern template class SSpinBox<float>;
+extern template class SSpinBox<uint64>;
+extern template class SSpinBox<uint32>;
+extern template class SSpinBox<uint16>;
+extern template class SSpinBox<uint8>;
+extern template class SSpinBox<int64>;
+extern template class SSpinBox<int32>;
+extern template class SSpinBox<int16>;
+extern template class SSpinBox<int8>;
\ No newline at end of file
