[Slate] Correctly forwarded VAlign from AddWidget

#rnx

diff --git a/Engine/Source/Runtime/Slate/Private/Framework/MultiBox/MultiBoxBuilder.cpp b/Engine/Source/Runtime/Slate/Private/Framework/MultiBox/MultiBoxBuilder.cpp
--- a/Engine/Source/Runtime/Slate/Private/Framework/MultiBox/MultiBoxBuilder.cpp
+++ b/Engine/Source/Runtime/Slate/Private/Framework/MultiBox/MultiBoxBuilder.cpp
@@ -705,7 +705,7 @@ void FToolBarBuilder::AddWidget(
 	FMenuEntryStyleParams StyleParams = InStyleParams;
 	StyleParams.VerticalAlignment = VAlign_Center; // Default to centered, versions prior to 5.5 decide alignment based on label presence
 
-	AddWidgetInternal(InWidget, InStyleParams, InTutorialHighlightName, bInSearchable, InCustomMenuDelegate, InVisibilityOverride);
+	AddWidgetInternal(InWidget, StyleParams, InTutorialHighlightName, bInSearchable, InCustomMenuDelegate, InVisibilityOverride);
 }
 
 void FToolBarBuilder::AddWidgetInternal(
