SpinBox Optimization
* Defer committing values from MouseMove sources until the widget's Tick to decouple raw mouse input events from sideffects of spinbox value commits

diff --git a/Engine/Source/Runtime/Slate/Private/Widgets/Input/SSpinBox.cpp b/Engine/Source/Runtime/Slate/Private/Widgets/Input/SSpinBox.cpp
--- a/Engine/Source/Runtime/Slate/Private/Widgets/Input/SSpinBox.cpp
+++ b/Engine/Source/Runtime/Slate/Private/Widgets/Input/SSpinBox.cpp
@@ -1,6 +1,17 @@
 // Copyright Epic Games, Inc. All Rights Reserved.
 #include "Widgets/Input/SSpinBox.h"
 
+namespace SpinBoxPrivate
+{
+	bool bUseSpinBoxMouseMoveOptimization = true;
+	FAutoConsoleVariableRef CVar
+		(
+			TEXT("Slate.Spinbox.MouseMoveOptimization"),
+			bUseSpinBoxMouseMoveOptimization,
+			TEXT("")
+		);
+}
+
 float SpinBoxComputeExponentSliderFraction(float FractionFilled, float StartFractionFilled, float SliderExponent)
 {
 	if (FractionFilled <= StartFractionFilled)
diff --git a/Engine/Source/Runtime/Slate/Public/Widgets/Input/SSpinBox.h b/Engine/Source/Runtime/Slate/Public/Widgets/Input/SSpinBox.h
--- a/Engine/Source/Runtime/Slate/Public/Widgets/Input/SSpinBox.h
+++ b/Engine/Source/Runtime/Slate/Public/Widgets/Input/SSpinBox.h
@@ -28,6 +28,11 @@
 
 #include <limits>
 
+namespace SpinBoxPrivate
+{
+	SLATE_API extern bool bUseSpinBoxMouseMoveOptimization;
+}
+
 /*
  * This function compute a slider position by simulating two log on both side of the neutral value
  * Example a slider going from 0.0 to 2.0 with a neutral value of 1.0, the user will have a lot of precision around the neutral value
@@ -403,6 +408,16 @@ public:
 
 		return FMath::Max(FilledLayer, SCompoundWidget::OnPaint(Args, AllottedGeometry, MyCullingRect, OutDrawElements, FilledLayer, InWidgetStyle, bEnabled));
 	}
+	
+	void Tick(const FGeometry& AlottedGeometry, const double InCurrentTime, const float InDeltaTime)
+	{
+		if (PendingCommitValue)
+		{
+			const NumericType RoundedNewValue = RoundIfIntegerValue(PendingCommitValue->NewValue);
+			CommitValue(RoundedNewValue, PendingCommitValue->NewValue, PendingCommitValue->CommitMethod, ETextCommit::OnEnter);
+			PendingCommitValue.Reset();
+		}
+	}
 
 	const bool CommitWithMultiplier(const FPointerEvent& MouseEvent)
 	{
@@ -691,8 +706,26 @@ public:
 					}
 				}
 
-				NumericType RoundedNewValue = RoundIfIntegerValue(NewValue);
-				CommitValue(RoundedNewValue, NewValue, CommitMethod, ETextCommit::OnEnter);
+				if (SpinBoxPrivate::bUseSpinBoxMouseMoveOptimization)
+				{
+					if (CommitMethod == ECommitMethod::CommittedViaSpin)
+					{
+						NewValue = FMath::Clamp<double>(NewValue, (double)GetMinSliderValue(), (double)GetMaxSliderValue());
+					}
+					NewValue = FMath::Clamp<double>(NewValue, (double)GetMinValue(), (double)GetMaxValue());
+					InternalValue = NewValue;
+
+					PendingCommitValue.Emplace(FPendingCommitValue
+							{
+								.NewValue = NewValue,
+								.CommitMethod = CommitMethod
+							}); 
+				}
+				else
+				{
+					NumericType RoundedNewValue = RoundIfIntegerValue(NewValue);
+					CommitValue(RoundedNewValue, NewValue, CommitMethod, ETextCommit::OnEnter);
+				}
 			}
 
 			return FReply::Handled();
@@ -1031,6 +1064,7 @@ protected:
 	 */
 	void CommitValue(NumericType NewValue, double NewSpinValue, ECommitMethod CommitMethod, ETextCommit::Type OriginalCommitInfo)
 	{
+		TRACE_CPUPROFILER_EVENT_SCOPE(SSpinBox_CommitValue);
 		if (CommitMethod == CommittedViaSpin || CommitMethod == CommittedViaArrowKey)
 		{
 			const NumericType LocalMinSliderValue = GetMinSliderValue();
@@ -1141,6 +1175,15 @@ protected:
 	}
 
 private:
+	
+	// New value to be Committed on Tick for MouseMove events
+	// This exists to insulate high-frequency mouse move events which can fire many times during input processing
+	// from the side-effects of committing the spinbox value
+	struct FPendingCommitValue
+	{
+		double NewValue;
+		ECommitMethod CommitMethod;
+	};
 
 	/** The default minimum fractional digits */
 	static const int32 DefaultMinFractionalDigits;
@@ -1294,6 +1337,8 @@ private:
 
 	/** True to broadcast every time we type. */
 	bool bBroadcastValueChangesPerKey = false;
+	
+	TOptional<FPendingCommitValue> PendingCommitValue;
 };
 
 template<typename NumericType>
