Fix non-stretched widgets in toolbars

diff --git a/Engine/Source/Developer/ToolMenus/Private/ToolMenus.cpp b/Engine/Source/Developer/ToolMenus/Private/ToolMenus.cpp
--- a/Engine/Source/Developer/ToolMenus/Private/ToolMenus.cpp
+++ b/Engine/Source/Developer/ToolMenus/Private/ToolMenus.cpp
@@ -1848,6 +1848,11 @@ void UToolMenus::PopulateToolBarBuilderWithEntry(
 
 		FMenuEntryStyleParams StyleParams = Block.WidgetData.StyleParams;
 		StyleParams.HorizontalAlignment = HAlign_Fill;
+		// Default to vertical fill if vertical alignment hasn't been modified for this particular entry.
+		if (!StyleParams.VerticalAlignment.IsSet())
+		{
+			StyleParams.VerticalAlignment = VAlign_Fill;
+		}
 
 		ToolBarBuilder.AddWidget(Widget.ToSharedRef(), StyleParams, Block.TutorialHighlightName, Block.WidgetData.bSearchable, FNewMenuDelegate(), VisibilityOverride);
 	}
diff --git a/Engine/Source/Runtime/Slate/Private/Framework/MultiBox/MultiBoxBuilder.cpp b/Engine/Source/Runtime/Slate/Private/Framework/MultiBox/MultiBoxBuilder.cpp
--- a/Engine/Source/Runtime/Slate/Private/Framework/MultiBox/MultiBoxBuilder.cpp
+++ b/Engine/Source/Runtime/Slate/Private/Framework/MultiBox/MultiBoxBuilder.cpp
@@ -703,7 +703,10 @@ void FToolBarBuilder::AddWidget(
 )
 {
 	FMenuEntryStyleParams StyleParams = InStyleParams;
-	StyleParams.VerticalAlignment = VAlign_Center; // Default to centered, versions prior to 5.5 decide alignment based on label presence
+	if (!StyleParams.VerticalAlignment.IsSet())
+	{
+		StyleParams.VerticalAlignment = VAlign_Center; // Default to centered, versions prior to 5.5 decide alignment based on label presence
+	}
 
 	AddWidgetInternal(InWidget, StyleParams, InTutorialHighlightName, bInSearchable, InCustomMenuDelegate, InVisibilityOverride);
 }
