FLayoutSaveRestore : Move the checking of CanSavePersistentLayouts() into SaveToConfig

Previously, this check was needed by all tools calling FLayoutSaveRestore::SaveToConfig but many of them were not checking it..  This change removes that responsibility from the caller and fixes reset-to-default-layout functionality for several tools ( e.g. Blueprint Debugger, Niagara Debugger ).

diff --git a/Engine/Plugins/Experimental/EditorDataStorageFeatures/Source/TedsDebugger/Private/STedsDebugger.cpp b/Engine/Plugins/Experimental/EditorDataStorageFeatures/Source/TedsDebugger/Private/STedsDebugger.cpp
--- a/Engine/Plugins/Experimental/EditorDataStorageFeatures/Source/TedsDebugger/Private/STedsDebugger.cpp
+++ b/Engine/Plugins/Experimental/EditorDataStorageFeatures/Source/TedsDebugger/Private/STedsDebugger.cpp
@@ -77,12 +77,9 @@ void STedsDebugger::Construct(const FArguments& InArgs, const TSharedRef<SDockTa
 		FTabManager::FOnPersistLayout::CreateStatic(
 			[](const TSharedRef<FTabManager::FLayout>& InLayout)
 			{
-				if (FGlobalTabmanager::Get()->CanSavePersistentLayouts())
+				if (InLayout->GetPrimaryArea().Pin().IsValid())
 				{
-					if (InLayout->GetPrimaryArea().Pin().IsValid())
-					{
-						FLayoutSaveRestore::SaveToConfig(GEditorLayoutIni, InLayout);
-					}
+					FLayoutSaveRestore::SaveToConfig(GEditorLayoutIni, InLayout);
 				}
 			}
 		)
diff --git a/Engine/Source/Developer/SlateReflector/Private/Widgets/SWidgetReflector.cpp b/Engine/Source/Developer/SlateReflector/Private/Widgets/SWidgetReflector.cpp
--- a/Engine/Source/Developer/SlateReflector/Private/Widgets/SWidgetReflector.cpp
+++ b/Engine/Source/Developer/SlateReflector/Private/Widgets/SWidgetReflector.cpp
@@ -1065,10 +1065,7 @@ TSharedRef<SDockTab> SWidgetReflector::SpawnWidgetList(const FSpawnTabArgs& Args
 void SWidgetReflector::HandleTabManagerPersistLayout(const TSharedRef<FTabManager::FLayout>& LayoutToSave)
 {
 #if WITH_EDITOR
-	if (FGlobalTabmanager::Get()->CanSavePersistentLayouts())
-	{
-		FLayoutSaveRestore::SaveToConfig(GEditorLayoutIni, LayoutToSave);
-	}
+	FLayoutSaveRestore::SaveToConfig(GEditorLayoutIni, LayoutToSave);
 #endif //WITH_EDITOR
 }
 
diff --git a/Engine/Source/Editor/LevelEditor/Public/LevelEditor.h b/Engine/Source/Editor/LevelEditor/Public/LevelEditor.h
--- a/Engine/Source/Editor/LevelEditor/Public/LevelEditor.h
+++ b/Engine/Source/Editor/LevelEditor/Public/LevelEditor.h
@@ -432,10 +432,7 @@ private:
 	// Callback for persisting the Level Editor's layout.
 	void HandleTabManagerPersistLayout( const TSharedRef<FTabManager::FLayout>& LayoutToSave )
 	{
-		if (FGlobalTabmanager::Get()->CanSavePersistentLayouts())
-		{
-			FLayoutSaveRestore::SaveToConfig(GEditorLayoutIni, LayoutToSave);
-		}
+		FLayoutSaveRestore::SaveToConfig(GEditorLayoutIni, LayoutToSave);
 	}
 
 	/** Spawn the main level editor tab */
diff --git a/Engine/Source/Editor/MainFrame/Private/Frame/MainFrameHandler.h b/Engine/Source/Editor/MainFrame/Private/Frame/MainFrameHandler.h
--- a/Engine/Source/Editor/MainFrame/Private/Frame/MainFrameHandler.h
+++ b/Engine/Source/Editor/MainFrame/Private/Frame/MainFrameHandler.h
@@ -258,10 +258,7 @@ private:
 	// Callback for persisting the Level Editor's layout.
 	void HandleTabManagerPersistLayout( const TSharedRef<FTabManager::FLayout>& LayoutToSave )
 	{
-		if (FGlobalTabmanager::Get()->CanSavePersistentLayouts())
-		{
-			FLayoutSaveRestore::SaveToConfig(GEditorLayoutIni, LayoutToSave);
-		}
+		FLayoutSaveRestore::SaveToConfig(GEditorLayoutIni, LayoutToSave);
 	}
 
 private:
diff --git a/Engine/Source/Editor/MassEntityDebugger/Private/SMassDebugger.cpp b/Engine/Source/Editor/MassEntityDebugger/Private/SMassDebugger.cpp
--- a/Engine/Source/Editor/MassEntityDebugger/Private/SMassDebugger.cpp
+++ b/Engine/Source/Editor/MassEntityDebugger/Private/SMassDebugger.cpp
@@ -139,13 +139,10 @@ void SMassDebugger::Construct(const FArguments& InArgs, const TSharedRef<SDockTa
 		FTabManager::FOnPersistLayout::CreateStatic(
 			[](const TSharedRef<FTabManager::FLayout>& InLayout)
 			{
-				if (FGlobalTabmanager::Get()->CanSavePersistentLayouts())
+				if (InLayout->GetPrimaryArea().Pin().IsValid())
 				{
-					if (InLayout->GetPrimaryArea().Pin().IsValid())
-					{
-						FLayoutSaveRestore::SaveToConfig(GEditorLayoutIni, InLayout);
-					}
-				}				
+					FLayoutSaveRestore::SaveToConfig(GEditorLayoutIni, InLayout);
+				}
 			}
 		)
 	);
diff --git a/Engine/Source/Editor/UnrealEd/Private/ApplicationMode.cpp b/Engine/Source/Editor/UnrealEd/Private/ApplicationMode.cpp
--- a/Engine/Source/Editor/UnrealEd/Private/ApplicationMode.cpp
+++ b/Engine/Source/Editor/UnrealEd/Private/ApplicationMode.cpp
@@ -30,10 +30,8 @@ void FApplicationMode::DeactivateMode(TSharedPtr<FTabManager> InTabManager)
 {
 	// Save the layout to INI
 	check(InTabManager.IsValid());
-	if (FGlobalTabmanager::Get()->CanSavePersistentLayouts())
-	{
-		FLayoutSaveRestore::SaveToConfig(GEditorLayoutIni, InTabManager->PersistLayout());
-	}
+	FLayoutSaveRestore::SaveToConfig(GEditorLayoutIni, InTabManager->PersistLayout());
+
 	// Unregister the tabs
 	/*
 	for (int32 Index = 0; Index < AllowableTabs.Num(); ++Index)
diff --git a/Engine/Source/Editor/UnrealEd/Public/Toolkits/AssetEditorToolkit.h b/Engine/Source/Editor/UnrealEd/Public/Toolkits/AssetEditorToolkit.h
--- a/Engine/Source/Editor/UnrealEd/Public/Toolkits/AssetEditorToolkit.h
+++ b/Engine/Source/Editor/UnrealEd/Public/Toolkits/AssetEditorToolkit.h
@@ -446,10 +446,7 @@ private:
 	// Callback for persisting the Asset Editor's layout.
 	void HandleTabManagerPersistLayout( const TSharedRef<FTabManager::FLayout>& LayoutToSave )
 	{
-		if (FGlobalTabmanager::Get()->CanSavePersistentLayouts())
-		{
-			FLayoutSaveRestore::SaveToConfig(GEditorLayoutIni, LayoutToSave);
-		}
+		FLayoutSaveRestore::SaveToConfig(GEditorLayoutIni, LayoutToSave);
 	}
 
 	/** Initialize the base level customizations for read only mode and register the profile */
diff --git a/Engine/Source/Runtime/Slate/Private/Framework/Docking/LayoutService.cpp b/Engine/Source/Runtime/Slate/Private/Framework/Docking/LayoutService.cpp
--- a/Engine/Source/Runtime/Slate/Private/Framework/Docking/LayoutService.cpp
+++ b/Engine/Source/Runtime/Slate/Private/Framework/Docking/LayoutService.cpp
@@ -214,13 +214,16 @@ static bool LoadLayoutFromJson(const FString& InConfigFileName, const FString& I
 
 void FLayoutSaveRestore::SaveToConfig( const FString& InConfigFileName, const TSharedRef<FTabManager::FLayout>& InLayoutToSave )
 {
-	// Only save to config if it's not the FTabManager::FLayout::NullLayout
-	if (InLayoutToSave->GetLayoutName() != FTabManager::FLayout::NullLayout->GetLayoutName())
+	if (FGlobalTabmanager::Get()->CanSavePersistentLayouts())
 	{
-		const FString LayoutAsString = PrepareLayoutStringForIni(InLayoutToSave->ToString());
-		GConfig->SetString(GetEditorLayoutsSectionName(), *InLayoutToSave->GetLayoutName().ToString(), *LayoutAsString, InConfigFileName);
+		// Only save to config if it's not the FTabManager::FLayout::NullLayout
+		if (InLayoutToSave->GetLayoutName() != FTabManager::FLayout::NullLayout->GetLayoutName())
+		{
+			const FString LayoutAsString = PrepareLayoutStringForIni(InLayoutToSave->ToString());
+			GConfig->SetString(GetEditorLayoutsSectionName(), *InLayoutToSave->GetLayoutName().ToString(), *LayoutAsString, InConfigFileName);
 
-		SaveLayoutToJson(InConfigFileName, InLayoutToSave);
+			SaveLayoutToJson(InConfigFileName, InLayoutToSave);
+		}
 	}
 }
 
