Avoid padding inside Combo Buttons between icon and label when label is empty

diff --git a/Engine/Source/Runtime/Slate/Private/Framework/MultiBox/SToolBarComboButtonBlock.cpp b/Engine/Source/Runtime/Slate/Private/Framework/MultiBox/SToolBarComboButtonBlock.cpp
--- a/Engine/Source/Runtime/Slate/Private/Framework/MultiBox/SToolBarComboButtonBlock.cpp
+++ b/Engine/Source/Runtime/Slate/Private/Framework/MultiBox/SToolBarComboButtonBlock.cpp
@@ -146,12 +146,13 @@ void SToolBarComboButtonBlock::BuildMultiBlockWidget(const ISlateStyle* StyleSet
 	{
 		TagName = ToolBarComboButtonBlock->GetTutorialHighlightName();
 	}
-	
+
 	// Create the content for our button
 	TSharedRef<SWidget> ButtonContent = SNullWidget::NullWidget;
 	if (MultiBox->GetType() == EMultiBoxType::SlimHorizontalToolBar
 		|| MultiBox->GetType() == EMultiBoxType::SlimHorizontalUniformToolBar)
 	{
+		// clang-format off
 		ButtonContent =
 			SNew(SHorizontalBox)
 			// Icon image
@@ -165,7 +166,12 @@ void SToolBarComboButtonBlock::BuildMultiBlockWidget(const ISlateStyle* StyleSet
 			// Label text
 			+ SHorizontalBox::Slot()
 			.AutoWidth()
-			.Padding(ToolBarStyle.LabelPadding)
+			.Padding(TAttribute<FMargin>::CreateLambda(
+				[Label, ToolBarStyle]() -> FMargin
+				{
+					return Label.Get().IsEmpty() ? FMargin(0) : ToolBarStyle.LabelPadding;
+				}
+			))
 			.VAlign(VAlign_Center)	// Center the label text horizontally
 			[
 				SNew(STextBlock)
@@ -174,7 +180,7 @@ void SToolBarComboButtonBlock::BuildMultiBlockWidget(const ISlateStyle* StyleSet
 				// Smaller font for tool tip labels
 				.TextStyle(&ToolBarStyle.LabelStyle)
 			];
-
+		// clang-format on
 	}
 	else
 	{
@@ -191,7 +197,12 @@ void SToolBarComboButtonBlock::BuildMultiBlockWidget(const ISlateStyle* StyleSet
 			// Label text
 			+ SVerticalBox::Slot()
 			.AutoHeight()
-			.Padding(ToolBarStyle.LabelPadding)
+			.Padding(TAttribute<FMargin>::CreateLambda(
+				[Label, ToolBarStyle]() -> FMargin
+				{
+					return Label.Get().IsEmpty() ? FMargin(0) : ToolBarStyle.LabelPadding;
+				}
+			))
 			.HAlign(HAlign_Center)	// Center the label text horizontally
 			[
 				SNew(STextBlock)
