[Slate] Added optional DesiredWidth/Height overrides to FMenuEntryStyleParams. Effectively allows entries with min/max sizes to default to to their max, rather than whatever desired-size reports.

#rnx

diff --git a/Engine/Source/Runtime/Slate/Private/Framework/MultiBox/MultiBoxBuilder.cpp b/Engine/Source/Runtime/Slate/Private/Framework/MultiBox/MultiBoxBuilder.cpp
--- a/Engine/Source/Runtime/Slate/Private/Framework/MultiBox/MultiBoxBuilder.cpp
+++ b/Engine/Source/Runtime/Slate/Private/Framework/MultiBox/MultiBoxBuilder.cpp
@@ -725,6 +725,8 @@ void FToolBarBuilder::AddWidgetInternal(
 	TSharedRef<SWidget> ChildWidget = InWidget;
 	InWidget =
 		SNew( SBox )
+		.WidthOverride(InStyleParams.DesiredWidth)
+		.HeightOverride(InStyleParams.DesiredHeight)
 		.AddMetaData<FTagMetaData>(FTagMetaData(InTutorialHighlightName))
 		[
 			ChildWidget
diff --git a/Engine/Source/Runtime/Slate/Public/Framework/MultiBox/MultiBoxDefs.h b/Engine/Source/Runtime/Slate/Public/Framework/MultiBox/MultiBoxDefs.h
--- a/Engine/Source/Runtime/Slate/Public/Framework/MultiBox/MultiBoxDefs.h
+++ b/Engine/Source/Runtime/Slate/Public/Framework/MultiBox/MultiBoxDefs.h
@@ -161,4 +161,10 @@ struct FMenuEntryStyleParams
 
 	/** (Optionally) override the maximum size. This will apply to the width or height, depending on the menu orientation */
 	TOptional<float> MaxSize;
+
+	/** (Optionally) set the desired width override */
+	TOptional<float> DesiredWidth;
+
+	/** (Optionally) set the desired height override */
+	TOptional<float> DesiredHeight;
 };
