Slate Text Highlights: Replace OffsetX with an Offset Vector2D and set the Y offset based on the difference between the block's LocationOffset and the line's offset. This lets us properly align highlights on lines of text that contain multiple font sizes

diff --git a/Engine/Source/Developer/OutputLog/Private/SOutputLog.cpp b/Engine/Source/Developer/OutputLog/Private/SOutputLog.cpp
--- a/Engine/Source/Developer/OutputLog/Private/SOutputLog.cpp
+++ b/Engine/Source/Developer/OutputLog/Private/SOutputLog.cpp
@@ -51,9 +51,9 @@ public:
 		return MakeShareable(new FCategoryLineHighlighter());
 	}
 
-	virtual int32 OnPaint(const FPaintArgs& Args, const FTextLayout::FLineView& Line, const float OffsetX, const float Width, const FTextBlockStyle& DefaultStyle, const FGeometry& AllottedGeometry, const FSlateRect& MyCullingRect, FSlateWindowElementList& OutDrawElements, int32 LayerId, const FWidgetStyle& InWidgetStyle, bool bParentEnabled) const override
+	virtual int32 OnPaint(const FPaintArgs& Args, const FTextLayout::FLineView& Line, const FVector2D Offset, const float Width, const FTextBlockStyle& DefaultStyle, const FGeometry& AllottedGeometry, const FSlateRect& MyCullingRect, FSlateWindowElementList& OutDrawElements, int32 LayerId, const FWidgetStyle& InWidgetStyle, bool bParentEnabled) const override
 	{
-		const FVector2D Location(Line.Offset.X + OffsetX, Line.Offset.Y);
+		const FVector2D Location(Line.Offset.X + Offset.X, Line.Offset.Y + Offset.Y);
 
 		// If we've not been set to an explicit color, calculate a suitable one from the linked color
 		FLinearColor SelectionBackgroundColorAndOpacity = DefaultStyle.SelectedBackgroundColor.GetColor(InWidgetStyle);// *InWidgetStyle.GetColorAndOpacityTint();
@@ -93,9 +93,9 @@ public:
 		return MakeShareable(new FCategoryBadgeHighlighter(InBadgeColor));
 	}
 
-	virtual int32 OnPaint(const FPaintArgs& Args, const FTextLayout::FLineView& Line, const float OffsetX, const float Width, const FTextBlockStyle& DefaultStyle, const FGeometry& AllottedGeometry, const FSlateRect& MyCullingRect, FSlateWindowElementList& OutDrawElements, int32 LayerId, const FWidgetStyle& InWidgetStyle, bool bParentEnabled) const override
+	virtual int32 OnPaint(const FPaintArgs& Args, const FTextLayout::FLineView& Line, const FVector2D Offset, const float Width, const FTextBlockStyle& DefaultStyle, const FGeometry& AllottedGeometry, const FSlateRect& MyCullingRect, FSlateWindowElementList& OutDrawElements, int32 LayerId, const FWidgetStyle& InWidgetStyle, bool bParentEnabled) const override
 	{
-		const FVector2D Location(Line.Offset.X + OffsetX, Line.Offset.Y);
+		const FVector2D Location(Line.Offset.X + Offset.X, Line.Offset.Y + Offset.Y);
 
 		// The block size and offset values are pre-scaled, so we need to account for that when converting the block offsets into paint geometry
 		const float InverseScale = Inverse(AllottedGeometry.Scale);
diff --git a/Engine/Source/Runtime/Slate/Private/Framework/Text/SlateTextLayout.cpp b/Engine/Source/Runtime/Slate/Private/Framework/Text/SlateTextLayout.cpp
--- a/Engine/Source/Runtime/Slate/Private/Framework/Text/SlateTextLayout.cpp
+++ b/Engine/Source/Runtime/Slate/Private/Framework/Text/SlateTextLayout.cpp
@@ -265,7 +265,7 @@ int32 FSlateTextLayout::OnPaintHighlights( const FPaintArgs& Args, const FTextLa
 		const TSharedPtr< ISlateLineHighlighter > LineHighlighter = StaticCastSharedPtr< ISlateLineHighlighter >( Highlight.Highlighter );
 		if (LineHighlighter.IsValid())
 		{
-			CurrentLayerId = LineHighlighter->OnPaint( Args, LineView, Highlight.OffsetX, Highlight.Width, InDefaultTextStyle, AllottedGeometry, CullingRect, OutDrawElements, CurrentLayerId, InWidgetStyle, bParentEnabled );
+			CurrentLayerId = LineHighlighter->OnPaint( Args, LineView, Highlight.Offset, Highlight.Width, InDefaultTextStyle, AllottedGeometry, CullingRect, OutDrawElements, CurrentLayerId, InWidgetStyle, bParentEnabled );
 		}
 	}
 
diff --git a/Engine/Source/Runtime/Slate/Private/Framework/Text/SlateTextUnderlineLineHighlighter.cpp b/Engine/Source/Runtime/Slate/Private/Framework/Text/SlateTextUnderlineLineHighlighter.cpp
--- a/Engine/Source/Runtime/Slate/Private/Framework/Text/SlateTextUnderlineLineHighlighter.cpp
+++ b/Engine/Source/Runtime/Slate/Private/Framework/Text/SlateTextUnderlineLineHighlighter.cpp
@@ -15,7 +15,7 @@ ISlateTextLineHighlighter::ISlateTextLineHighlighter(const FSlateBrush& InLineBr
 {
 }
 
-int32 ISlateTextLineHighlighter::OnPaint(const FPaintArgs& Args, const FTextLayout::FLineView& Line, const float OffsetX, const float Width, const FTextBlockStyle& DefaultStyle, const FGeometry& AllottedGeometry, const FSlateRect& MyCullingRect, FSlateWindowElementList& OutDrawElements, int32 LayerId, const FWidgetStyle& InWidgetStyle, bool bParentEnabled) const
+int32 ISlateTextLineHighlighter::OnPaint(const FPaintArgs& Args, const FTextLayout::FLineView& Line, const FVector2D Offset, const float Width, const FTextBlockStyle& DefaultStyle, const FGeometry& AllottedGeometry, const FSlateRect& MyCullingRect, FSlateWindowElementList& OutDrawElements, int32 LayerId, const FWidgetStyle& InWidgetStyle, bool bParentEnabled) const
 {
 	TSharedRef<FSlateFontCache> FontCache = FSlateApplication::Get().GetRenderer()->GetFontCache();
 
@@ -25,7 +25,7 @@ int32 ISlateTextLineHighlighter::OnPaint(const FPaintArgs& Args, const FTextLayo
 	int16 LinePos, LineThickness;
 	GetLineMetrics(AllottedGeometry.Scale, LinePos, LineThickness);
 
-	const FVector2f Location(Line.Offset.X + OffsetX, Line.Offset.Y + MaxHeight + Baseline - (LinePos * 0.5f));
+	const FVector2f Location(Line.Offset.X + Offset.X, Line.Offset.Y + Offset.Y + MaxHeight + Baseline - (LinePos * 0.5f));
 	const FVector2f Size(Width, FMath::Max<int16>(1, LineThickness));
 
 	// The block size and offset values are pre-scaled, so we need to account for that when converting the block offsets into paint geometry
diff --git a/Engine/Source/Runtime/Slate/Private/Framework/Text/TextLayout.cpp b/Engine/Source/Runtime/Slate/Private/Framework/Text/TextLayout.cpp
--- a/Engine/Source/Runtime/Slate/Private/Framework/Text/TextLayout.cpp
+++ b/Engine/Source/Runtime/Slate/Private/Framework/Text/TextLayout.cpp
@@ -808,12 +808,12 @@ void FTextLayout::FlowHighlights()
 			}
 
 			FLineViewHighlight LineViewHighlight;
-			LineViewHighlight.OffsetX = 0.0f;
+			LineViewHighlight.Offset = FVector2D(0.0f, 0.0f);
 			LineViewHighlight.Width = 0.0f;
 			LineViewHighlight.Highlighter = LineHighlight.Highlighter;
 
 			bool bHasStartedHighlighting = false;
-			float RunningBlockOffset = LineViewHighlight.OffsetX;
+			float RunningBlockOffset = LineViewHighlight.Offset.X;
 
 			auto AppendLineViewHighlight = [&LineHighlight, &LineView, &LineViewHighlight, &LineViewHighlightRange]()
 			{
@@ -834,7 +834,7 @@ void FTextLayout::FlowHighlights()
 			{
 				AppendLineViewHighlight();
 
-				LineViewHighlight.OffsetX = RunningBlockOffset;
+				LineViewHighlight.Offset.X = RunningBlockOffset;
 				LineViewHighlight.Width = 0.0f;
 			};
 
@@ -859,7 +859,8 @@ void FTextLayout::FlowHighlights()
 						// This block is the start of the highlight
 						bHasStartedHighlighting = true;
 
-						LineViewHighlight.OffsetX = CurrentBlockRunningOffset;
+						LineViewHighlight.Offset.X = CurrentBlockRunningOffset;
+						LineViewHighlight.Offset.Y = Block->GetLocationOffset().Y - LineView.Offset.Y;
 
 						// The width always includes size of the intersecting text
 						if (!IntersectedRange.IsEmpty())
@@ -870,18 +871,25 @@ void FTextLayout::FlowHighlights()
 							// In RTL text, the space after the end of the text (which is visually on the left) is added as an offset
 							if (Block->GetTextContext().TextDirection == TextBiDi::ETextDirection::LeftToRight)
 							{
-								LineViewHighlight.OffsetX += Run->Measure(BlockTextRange.BeginIndex, IntersectedRange.BeginIndex, Scale, RunTextContext).X;
+								LineViewHighlight.Offset.X += Run->Measure(BlockTextRange.BeginIndex, IntersectedRange.BeginIndex, Scale, RunTextContext).X;
 							}
 							else
 							{
-								LineViewHighlight.OffsetX += Run->Measure(IntersectedRange.EndIndex, BlockTextRange.EndIndex, Scale, RunTextContext).X;
+								LineViewHighlight.Offset.X += Run->Measure(IntersectedRange.EndIndex, BlockTextRange.EndIndex, Scale, RunTextContext).X;
 							}
 						}
 					}
 					else
 					{
+
+						// For the y offset, we'll use the smallest LocationOffset included in the highlighted area. This should align any highlights with the tallest block
+						if (Block->GetLocationOffset().Y < LineViewHighlight.Offset.Y + LineView.Offset.Y)
+						{
+							LineViewHighlight.Offset.Y = Block->GetLocationOffset().Y - LineView.Offset.Y;
+						}
+
 						// Test to see whether the current highlight covered the RHS of the previous block
-						const bool bPrevBlockHighlightedRHS = (LineViewHighlight.OffsetX + LineViewHighlight.Width) >= CurrentBlockRunningOffset;
+						const bool bPrevBlockHighlightedRHS = (LineViewHighlight.Offset.X + LineViewHighlight.Width) >= CurrentBlockRunningOffset;
 
 						// This block is part of an existing highlight... but should it split the current visual highlight?
 						if (bBlockIsFullyHighlighted)
@@ -896,7 +904,7 @@ void FTextLayout::FlowHighlights()
 								// Fully highlighted blocks following a visual gap highlighted block have to start a new visual highlight
 								AppendLineViewHighlightAndReset();
 
-								LineViewHighlight.OffsetX = CurrentBlockRunningOffset;
+								LineViewHighlight.Offset.X = CurrentBlockRunningOffset;
 								LineViewHighlight.Width += Block->GetSize().X;
 							}
 						}
@@ -941,7 +949,7 @@ void FTextLayout::FlowHighlights()
 							{
 								// There's a visual gap, so we need to end this visual highlight and start another
 								AppendLineViewHighlightAndReset();
-								LineViewHighlight.OffsetX = CurrentBlockRunningOffset + BlockHighlightOffsetAdjustment;
+								LineViewHighlight.Offset.X = CurrentBlockRunningOffset + BlockHighlightOffsetAdjustment;
 								LineViewHighlight.Width += BlockHighlightWidth;
 							}
 						}
diff --git a/Engine/Source/Runtime/Slate/Private/Widgets/Text/SlateEditableTextTypes.cpp b/Engine/Source/Runtime/Slate/Private/Widgets/Text/SlateEditableTextTypes.cpp
--- a/Engine/Source/Runtime/Slate/Private/Widgets/Text/SlateEditableTextTypes.cpp
+++ b/Engine/Source/Runtime/Slate/Private/Widgets/Text/SlateEditableTextTypes.cpp
@@ -103,9 +103,9 @@ void FCursorLineHighlighter::SetCursorBrush(const TAttribute<const FSlateBrush*>
 	CursorBrush = InCursorBrush;
 }
 
-int32 FCursorLineHighlighter::OnPaint(const FPaintArgs& Args, const FTextLayout::FLineView& Line, const float OffsetX, const float Width, const FTextBlockStyle& DefaultStyle, const FGeometry& AllottedGeometry, const FSlateRect& MyCullingRect, FSlateWindowElementList& OutDrawElements, int32 LayerId, const FWidgetStyle& InWidgetStyle, bool bParentEnabled) const
+int32 FCursorLineHighlighter::OnPaint(const FPaintArgs& Args, const FTextLayout::FLineView& Line, const FVector2D Offset, const float Width, const FTextBlockStyle& DefaultStyle, const FGeometry& AllottedGeometry, const FSlateRect& MyCullingRect, FSlateWindowElementList& OutDrawElements, int32 LayerId, const FWidgetStyle& InWidgetStyle, bool bParentEnabled) const
 {
-	const FVector2D Location(Line.Offset.X + OffsetX, Line.Offset.Y);
+	const FVector2D Location(Line.Offset.X + Offset.X, Line.Offset.Y + Offset.Y);
 	const FVector2D Size(Width, Line.TextHeight);
 
 	FLinearColor CursorColorAndOpacity = InWidgetStyle.GetForegroundColor();
@@ -162,9 +162,9 @@ void FTextCompositionHighlighter::SetCompositionBrush(const TAttribute<const FSl
 	CompositionBrush = InCompositionBrush;
 }
 
-int32 FTextCompositionHighlighter::OnPaint(const FPaintArgs& Args, const FTextLayout::FLineView& Line, const float OffsetX, const float Width, const FTextBlockStyle& DefaultStyle, const FGeometry& AllottedGeometry, const FSlateRect& MyCullingRect, FSlateWindowElementList& OutDrawElements, int32 LayerId, const FWidgetStyle& InWidgetStyle, bool bParentEnabled) const
+int32 FTextCompositionHighlighter::OnPaint(const FPaintArgs& Args, const FTextLayout::FLineView& Line, const FVector2D Offset, const float Width, const FTextBlockStyle& DefaultStyle, const FGeometry& AllottedGeometry, const FSlateRect& MyCullingRect, FSlateWindowElementList& OutDrawElements, int32 LayerId, const FWidgetStyle& InWidgetStyle, bool bParentEnabled) const
 {
-	const FVector2D Location(Line.Offset.X + OffsetX, Line.Offset.Y);
+	const FVector2D Location(Line.Offset.X + Offset.X, Line.Offset.Y + Offset.Y);
 	const FVector2D Size(Width, Line.TextHeight);
 
 	// The block size and offset values are pre-scaled, so we need to account for that when converting the block offsets into paint geometry
@@ -197,9 +197,9 @@ FTextSelectionHighlighter::FTextSelectionHighlighter()
 {
 }
 
-int32 FTextSelectionHighlighter::OnPaint(const FPaintArgs& Args, const FTextLayout::FLineView& Line, const float OffsetX, const float Width, const FTextBlockStyle& DefaultStyle, const FGeometry& AllottedGeometry, const FSlateRect& MyCullingRect, FSlateWindowElementList& OutDrawElements, int32 LayerId, const FWidgetStyle& InWidgetStyle, bool bParentEnabled) const
+int32 FTextSelectionHighlighter::OnPaint(const FPaintArgs& Args, const FTextLayout::FLineView& Line, const FVector2D Offset, const float Width, const FTextBlockStyle& DefaultStyle, const FGeometry& AllottedGeometry, const FSlateRect& MyCullingRect, FSlateWindowElementList& OutDrawElements, int32 LayerId, const FWidgetStyle& InWidgetStyle, bool bParentEnabled) const
 {
-	const FVector2D Location(Line.Offset.X + OffsetX, Line.Offset.Y);
+	const FVector2D Location(Line.Offset.X + Offset.X, Line.Offset.Y + Offset.Y);
 
 	// If we've not been set to an explicit color, calculate a suitable one from the linked color
 	const FLinearColor SelectionBackgroundColorAndOpacity = DefaultStyle.SelectedBackgroundColor.IsColorSpecified()
@@ -237,9 +237,9 @@ FTextSearchHighlighter::FTextSearchHighlighter()
 {
 }
 
-int32 FTextSearchHighlighter::OnPaint(const FPaintArgs& Args, const FTextLayout::FLineView& Line, const float OffsetX, const float Width, const FTextBlockStyle& DefaultStyle, const FGeometry& AllottedGeometry, const FSlateRect& MyCullingRect, FSlateWindowElementList& OutDrawElements, int32 LayerId, const FWidgetStyle& InWidgetStyle, bool bParentEnabled) const
+int32 FTextSearchHighlighter::OnPaint(const FPaintArgs& Args, const FTextLayout::FLineView& Line, const FVector2D Offset, const float Width, const FTextBlockStyle& DefaultStyle, const FGeometry& AllottedGeometry, const FSlateRect& MyCullingRect, FSlateWindowElementList& OutDrawElements, int32 LayerId, const FWidgetStyle& InWidgetStyle, bool bParentEnabled) const
 {
-	const FVector2D Location(Line.Offset.X + OffsetX, Line.Offset.Y);
+	const FVector2D Location(Line.Offset.X + Offset.X, Line.Offset.Y + Offset.Y);
 
 	// If we've not been set to an explicit color, calculate a suitable one from the linked color
 	FLinearColor SelectionBackgroundColorAndOpacity = DefaultStyle.HighlightColor.GetColor(InWidgetStyle) * InWidgetStyle.GetColorAndOpacityTint();
diff --git a/Engine/Source/Runtime/Slate/Public/Framework/Text/ISlateLineHighlighter.h b/Engine/Source/Runtime/Slate/Public/Framework/Text/ISlateLineHighlighter.h
--- a/Engine/Source/Runtime/Slate/Public/Framework/Text/ISlateLineHighlighter.h
+++ b/Engine/Source/Runtime/Slate/Public/Framework/Text/ISlateLineHighlighter.h
@@ -17,5 +17,11 @@ class ISlateLineHighlighter : public ILineHighlighter
 public:
 	virtual ~ISlateLineHighlighter() {}
 
-	virtual int32 OnPaint(const FPaintArgs& Args, const FTextLayout::FLineView& Line, const float OffsetX, const float Width, const FTextBlockStyle& DefaultStyle, const FGeometry& AllottedGeometry, const FSlateRect& MyCullingRect, FSlateWindowElementList& OutDrawElements, int32 LayerId, const FWidgetStyle& InWidgetStyle, bool bParentEnabled) const = 0;
+	UE_DEPRECATED(5.6, "Please use OnPaint with Offset (FVector2D) instead of OffsetX (float).")
+	virtual int32 OnPaint(const FPaintArgs& Args, const FTextLayout::FLineView& Line, const float OffsetX, const float Width, const FTextBlockStyle& DefaultStyle, const FGeometry& AllottedGeometry, const FSlateRect& MyCullingRect, FSlateWindowElementList& OutDrawElements, int32 LayerId, const FWidgetStyle& InWidgetStyle, bool bParentEnabled) const final
+	{
+		return OnPaint(Args, Line, FVector2D(OffsetX, 0), Width, DefaultStyle, AllottedGeometry, MyCullingRect, OutDrawElements, LayerId, InWidgetStyle, bParentEnabled);
+	}
+	
+	virtual int32 OnPaint(const FPaintArgs& Args, const FTextLayout::FLineView& Line, const FVector2D Offset, const float Width, const FTextBlockStyle& DefaultStyle, const FGeometry& AllottedGeometry, const FSlateRect& MyCullingRect, FSlateWindowElementList& OutDrawElements, int32 LayerId, const FWidgetStyle& InWidgetStyle, bool bParentEnabled) const = 0;
 };
diff --git a/Engine/Source/Runtime/Slate/Public/Framework/Text/SlateTextUnderlineLineHighlighter.h b/Engine/Source/Runtime/Slate/Public/Framework/Text/SlateTextUnderlineLineHighlighter.h
--- a/Engine/Source/Runtime/Slate/Public/Framework/Text/SlateTextUnderlineLineHighlighter.h
+++ b/Engine/Source/Runtime/Slate/Public/Framework/Text/SlateTextUnderlineLineHighlighter.h
@@ -18,7 +18,7 @@ struct FTextBlockStyle;
 class ISlateTextLineHighlighter : public ISlateLineHighlighter
 {
 public:
-	SLATE_API virtual int32 OnPaint(const FPaintArgs& Args, const FTextLayout::FLineView& Line, const float OffsetX, const float Width, const FTextBlockStyle& DefaultStyle, const FGeometry& AllottedGeometry, const FSlateRect& MyCullingRect, FSlateWindowElementList& OutDrawElements, int32 LayerId, const FWidgetStyle& InWidgetStyle, bool bParentEnabled) const override;
+	SLATE_API virtual int32 OnPaint(const FPaintArgs& Args, const FTextLayout::FLineView& Line, const FVector2D Offset, const float Width, const FTextBlockStyle& DefaultStyle, const FGeometry& AllottedGeometry, const FSlateRect& MyCullingRect, FSlateWindowElementList& OutDrawElements, int32 LayerId, const FWidgetStyle& InWidgetStyle, bool bParentEnabled) const override;
 
 	static const int32 DefaultZIndex = 1;
 
diff --git a/Engine/Source/Runtime/Slate/Public/Framework/Text/TextLayout.h b/Engine/Source/Runtime/Slate/Public/Framework/Text/TextLayout.h
--- a/Engine/Source/Runtime/Slate/Public/Framework/Text/TextLayout.h
+++ b/Engine/Source/Runtime/Slate/Public/Framework/Text/TextLayout.h
@@ -316,8 +316,8 @@ public:
 
 	struct FLineViewHighlight
 	{
-		/** Offset in X for this highlight, relative to the FLineView::Offset that contains it */
-		float OffsetX;
+		/** Offset for this highlight, relative to the FLineView::Offset that contains it */
+		FVector2D Offset;
 		/** Width for this highlight, the height will be either FLineView::Size.Y or FLineView::TextHeight depending on whether you want to highlight the entire line, or just the text within the line */
 		float Width;
 		/** Custom highlighter implementation used to do the painting */
diff --git a/Engine/Source/Runtime/Slate/Public/Widgets/Text/SlateEditableTextTypes.h b/Engine/Source/Runtime/Slate/Public/Widgets/Text/SlateEditableTextTypes.h
--- a/Engine/Source/Runtime/Slate/Public/Widgets/Text/SlateEditableTextTypes.h
+++ b/Engine/Source/Runtime/Slate/Public/Widgets/Text/SlateEditableTextTypes.h
@@ -141,7 +141,7 @@ public:
 
 	SLATE_API void SetCursorBrush(const TAttribute<const FSlateBrush*>& InCursorBrush);
 
-	SLATE_API virtual int32 OnPaint(const FPaintArgs& Args, const FTextLayout::FLineView& Line, const float OffsetX, const float Width, const FTextBlockStyle& DefaultStyle, const FGeometry& AllottedGeometry, const FSlateRect& MyCullingRect, FSlateWindowElementList& OutDrawElements, int32 LayerId, const FWidgetStyle& InWidgetStyle, bool bParentEnabled) const override;
+	SLATE_API virtual int32 OnPaint(const FPaintArgs& Args, const FTextLayout::FLineView& Line, const FVector2D Offset, const float Width, const FTextBlockStyle& DefaultStyle, const FGeometry& AllottedGeometry, const FSlateRect& MyCullingRect, FSlateWindowElementList& OutDrawElements, int32 LayerId, const FWidgetStyle& InWidgetStyle, bool bParentEnabled) const override;
 
 protected:
 	SLATE_API FCursorLineHighlighter(const FCursorInfo* InCursorInfo);
@@ -161,7 +161,7 @@ public:
 
 	SLATE_API void SetCompositionBrush(const TAttribute<const FSlateBrush*>& InCompositionBrush);
 
-	SLATE_API virtual int32 OnPaint(const FPaintArgs& Args, const FTextLayout::FLineView& Line, const float OffsetX, const float Width, const FTextBlockStyle& DefaultStyle, const FGeometry& AllottedGeometry, const FSlateRect& MyCullingRect, FSlateWindowElementList& OutDrawElements, int32 LayerId, const FWidgetStyle& InWidgetStyle, bool bParentEnabled) const override;
+	SLATE_API virtual int32 OnPaint(const FPaintArgs& Args, const FTextLayout::FLineView& Line, const FVector2D Offset, const float Width, const FTextBlockStyle& DefaultStyle, const FGeometry& AllottedGeometry, const FSlateRect& MyCullingRect, FSlateWindowElementList& OutDrawElements, int32 LayerId, const FWidgetStyle& InWidgetStyle, bool bParentEnabled) const override;
 
 protected:
 	SLATE_API FTextCompositionHighlighter();
@@ -176,7 +176,7 @@ class FTextSelectionHighlighter : public ISlateLineHighlighter
 public:
 	static SLATE_API TSharedRef< FTextSelectionHighlighter > Create();
 
-	SLATE_API virtual int32 OnPaint(const FPaintArgs& Args, const FTextLayout::FLineView& Line, const float OffsetX, const float Width, const FTextBlockStyle& DefaultStyle, const FGeometry& AllottedGeometry, const FSlateRect& MyCullingRect, FSlateWindowElementList& OutDrawElements, int32 LayerId, const FWidgetStyle& InWidgetStyle, bool bParentEnabled) const override;
+	SLATE_API virtual int32 OnPaint(const FPaintArgs& Args, const FTextLayout::FLineView& Line, const FVector2D Offset, const float Width, const FTextBlockStyle& DefaultStyle, const FGeometry& AllottedGeometry, const FSlateRect& MyCullingRect, FSlateWindowElementList& OutDrawElements, int32 LayerId, const FWidgetStyle& InWidgetStyle, bool bParentEnabled) const override;
 
 	void SetHasKeyboardFocus(const bool bInHasKeyboardFocus)
 	{
@@ -196,7 +196,7 @@ class FTextSearchHighlighter : public ISlateLineHighlighter
 public:
 	static SLATE_API TSharedRef< FTextSearchHighlighter > Create();
 
-	SLATE_API virtual int32 OnPaint(const FPaintArgs& Args, const FTextLayout::FLineView& Line, const float OffsetX, const float Width, const FTextBlockStyle& DefaultStyle, const FGeometry& AllottedGeometry, const FSlateRect& MyCullingRect, FSlateWindowElementList& OutDrawElements, int32 LayerId, const FWidgetStyle& InWidgetStyle, bool bParentEnabled) const override;
+	SLATE_API virtual int32 OnPaint(const FPaintArgs& Args, const FTextLayout::FLineView& Line, const FVector2D Offset, const float Width, const FTextBlockStyle& DefaultStyle, const FGeometry& AllottedGeometry, const FSlateRect& MyCullingRect, FSlateWindowElementList& OutDrawElements, int32 LayerId, const FWidgetStyle& InWidgetStyle, bool bParentEnabled) const override;
 
 	void SetHasKeyboardFocus(const bool bInHasKeyboardFocus)
 	{
