Add precision mouse control to color wheel and spectrum using ctrl key modifier

diff --git a/Engine/Source/Runtime/Slate/Private/Widgets/Colors/SColorSpectrum.cpp b/Engine/Source/Runtime/Slate/Private/Widgets/Colors/SColorSpectrum.cpp
--- a/Engine/Source/Runtime/Slate/Private/Widgets/Colors/SColorSpectrum.cpp
+++ b/Engine/Source/Runtime/Slate/Private/Widgets/Colors/SColorSpectrum.cpp
@@ -24,6 +24,8 @@ void SColorSpectrum::Construct( const FArguments& InArgs )
 	OnMouseCaptureBegin = InArgs._OnMouseCaptureBegin;
 	OnMouseCaptureEnd = InArgs._OnMouseCaptureEnd;
 	OnValueChanged = InArgs._OnValueChanged;
+
+	CtrlMultiplier = InArgs._CtrlMultiplier;
 }
 
 
@@ -49,7 +51,7 @@ FReply SColorSpectrum::OnMouseButtonDown( const FGeometry& MyGeometry, const FPo
 		ProcessMouseAction(MyGeometry, MouseEvent);
 		OnMouseCaptureBegin.ExecuteIfBound();
 
-		return FReply::Handled().CaptureMouse(SharedThis(this));
+		return FReply::Handled().CaptureMouse(SharedThis(this)).UseHighPrecisionMouseMovement(SharedThis(this)).SetUserFocus(SharedThis(this), EFocusCause::Mouse);
 	}
 
 	return FReply::Unhandled();
@@ -60,9 +62,14 @@ FReply SColorSpectrum::OnMouseButtonUp( const FGeometry& MyGeometry, const FPoin
 {
 	if ((MouseEvent.GetEffectingButton() == EKeys::LeftMouseButton) && HasMouseCapture())
 	{
+		bDragging = false;
+
 		OnMouseCaptureEnd.ExecuteIfBound();
 
-		return FReply::Handled().ReleaseMouseCapture();
+		// Before showing the mouse position again, reset its position to the final location of the selector on the color spectrum
+		const FVector2D FinalMousePosition = CalcRelativeSelectedPosition() * MyGeometry.Size;
+
+		return FReply::Handled().ReleaseMouseCapture().SetMousePos(MyGeometry.LocalToAbsolute(FinalMousePosition).IntPoint());
 	}
 
 	return FReply::Unhandled();
@@ -76,6 +83,12 @@ FReply SColorSpectrum::OnMouseMove( const FGeometry& MyGeometry, const FPointerE
 		return FReply::Unhandled();
 	}
 
+	if (!bDragging)
+	{
+		bDragging = true;
+		LastSpectrumPosition = MyGeometry.AbsoluteToLocal(MouseEvent.GetScreenSpacePosition());
+	}
+
 	ProcessMouseAction(MyGeometry, MouseEvent);
 
 	return FReply::Handled();
@@ -135,7 +148,24 @@ FVector2D SColorSpectrum::CalcRelativeSelectedPosition( ) const
 
 void SColorSpectrum::ProcessMouseAction(const FGeometry& MyGeometry, const FPointerEvent& MouseEvent)
 {
-	FVector2D NormalizedMousePosition = MyGeometry.AbsoluteToLocal(MouseEvent.GetScreenSpacePosition()) / MyGeometry.GetLocalSize();
+	FVector2f LocalMouseCoordinate;
+	if (bDragging)
+	{
+		FVector2f Delta = MouseEvent.GetCursorDelta();
+		if (MouseEvent.IsControlDown())
+		{
+			Delta *= CtrlMultiplier.Get();
+		}
+
+		LocalMouseCoordinate = LastSpectrumPosition + Delta;
+		LastSpectrumPosition = LocalMouseCoordinate;
+	}
+	else
+	{
+		LocalMouseCoordinate = MyGeometry.AbsoluteToLocal(MouseEvent.GetScreenSpacePosition());
+	}
+
+	FVector2D NormalizedMousePosition = LocalMouseCoordinate / MyGeometry.GetLocalSize();
 	NormalizedMousePosition = NormalizedMousePosition.ClampAxes(0.0f, 1.0f);
 
 	SelectedColor.UpdateNow(*this);
@@ -155,3 +185,10 @@ void SColorSpectrum::ProcessMouseAction(const FGeometry& MyGeometry, const FPoin
 
 	OnValueChanged.ExecuteIfBound(NewColor);
 }
+
+FCursorReply SColorSpectrum::OnCursorQuery(const FGeometry& MyGeometry, const FPointerEvent& CursorEvent) const
+{
+	return bDragging ?
+		FCursorReply::Cursor(EMouseCursor::None) :
+		FCursorReply::Cursor(EMouseCursor::Default);
+}
diff --git a/Engine/Source/Runtime/Slate/Private/Widgets/Colors/SColorWheel.cpp b/Engine/Source/Runtime/Slate/Private/Widgets/Colors/SColorWheel.cpp
--- a/Engine/Source/Runtime/Slate/Private/Widgets/Colors/SColorWheel.cpp
+++ b/Engine/Source/Runtime/Slate/Private/Widgets/Colors/SColorWheel.cpp
@@ -22,6 +22,8 @@ void SColorWheel::Construct(const FArguments& InArgs)
 	OnMouseCaptureBegin = InArgs._OnMouseCaptureBegin;
 	OnMouseCaptureEnd = InArgs._OnMouseCaptureEnd;
 	OnValueChanged = InArgs._OnValueChanged;
+
+	CtrlMultiplier = InArgs._CtrlMultiplier;
 }
 
 
@@ -52,7 +54,7 @@ FReply SColorWheel::OnMouseButtonDown(const FGeometry& MyGeometry, const FPointe
 			return FReply::Unhandled();
 		}
 
-		return FReply::Handled().CaptureMouse(SharedThis(this));
+		return FReply::Handled().CaptureMouse(SharedThis(this)).UseHighPrecisionMouseMovement(SharedThis(this)).SetUserFocus(SharedThis(this), EFocusCause::Mouse);
 	}
 
 	return FReply::Unhandled();
@@ -63,9 +65,15 @@ FReply SColorWheel::OnMouseButtonUp(const FGeometry& MyGeometry, const FPointerE
 {
 	if (MouseEvent.GetEffectingButton() == EKeys::LeftMouseButton && HasMouseCapture())
 	{
+		bDragging = false;
+
 		OnMouseCaptureEnd.ExecuteIfBound();
 
-		return FReply::Handled().ReleaseMouseCapture();
+		// Before showing the mouse position again, reset its position to the final location of the selector on the color wheel
+		const FVector2f CircleSize = MyGeometry.GetLocalSize() - SelectorImage->ImageSize;
+		const FVector2f FinalMousePosition = 0.5f * (MyGeometry.GetLocalSize() + CalcRelativePositionFromCenter() * CircleSize);
+
+		return FReply::Handled().ReleaseMouseCapture().SetMousePos(MyGeometry.LocalToAbsolute(FinalMousePosition).IntPoint());
 	}
 
 	return FReply::Unhandled();
@@ -79,6 +87,12 @@ FReply SColorWheel::OnMouseMove(const FGeometry& MyGeometry, const FPointerEvent
 		return FReply::Unhandled();
 	}
 
+	if (!bDragging)
+	{
+		bDragging = true;
+		LastWheelPosition = MyGeometry.AbsoluteToLocal(MouseEvent.GetScreenSpacePosition());
+	}
+
 	ProcessMouseAction(MyGeometry, MouseEvent, true);
 
 	return FReply::Handled();
@@ -130,7 +144,23 @@ UE::Slate::FDeprecateVector2DResult SColorWheel::CalcRelativePositionFromCenter(
 
 bool SColorWheel::ProcessMouseAction(const FGeometry& MyGeometry, const FPointerEvent& MouseEvent, bool bProcessWhenOutsideColorWheel)
 {
-	const FVector2f LocalMouseCoordinate = MyGeometry.AbsoluteToLocal(MouseEvent.GetScreenSpacePosition());
+	FVector2f LocalMouseCoordinate;
+	if (bDragging)
+	{
+		FVector2f Delta = MouseEvent.GetCursorDelta();
+		if (MouseEvent.IsControlDown())
+		{
+			Delta *= CtrlMultiplier.Get();
+		}
+
+		LocalMouseCoordinate = LastWheelPosition + Delta;
+		LastWheelPosition = LocalMouseCoordinate;
+	}
+	else
+	{
+		LocalMouseCoordinate = MyGeometry.AbsoluteToLocal(MouseEvent.GetScreenSpacePosition());
+	}
+	
 	const FVector2f RelativePositionFromCenter = (2.0f * LocalMouseCoordinate - MyGeometry.GetLocalSize()) / (MyGeometry.GetLocalSize() - SelectorImage->ImageSize);
 	const float RelativeRadius = RelativePositionFromCenter.Size();
 
@@ -155,3 +185,10 @@ bool SColorWheel::ProcessMouseAction(const FGeometry& MyGeometry, const FPointer
 
 	return (RelativeRadius <= 1.0f);
 }
+
+FCursorReply SColorWheel::OnCursorQuery(const FGeometry& MyGeometry, const FPointerEvent& CursorEvent) const
+{
+	return bDragging ?
+		FCursorReply::Cursor(EMouseCursor::None) :
+		FCursorReply::Cursor(EMouseCursor::Default);
+}
diff --git a/Engine/Source/Runtime/Slate/Public/Widgets/Colors/SColorSpectrum.h b/Engine/Source/Runtime/Slate/Public/Widgets/Colors/SColorSpectrum.h
--- a/Engine/Source/Runtime/Slate/Public/Widgets/Colors/SColorSpectrum.h
+++ b/Engine/Source/Runtime/Slate/Public/Widgets/Colors/SColorSpectrum.h
@@ -26,6 +26,7 @@ public:
 		, _OnMouseCaptureBegin()
 		, _OnMouseCaptureEnd()
 		, _OnValueChanged()
+		, _CtrlMultiplier(0.1f)
 	{ }
 	
 		/** The current color selected by the user */
@@ -37,9 +38,12 @@ public:
 		/** Invoked when the mouse is released and a capture ends */
 		SLATE_EVENT(FSimpleDelegate, OnMouseCaptureEnd)
 
-		/** Invoked when a new value is selected on the color wheel */
+		/** Invoked when a new value is selected on the color spectrum */
 		SLATE_EVENT(FOnLinearColorValueChanged, OnValueChanged)
 
+		/** Multiplier to use when ctrl is held down */
+		SLATE_ATTRIBUTE(float, CtrlMultiplier)
+
 	SLATE_END_ARGS()
 	
 public:
@@ -63,7 +67,8 @@ public:
 	SLATE_API virtual FReply OnMouseButtonUp( const FGeometry& MyGeometry, const FPointerEvent& MouseEvent ) override;
 	SLATE_API virtual FReply OnMouseMove( const FGeometry& MyGeometry, const FPointerEvent& MouseEvent ) override;
 	SLATE_API virtual int32 OnPaint( const FPaintArgs& Args, const FGeometry& AllottedGeometry, const FSlateRect& MyCullingRect, FSlateWindowElementList& OutDrawElements, int32 LayerId, const FWidgetStyle& InWidgetStyle, bool bParentEnabled ) const override;
-	
+	SLATE_API virtual FCursorReply OnCursorQuery(const FGeometry& MyGeometry, const FPointerEvent& CursorEvent) const override;
+
 protected:
 
 	/**
@@ -80,15 +85,24 @@ protected:
 
 private:
 
-	// The color wheel image to show.
+	// The color spectrum image to show.
 	const FSlateBrush* Image;
 	
 	// The current color selected by the user.
 	TSlateAttribute<FLinearColor, EInvalidateWidgetReason::Paint> SelectedColor;
 
+	// Mouse sensitivity multiplier to use when dragging the selector on the color spectrum, applied when the ctrl modifier key is pressed
+	TAttribute<float> CtrlMultiplier;
+
 	// The color selector image to show.
 	const FSlateBrush* SelectorImage;
 
+	// Whether the user is dragging the slider
+	bool bDragging = false;
+
+	// Cached mouse position to restore after dragging.
+	FVector2f LastSpectrumPosition;
+
 private:
 
 	// Holds a delegate that is executed when the mouse is pressed and a capture begins.
@@ -97,6 +111,6 @@ private:
 	// Holds a delegate that is executed when the mouse is let up and a capture ends.
 	FSimpleDelegate OnMouseCaptureEnd;
 
-	// Holds a delegate that is executed when a new value is selected on the color wheel.
+	// Holds a delegate that is executed when a new value is selected on the color spectrum.
 	FOnLinearColorValueChanged OnValueChanged;
 };
diff --git a/Engine/Source/Runtime/Slate/Public/Widgets/Colors/SColorWheel.h b/Engine/Source/Runtime/Slate/Public/Widgets/Colors/SColorWheel.h
--- a/Engine/Source/Runtime/Slate/Public/Widgets/Colors/SColorWheel.h
+++ b/Engine/Source/Runtime/Slate/Public/Widgets/Colors/SColorWheel.h
@@ -26,6 +26,7 @@ public:
 		, _OnMouseCaptureBegin()
 		, _OnMouseCaptureEnd()
 		, _OnValueChanged()
+		, _CtrlMultiplier(0.1f)
 	{ }
 	
 		/** The current color selected by the user. */
@@ -40,6 +41,9 @@ public:
 		/** Invoked when a new value is selected on the color wheel. */
 		SLATE_EVENT(FOnLinearColorValueChanged, OnValueChanged)
 
+		/** Multiplier to use when ctrl is held down */
+		SLATE_ATTRIBUTE(float, CtrlMultiplier)
+
 	SLATE_END_ARGS()
 	
 public:
@@ -63,7 +67,8 @@ public:
 	SLATE_API virtual FReply OnMouseButtonUp(const FGeometry& MyGeometry, const FPointerEvent& MouseEvent) override;
 	SLATE_API virtual FReply OnMouseMove(const FGeometry& MyGeometry, const FPointerEvent& MouseEvent) override;
 	SLATE_API virtual int32 OnPaint(const FPaintArgs& Args, const FGeometry& AllottedGeometry, const FSlateRect& MyCullingRect, FSlateWindowElementList& OutDrawElements, int32 LayerId, const FWidgetStyle& InWidgetStyle, bool bParentEnabled) const override;
-	
+	SLATE_API virtual FCursorReply OnCursorQuery(const FGeometry& MyGeometry, const FPointerEvent& CursorEvent) const override;
+
 protected:
 
 	/**
@@ -88,9 +93,18 @@ private:
 	/** The current color selected by the user. */
 	TSlateAttribute<FLinearColor, EInvalidateWidgetReason::Paint> SelectedColor;
 
+	/** Mouse sensitivity multiplier to use when dragging the selector on the color wheel, applied when the ctrl modifier key is pressed */
+	TAttribute<float> CtrlMultiplier;
+
 	/** The color selector image to show. */
 	const FSlateBrush* SelectorImage;
 
+	/** Whether the user is dragging the slider */
+	bool bDragging = false;
+
+	/** Cached mouse position to restore after dragging. */
+	FVector2f LastWheelPosition;
+
 private:
 
 	/** Invoked when the mouse is pressed and a capture begins. */
