Improve the overflow behavior of Slate's clipping boxes

This change makes toolbar overflow menus contain exactly one entry per clipped toolbar block. Previously, the overflow menus often contained a different number of entries which can be confusing to users.

The underlying change behind this improvement is making SClippingHorizontalBox and SClippingVerticalBox report a list of clipped widgets instead of an index at which widgets have been clipped. This allows the SMultiBoxWidget to map the reportedly clipped widgets to the source FMultiBlocks and display those in the overflow menu.

Having the boxes report a list of clipped widgets instead of an index paves the way for clipping widgets in any order, such as by priority, instead of always clipping at the right hand side of the toolbar.

* Add SMultiBoxWidget::FinalMultiBoxWidgets for keeping track of the final widgets added to the underlying boxes (SClippingHorizontalBox, SClippingVerticalBox, ...). We need this so we can map clipped widgets to their originating multiblock.
* Add SClippingHorizontalBox::GetClippedWidgets() and SClippingVerticalBox::GetClippedWidgets() which can be used to get the array of currently clipped widgets.
* Remove all old index tracking of clipped widgets from SClippingHorizontalBox and SClippingVerticalBox.
* Add new private SClippingVerticalBox::ClippedIndices array so we can keep support for a SelectedIndex and know if that selected index is clipped or not.
* Add missing default initializations for some fields of SClippingHorizontalBox and SClippingVerticalBox.

diff --git a/Engine/Source/Runtime/Slate/Private/Framework/MultiBox/MultiBox.cpp b/Engine/Source/Runtime/Slate/Private/Framework/MultiBox/MultiBox.cpp
--- a/Engine/Source/Runtime/Slate/Private/Framework/MultiBox/MultiBox.cpp
+++ b/Engine/Source/Runtime/Slate/Private/Framework/MultiBox/MultiBox.cpp
@@ -776,8 +776,6 @@ void SMultiBoxWidget::AddBlockWidget(const FMultiBlock& Block, TSharedPtr<SHoriz
 
 	TSharedPtr<SWidget> FinalWidget;
 
-	const EMultiBlockType BlockType = Block.GetType();
-
 	if (MultiBox->ModifyBlockWidgetAfterMake.IsBound())
 	{
 		FinalWidget = MultiBox->ModifyBlockWidgetAfterMake.Execute(SharedThis(this), Block, BlockWidget);
@@ -829,6 +827,8 @@ void SMultiBoxWidget::AddBlockWidget(const FMultiBlock& Block, TSharedPtr<SHoriz
 					FinalWidgetWithHook
 				];
 
+				FinalMultiBoxWidgets.Add(FinalWidgetWithHook.ToSharedPtr(), BlockRef.ToSharedPtr());
+
 				if (bOverride)
 				{
 					if (StyleParams.SizeRule.IsSet()
@@ -888,6 +888,8 @@ void SMultiBoxWidget::AddBlockWidget(const FMultiBlock& Block, TSharedPtr<SHoriz
 					FinalWidgetWithHook
 				];
 			}
+
+			FinalMultiBoxWidgets.Add(FinalWidgetWithHook.ToSharedPtr(), BlockRef.ToSharedPtr());
 		}
 		break;
 	case EMultiBoxType::UniformToolBar:
@@ -897,11 +899,15 @@ void SMultiBoxWidget::AddBlockWidget(const FMultiBlock& Block, TSharedPtr<SHoriz
 			[
 				FinalWidgetWithHook
 			];
+
+			FinalMultiBoxWidgets.Add(FinalWidgetWithHook.ToSharedPtr(), BlockRef.ToSharedPtr());
 		}
 		break;
 	case EMultiBoxType::ButtonRow:
 		{
 			TileViewWidgets.Add( FinalWidget.ToSharedRef() );
+
+			FinalMultiBoxWidgets.Add(FinalWidget, BlockRef.ToSharedPtr());
 		}
 		break;
 	case EMultiBoxType::Menu:
@@ -925,6 +931,8 @@ void SMultiBoxWidget::AddBlockWidget(const FMultiBlock& Block, TSharedPtr<SHoriz
 					FinalWidget.ToSharedRef()
 				]
 			];
+
+			FinalMultiBoxWidgets.Add(FinalWidget, BlockRef.ToSharedPtr());
 		}
 		break;
 	}
@@ -1379,17 +1387,31 @@ void SMultiBoxWidget::BuildMultiBoxWidget()
 TSharedRef<SWidget> SMultiBoxWidget::OnWrapButtonClicked()
 {
 	FMenuBuilder MenuBuilder(true, MultiBox->GetLastCommandList(), TSharedPtr<FExtender>(), false, GetStyleSet());
+
+	if (ClippedVerticalBox)
 	{
-		if (ClippedVerticalBox)
+		MenuBuilder.SetCheckBoxStyle("ClippingVerticalBox.Check");
+	}
+
+	const TArray<TWeakPtr<SWidget>> ClippedWidgets = ClippedHorizontalBox ? ClippedHorizontalBox->GetClippedWidgets()
+																		  : ClippedVerticalBox->GetClippedWidgets();
+
+	for (const TWeakPtr<SWidget>& WeakClippedWidget : ClippedWidgets)
+	{
+		const TSharedPtr<SWidget> ClippedWidget = WeakClippedWidget.Pin();
+		if (!ClippedWidget)
 		{
-			MenuBuilder.SetCheckBoxStyle("ClippingVerticalBox.Check");
+			continue;
 		}
-		const TArray< TSharedRef< const FMultiBlock > >& Blocks = MultiBox->GetBlocks();
 
-		int32 BlockIdx = ClippedHorizontalBox ? ClippedHorizontalBox->GetClippedIndex() : ClippedVerticalBox->GetClippedIndex();
-		for (; BlockIdx < Blocks.Num(); ++BlockIdx)
+		// Find the block of the clipped widget and add it to the menu.
+		for (const TPair<TSharedPtr<SWidget>, TSharedPtr<const FMultiBlock>>& Pair : FinalMultiBoxWidgets)
 		{
-			Blocks[BlockIdx]->CreateMenuEntry(MenuBuilder);
+			if (ClippedWidget == Pair.Key)
+			{
+				Pair.Value->CreateMenuEntry(MenuBuilder);
+				break;
+			}
 		}
 	}
 
diff --git a/Engine/Source/Runtime/Slate/Private/Framework/MultiBox/SClippingHorizontalBox.cpp b/Engine/Source/Runtime/Slate/Private/Framework/MultiBox/SClippingHorizontalBox.cpp
--- a/Engine/Source/Runtime/Slate/Private/Framework/MultiBox/SClippingHorizontalBox.cpp
+++ b/Engine/Source/Runtime/Slate/Private/Framework/MultiBox/SClippingHorizontalBox.cpp
@@ -1,8 +1,10 @@
 // Copyright Epic Games, Inc. All Rights Reserved.
 
 #include "Framework/MultiBox/SClippingHorizontalBox.h"
+
 #include "Framework/Application/SlateApplication.h"
 #include "Layout/ArrangedChildren.h"
+#include "Layout/LayoutUtils.h"
 #include "Styling/ToolBarStyle.h"
 #include "Widgets/Images/SImage.h"
 #include "Widgets/Input/SComboButton.h"
@@ -13,11 +15,11 @@ void SClippingHorizontalBox::OnArrangeChildren( const FGeometry& AllottedGeometr
 	// been called and this method isn't going to behave properly
 	check(WrapButton.IsValid());
 
-	LastClippedIdx = ClippedIdx;
-
-	NumClippedChildren = 0;
 	SHorizontalBox::OnArrangeChildren(AllottedGeometry, ArrangedChildren);
 
+	// Clear our previously clipped widgets. We'll repopulate this below.
+	ClippedWidgets.Reset();
+
 	// Remove children that are clipped by the allotted geometry
 	const int32 NumChildren = ArrangedChildren.Num(); 
 	int32 IndexClippedAt = NumChildren;
@@ -31,7 +33,10 @@ void SClippingHorizontalBox::OnArrangeChildren( const FGeometry& AllottedGeometr
 		const FVector2D BoxPosRounded = FVector2D(FMath::CeilToInt(AbsBoxPos.X - KINDA_SMALL_NUMBER), FMath::CeilToInt(AbsBoxPos.Y - KINDA_SMALL_NUMBER));
 		if (WidgetPosRounded.X > BoxPosRounded.X)
 		{
-			++NumClippedChildren;
+			// Insert the clipped widget at the start of the array to keep the array sorted with
+			// the left-most clipped widget first and the right-most clipped widget last.
+			ClippedWidgets.Insert(CurWidget.Widget.ToWeakPtr(), 0);
+
 			ArrangedChildren.Remove(ChildIdx);
 			IndexClippedAt = ChildIdx;
 		}
@@ -66,18 +71,26 @@ void SClippingHorizontalBox::OnArrangeChildren( const FGeometry& AllottedGeometr
 				const FArrangedWidget& CurWidget = ArrangedChildren[ChildIdx];
 				if (FMath::TruncToInt(CurWidget.Geometry.AbsolutePosition.X + CurWidget.Geometry.GetLocalSize().X * CurWidget.Geometry.Scale) > WrapButtonXPosition)
 				{
-					++NumClippedChildren;
+					// Insert the clipped widget at the start of the array to keep the array sorted with
+					// the left-most clipped widget first and the right-most clipped widget last.
+					ClippedWidgets.Insert(CurWidget.Widget.ToWeakPtr(), 0);
+
 					ArrangedChildren.Remove(ChildIdx);
 				}
 			}
 		}
 		else // No space left for anything including WrapButton
 		{
+			for (int32 ChildIdx = 0; ChildIdx < ArrangedChildren.Num() - 1 /* Exclude the wrap button. */; ++ChildIdx)
+			{
+				const FArrangedWidget& CurWidget = ArrangedChildren[ChildIdx];
+				// Insert the clipped widget at the start of the array to keep the array sorted with
+				// the left-most clipped widget first and the right-most clipped widget last.
+				ClippedWidgets.Insert(CurWidget.Widget.ToWeakPtr(), 0);
+			}
 			ArrangedChildren.Empty();
 		}
 	}
-
-	ClippedIdx = ArrangedChildren.Num() - 1;
 }
 
 int32 SClippingHorizontalBox::OnPaint( const FPaintArgs& Args, const FGeometry& AllottedGeometry, const FSlateRect& MyCullingRect, FSlateWindowElementList& OutDrawElements, int32 LayerId, const FWidgetStyle& InWidgetStyle, bool bParentEnabled ) const
@@ -103,8 +116,6 @@ void SClippingHorizontalBox::Construct( const FArguments& InArgs )
 	StyleSet = InArgs._StyleSet;
 	StyleName = InArgs._StyleName;
 	bIsFocusable = InArgs._IsFocusable;
-
-	LastClippedIdx = ClippedIdx = INDEX_NONE;
 }
 
 void SClippingHorizontalBox::AddWrapButton()
@@ -157,7 +168,7 @@ void SClippingHorizontalBox::OnWrapButtonOpenChanged(bool bIsOpen)
 
 EActiveTimerReturnType SClippingHorizontalBox::UpdateWrapButtonStatus(double CurrentTime, float DeltaTime)
 {
-	if (LastClippedIdx != ClippedIdx || !WrapButton->IsOpen())
+	if (!WrapButton->IsOpen())
 	{
 		WrapButton->SetIsOpen(false);
 		WrapButtonOpenTimer.Reset();
diff --git a/Engine/Source/Runtime/Slate/Private/Framework/MultiBox/SClippingHorizontalBox.h b/Engine/Source/Runtime/Slate/Private/Framework/MultiBox/SClippingHorizontalBox.h
--- a/Engine/Source/Runtime/Slate/Private/Framework/MultiBox/SClippingHorizontalBox.h
+++ b/Engine/Source/Runtime/Slate/Private/Framework/MultiBox/SClippingHorizontalBox.h
@@ -3,9 +3,9 @@
 #pragma once
 
 #include "CoreMinimal.h"
-#include "Widgets/DeclarativeSyntaxSupport.h"
-#include "Styling/CoreStyle.h"
 #include "Framework/SlateDelegates.h"
+#include "Styling/CoreStyle.h"
+#include "Widgets/DeclarativeSyntaxSupport.h"
 #include "Widgets/SBoxPanel.h"
 
 class FArrangedChildren;
@@ -40,8 +40,10 @@ public:
 	/** Adds the wrap button */
 	void AddWrapButton();
 
-	/** Returns to index of the first clipped child/block */
-	int32 GetClippedIndex() { return ClippedIdx; }
+	TArray<TWeakPtr<SWidget>> GetClippedWidgets() const
+	{
+		return ClippedWidgets;
+	}
 
 private:
 	void OnWrapButtonOpenChanged(bool bIsOpen);
@@ -54,22 +56,16 @@ private:
 	/** Callback for when the wrap button is clicked */
 	FOnGetContent OnWrapButtonClicked;
 
-	/** Index of the first clipped child/block */
-	mutable int32 ClippedIdx;
-
-	mutable int32 LastClippedIdx;
-
-	/** Number of clipped children not including the wrap button */
-	mutable int32 NumClippedChildren;
-
 	TSharedPtr<FActiveTimerHandle> WrapButtonOpenTimer;
 
 	/** Can the wrap button be focused? */
-	bool bIsFocusable;
+	bool bIsFocusable = true;
 
 	float WrapButtonWidth = 0.0f; // Fixed width, initialized after button creation
 
 	/** The style to use */
-	const ISlateStyle* StyleSet;
+	const ISlateStyle* StyleSet = nullptr;
 	FName StyleName;
+
+	mutable TArray<TWeakPtr<SWidget>> ClippedWidgets;
 };
diff --git a/Engine/Source/Runtime/Slate/Private/Framework/MultiBox/SClippingVerticalBox.cpp b/Engine/Source/Runtime/Slate/Private/Framework/MultiBox/SClippingVerticalBox.cpp
--- a/Engine/Source/Runtime/Slate/Private/Framework/MultiBox/SClippingVerticalBox.cpp
+++ b/Engine/Source/Runtime/Slate/Private/Framework/MultiBox/SClippingVerticalBox.cpp
@@ -22,10 +22,11 @@ void SClippingVerticalBox::OnArrangeChildren( const FGeometry& AllottedGeometry,
 	// been called and this method isn't going to behave properly
 	check(WrapButton.IsValid());
 
+	ClippedWidgets.Empty();
+	ClippedIndices.Empty();
+
 	const FToolBarStyle& ToolBarStyle = StyleSet->GetWidgetStyle<FToolBarStyle>(StyleName);
-	LastClippedIdx = ClippedIdx;
 
-	NumClippedChildren = 0;
 	SVerticalBox::OnArrangeChildren(AllottedGeometry, ArrangedChildren);
 
 	// Remove children that are clipped by the allotted geometry
@@ -49,7 +50,11 @@ void SClippingVerticalBox::OnArrangeChildren( const FGeometry& AllottedGeometry,
 		
 		if ( CurrentWidgetMaxY > AllottedGeometry.Size.Y)
 		{
-			++NumClippedChildren;
+			// Insert the clipped widget at the start of the array to keep the array sorted with
+			// the top-most clipped widget first and the bottom-most clipped widget last.
+			ClippedWidgets.Insert(CurrentWidget.Widget.ToWeakPtr(), 0);
+			ClippedIndices.Add(ChildIdx);
+
 			ArrangedChildren.Remove(ChildIdx);
 		}
 		else if (LastToolBarButtonIndex == ChildIdx)
@@ -62,8 +67,6 @@ void SClippingVerticalBox::OnArrangeChildren( const FGeometry& AllottedGeometry,
 	FVector2D Size = ArrangedButton.Geometry.GetLocalSize();
 	Size.Y = OverflowButtonSize;
 	ArrangedButton.Geometry = AllottedGeometry.MakeChild(Size, FSlateLayoutTransform(AllottedGeometry.GetLocalSize() - Size));
-
-	ClippedIdx = ArrangedChildren.Num() - 1;
 }
 
 int32 SClippingVerticalBox::OnPaint( const FPaintArgs& Args, const FGeometry& AllottedGeometry, const FSlateRect& MyCullingRect, FSlateWindowElementList& OutDrawElements, int32 LayerId, const FWidgetStyle& InWidgetStyle, bool bParentEnabled ) const
@@ -90,13 +93,10 @@ void SClippingVerticalBox::Construct( const FArguments& InArgs )
 	StyleName = InArgs._StyleName;
 	bIsFocusable = InArgs._IsFocusable;
 	SelectedIndex = InArgs._SelectedIndex;
-
-	LastClippedIdx = ClippedIdx = INDEX_NONE;
 }
 
 void SClippingVerticalBox::AddWrapButton()
 {
-
 	InitializeWrapButton( WrapButton, false );
 	InitializeWrapButton( SelectedWrapButton, true );
 
@@ -137,15 +137,16 @@ void SClippingVerticalBox::InitializeWrapButton( TSharedPtr<SComboButton>& Butto
 	SelectedStyle.SetPressedPadding( 0 );
 	SelectedStyle.SetNormal( FSlateRoundedBoxBrush( FStyleColors::Primary, 4.f,
 		FLinearColor(0, 0, 0, .8), 0.5) );
-	
+
 	// Construct the wrap button used in toolbars and menu bars
 	// Always allow this to be focusable to prevent the menu from collapsing during interaction
+	// clang-format off
 	Button =
 		SNew( SComboButton )
 		.HasDownArrow( false )
 		.Visibility_Lambda( [this, bCreateSelectedAppearance] ()
 		{
-			bool bShowingSelected = SelectedIndex.Get() > LastClippedIdx - 1;
+			const bool bShowingSelected = ClippedIndices.Contains(SelectedIndex.Get());
 			if ( bCreateSelectedAppearance )
 			{
 				return bShowingSelected ? EVisibility::Visible : EVisibility::Collapsed;				
@@ -184,6 +185,7 @@ void SClippingVerticalBox::InitializeWrapButton( TSharedPtr<SComboButton>& Butto
 				]
 			
 		];
+	// clang-format on
 }
 
 void SClippingVerticalBox::OnWrapButtonOpenChanged(bool bIsOpen)
@@ -201,7 +203,7 @@ void SClippingVerticalBox::OnWrapButtonOpenChanged(bool bIsOpen)
 
 EActiveTimerReturnType SClippingVerticalBox::UpdateWrapButtonStatus(double CurrentTime, float DeltaTime)
 {
-	if (LastClippedIdx != ClippedIdx || !WrapButton->IsOpen())
+	if (!WrapButton->IsOpen())
 	{
 		WrapButton->SetIsOpen(false);
 		WrapButtonOpenTimer.Reset();
diff --git a/Engine/Source/Runtime/Slate/Private/Framework/MultiBox/SClippingVerticalBox.h b/Engine/Source/Runtime/Slate/Private/Framework/MultiBox/SClippingVerticalBox.h
--- a/Engine/Source/Runtime/Slate/Private/Framework/MultiBox/SClippingVerticalBox.h
+++ b/Engine/Source/Runtime/Slate/Private/Framework/MultiBox/SClippingVerticalBox.h
@@ -47,8 +47,10 @@ public:
 	/** Adds the wrap button */
 	void AddWrapButton();
 
-	/** Returns to index of the first clipped child/block */
-	int32 GetClippedIndex() { return ClippedIdx; }
+	TArray<TWeakPtr<SWidget>> GetClippedWidgets() const
+	{
+		return ClippedWidgets;
+	}
 
 private:
 	void OnWrapButtonOpenChanged(bool bIsOpen);
@@ -72,21 +74,13 @@ private:
 	/** Callback for when the wrap button is clicked */
 	FOnGetContent OnWrapButtonClicked;
 
-	/** Index of the first clipped child/block */
-	mutable int32 ClippedIdx;
-
-	mutable int32 LastClippedIdx;
-
-	/** Number of clipped children not including the wrap button */
-	mutable int32 NumClippedChildren;
-
 	TSharedPtr<FActiveTimerHandle> WrapButtonOpenTimer;
 
 	/** Can the wrap button be focused? */
-	bool bIsFocusable;
+	bool bIsFocusable = true;
 
 	/** The style to use */
-	const ISlateStyle* StyleSet;
+	const ISlateStyle* StyleSet = nullptr;
 
 	/** the button style for the clipped content button when nothing within the clipped content is selected */
 	FButtonStyle Style;
@@ -101,4 +95,7 @@ private:
 	mutable int32 LastToolBarButtonIndex = INDEX_NONE;
 
 	FName StyleName;
+
+	mutable TArray<TWeakPtr<SWidget>> ClippedWidgets;
+	mutable TArray<int32> ClippedIndices;
 };
diff --git a/Engine/Source/Runtime/Slate/Public/Framework/MultiBox/MultiBox.h b/Engine/Source/Runtime/Slate/Public/Framework/MultiBox/MultiBox.h
--- a/Engine/Source/Runtime/Slate/Public/Framework/MultiBox/MultiBox.h
+++ b/Engine/Source/Runtime/Slate/Public/Framework/MultiBox/MultiBox.h
@@ -949,6 +949,9 @@ private:
 	/** The multibox widgets that are contained, linked to their searchable text hierarchies. */
 	TMap<TSharedPtr<SWidget>, TArray<FText>> MultiBoxWidgets;
 
+	/** The final multibox widgets as they were added to the UI, linked to their original MultiBlocks. */
+	TMap<TSharedPtr<SWidget>, TSharedPtr<const FMultiBlock>> FinalMultiBoxWidgets;
+
 	/** The set of searchable blocks found in this multibox sub-menus that were collected and flatten in this multibox to support recursively searching this multibox hierarchy. */
 	TMap<TSharedPtr<const FMultiBlock>, TSharedPtr<FFlattenSearchableBlockInfo>> FlattenSearchableBlocks;
 
