[ContentBrowser] Content Browser hides when on an asset tab
The issue is mainly for all NomadTab (like the Content Browser), all the NomadTab when docked into a layout will be bound to that layout and will be closed when changing MajorTab even if undocked.
This change will unbound the NomadTab from the layout when its DroppedOntoNothing.

diff --git a/Engine/Source/Runtime/Slate/Private/Framework/Docking/FDockingDragOperation.cpp b/Engine/Source/Runtime/Slate/Private/Framework/Docking/FDockingDragOperation.cpp
--- a/Engine/Source/Runtime/Slate/Private/Framework/Docking/FDockingDragOperation.cpp
+++ b/Engine/Source/Runtime/Slate/Private/Framework/Docking/FDockingDragOperation.cpp
@@ -202,7 +202,14 @@ FVector2D FDockingDragOperation::GetTabGrabOffsetFraction() const
 bool FDockingDragOperation::CanDockInNode(const TSharedRef<SDockingNode>& DockNode, EViaTabwell IsDockingViaTabwell ) const
 {
 	const TSharedRef<FTabManager> TargetTabManager = DockNode->GetDockArea()->GetTabManager();
-	const TSharedRef<FTabManager> TabManagerOfOrigin = this->TabOwnerAreaOfOrigin->GetTabManager();
+	TSharedRef<FTabManager> TabManagerOfOriginToUse = this->TabOwnerAreaOfOrigin->GetTabManager();
+
+	// if we are dragging a NomadTab force it to use the FGlobalTabManager as the TabManager of origin
+	// There are cases where the NomadTab is spawned like a normal tab and his TabManager is not correct
+	if (TabBeingDragged->GetTabRole() == ETabRole::NomadTab)
+	{
+		TabManagerOfOriginToUse = FGlobalTabmanager::Get();
+	}
 
 	if (IsParentWidgetOf(TabBeingDragged->GetContent(), DockNode))
 	{
@@ -211,15 +218,28 @@ bool FDockingDragOperation::CanDockInNode(const TSharedRef<SDockingNode>& DockNo
 
 	if (TabBeingDragged->GetTabRole() == ETabRole::NomadTab)
 	{
-		if ( IsDockingViaTabwell == FDockingDragOperation::DockingViaTabWell )
+		if (TargetTabManager != FGlobalTabmanager::Get())
 		{
-			// Nomad tabs can be docked in in any tab well.
-			return true;
-		}
-		else
-		{
-			return TargetTabManager != FGlobalTabmanager::Get();
+			if (TabManagerOfOriginToUse == TargetTabManager && IsDockingViaTabwell == FDockingDragOperation::DockingViaTabWell)
+			{
+				return true;
+			}
+
+			const TSharedPtr<SDockingArea>& DockingArea = DockNode->GetDockArea();
+			if (DockingArea.IsValid())
+			{
+				if (!DockingArea->GetParentWindow() || DockingArea->GetParentWindow() == TabOwnerAreaOfOrigin->GetParentWindow())
+				{
+					// Nomad tab can be docked in any non-floating dock area or if the ParentWindow is the same
+					return true;
+				}
+			}
+			return false;
 		}
+
+		// Nomad tab can be docked to any GlobalTabManager area.
+		// This additional check is needed to not make the cross visible for the main root window of the engine, but to make it able anyway to dock it like a major tab.
+		return IsDockingViaTabwell == FDockingDragOperation::DockingViaTabWell || FGlobalTabmanager::Get()->GetRootWindow() != DockNode->GetDockArea()->GetParentWindow();
 	}
 	else if (TabBeingDragged->GetTabRole() == ETabRole::MajorTab)
 	{
@@ -227,14 +247,14 @@ bool FDockingDragOperation::CanDockInNode(const TSharedRef<SDockingNode>& DockNo
 		// be allowed to split areas. They are also confined to their
 		// tab manager of origin.
 		// The only exception is an empty area, where docking the tab should be really easy.
-		const bool bTabManagerMatches = TargetTabManager == TabManagerOfOrigin;
+		const bool bTabManagerMatches = TargetTabManager == TabManagerOfOriginToUse;
 		const bool bCanDockInEmptyArea = DockNode->GetNodeType() == SDockingNode::DockArea && StaticCastSharedRef<SDockingArea>(DockNode)->GetChildNodes().Num() == 0;
 		return bTabManagerMatches && (IsDockingViaTabwell == FDockingDragOperation::DockingViaTabWell || bCanDockInEmptyArea);
 	}
 	else
 	{
 		// Most commonly, tabs are confined to their tab manager of origin.
-		return (TargetTabManager == TabManagerOfOrigin);
+		return (TargetTabManager == TabManagerOfOriginToUse);
 	}	
 }
 
@@ -341,14 +361,21 @@ void FDockingDragOperation::DroppedOntoNothing()
 
 	TSharedPtr<SDockingTabStack> NewDockNode;
 
+	TSharedPtr<FTabManager> TabManagerToUse;
+
 	if (TabBeingDragged->GetTabRole() == ETabRole::NomadTab)
 	{
-		TabBeingDragged->SetTabManager(FGlobalTabmanager::Get());
+		TabManagerToUse = FGlobalTabmanager::Get();
+		TabBeingDragged->SetTabManager(TabManagerToUse);
+	}
+	else
+	{
+		TabManagerToUse = MyTabManager;
 	}
 
 	// Create a new dockarea
 	TSharedRef<SDockingArea> NewDockArea =
-		SNew(SDockingArea, MyTabManager.ToSharedRef(), FTabManager::NewPrimaryArea())
+		SNew(SDockingArea, TabManagerToUse.ToSharedRef(), FTabManager::NewPrimaryArea())
 		.ParentWindow(NewWindow)
 		.InitialContent
 		(
