[Editor] EditorLayout fix

diff --git a/Engine/Source/Runtime/Slate/Private/Framework/Docking/LayoutService.cpp b/Engine/Source/Runtime/Slate/Private/Framework/Docking/LayoutService.cpp
--- a/Engine/Source/Runtime/Slate/Private/Framework/Docking/LayoutService.cpp
+++ b/Engine/Source/Runtime/Slate/Private/Framework/Docking/LayoutService.cpp
@@ -2,12 +2,15 @@
 
 #include "Framework/Docking/LayoutService.h"
 #include "HAL/FileManager.h"
+#include "Internationalization/Regex.h"
 #include "Misc/App.h"
 #include "Misc/ConfigCacheIni.h"
 #include "Misc/FileHelper.h"
+#include "Misc/StringUtility.h"
 #include "Dom/JsonObject.h"
 #include "Serialization/JsonReader.h"
 #include "Serialization/JsonSerializer.h"
+#include "String/CaseConversion.h"
 
 DEFINE_LOG_CATEGORY_STATIC(LogLayoutService, Log, All);
 
@@ -26,7 +29,7 @@ static const TCHAR* GetEditorLayoutsSectionName()
 			EditorLayoutsSectionName = DefaultEditorLayoutsSectionName;
 		}
 	}
-	
+
 	return *EditorLayoutsSectionName;
 }
 
@@ -40,6 +43,30 @@ static bool GetConfigString(const TCHAR* Key, FString& Value, const FString& Fil
 		}
 		return false;
 	}
+
+	// Old custom layouts may have an incorrect name, so account for that here.
+	if (const bool bIsUserLayout = Filename.StartsWith(FPaths::EngineUserLayoutDir()))
+	{
+		// Inject corrected label into result
+		FRegexPattern Pattern(TEXT("\"\\w+\"")); // words between commas
+		FRegexMatcher Matcher(Pattern, Value);
+
+		int32 MatchStart = INDEX_NONE;
+		int32 MatchEnd = INDEX_NONE;
+
+		// We want the last match, so run until none found and use the last result
+		while (Matcher.FindNext())
+		{
+			MatchStart = Matcher.GetMatchBeginning();
+			MatchEnd = Matcher.GetMatchEnding();
+		}
+
+		if (MatchStart != INDEX_NONE && MatchEnd != INDEX_NONE)
+		{
+			Value = Value.Left(MatchStart + 1) + FString(WriteToString<16>(UE::String::PascalCase(FPaths::GetBaseFilename(Filename)))) + Value.Right(Value.Len() - MatchEnd + 1);
+		}
+	}
+
 	return true;
 }
 
@@ -128,7 +155,7 @@ static TSharedPtr<FJsonObject> ConvertSectionToJson(const TArray<FString>& Secti
 					bIsJson = false;
 				}
 			}
-			
+
 			if (!bIsJson)
 			{
 				RootObject->SetStringField(Key, Value);
@@ -359,7 +386,7 @@ FText FLayoutSaveRestore::LoadSectionFromConfig(const FString& InConfigFileName,
 
 	FText ValueText;
 	FTextStringHelper::ReadFromBuffer(*ValueString, ValueText, GetEditorLayoutsSectionName());
-	
+
 	return ValueText;
 }
 
