Slate: Fix transform calculation in SFxWidget::OnPaint that was incorrectly using RenderScale as a translation rather than a scale that was causing a 1 px padding offset for carousel widgets

diff --git a/Engine/Source/Runtime/Slate/Private/Widgets/Layout/SFxWidget.cpp b/Engine/Source/Runtime/Slate/Private/Widgets/Layout/SFxWidget.cpp
--- a/Engine/Source/Runtime/Slate/Private/Widgets/Layout/SFxWidget.cpp
+++ b/Engine/Source/Runtime/Slate/Private/Widgets/Layout/SFxWidget.cpp
@@ -62,9 +62,9 @@ int32 SFxWidget::OnPaint( const FPaintArgs& Args, const FGeometry& AllottedGeome
 	const FVector2D ScaleOrigin = RenderScaleOrigin.Get() * AllottedGeometry.GetLocalSize();
 	const FVector2D Offset = VisualOffset.Get() * AllottedGeometry.GetLocalSize();
 	// create the render transform as a scale around ScaleOrigin and offset it by Offset.
-	const FVector2D SlateRenderTransform = Concatenate<FVector2D, FVector2D, FVector2D, FVector2D>(Inverse(ScaleOrigin), FVector2D(RenderScale.Get()), ScaleOrigin, Offset);
+	const FSlateLayoutTransform SlateRenderTransform = Concatenate(FSlateLayoutTransform(Inverse(ScaleOrigin)), FSlateLayoutTransform(RenderScale.Get(), ScaleOrigin), FSlateLayoutTransform(Offset));
 	// This will append the render transform to the layout transform, and we only use it for rendering.
-	FGeometry ModifiedGeometry = AllottedGeometry.MakeChild(AllottedGeometry.GetLocalSize(), FSlateLayoutTransform(SlateRenderTransform));
+	FGeometry ModifiedGeometry = AllottedGeometry.MakeChild(AllottedGeometry.GetLocalSize(), SlateRenderTransform);
 	
 	FArrangedChildren ArrangedChildren(EVisibility::Visible);
 	this->ArrangeChildren(ModifiedGeometry, ArrangedChildren);
