Add prioritized resizing to SClippingHorizontalBox

* Add new FMenuEntryResizeParams for controlling the clipping priority, if clipping is allowed, and if a widget is visible in toolbar overflow menus.
* Add a FMenuEntryResizeParams to each mutli block as FMultiBlock::ResizeParams.
* Add FToolBarStyle::WrapButtonPadding for controlling the padding of wrap buttons.
* Add FToolBarStyle::WrapButtonIndex for controling where wrap buttons appear. Indexing from the left starts at 0 and indexing from the right starts at -1.
* Rewrite SClippingHorizontalBox::OnArrangeChildren to do prioritized resizing of what is returned by SHorizontalBox::OnArrangeChildren. This algorithm supports the resizing paramters described by FMenuEntryResizeParams.
* Add unit tests of the prioritized-resize algorithm to the SlateTests program.

diff --git a/Engine/Source/Programs/SlateTests/SlateTests.Build.cs b/Engine/Source/Programs/SlateTests/SlateTests.Build.cs
--- a/Engine/Source/Programs/SlateTests/SlateTests.Build.cs
+++ b/Engine/Source/Programs/SlateTests/SlateTests.Build.cs
@@ -1,6 +1,8 @@
 // Copyright Epic Games, Inc. All Rights Reserved.
 using UnrealBuildTool;
 
+using System.IO;
+
 public class SlateTests : TestModuleRules
 {
 	static SlateTests()
@@ -47,5 +49,11 @@ public class SlateTests : TestModuleRules
 					"DesktopPlatform"
 				});
 		}
+
+		PrivateIncludePaths.AddRange(
+			new string[] {
+				Path.Combine(Target.RelativeEnginePath, "Source", "Runtime", "Slate", "Private"),
+			}
+		);
 	}
 }
diff --git a/Engine/Source/Programs/SlateTests/Tests/Slate/SlateTest.cpp b/Engine/Source/Programs/SlateTests/Tests/Slate/SlateTest.cpp
--- a/Engine/Source/Programs/SlateTests/Tests/Slate/SlateTest.cpp
+++ b/Engine/Source/Programs/SlateTests/Tests/Slate/SlateTest.cpp
@@ -2,7 +2,302 @@
 
 #include "TestHarness.h"
 
-TEST_CASE("Runtime::Slate::dummy", "[Slate]")
+#include "Framework/MultiBox/SClippingHorizontalBox.h"
+
+#include <iostream>
+
+namespace UE::Slate
+{
+
+bool operator==(const UE::Slate::FClippingInfo& Left, const UE::Slate::FClippingInfo& Right)
+{
+	return Left.Widget == Right.Widget && Left.X == Right.X && Left.Width == Right.Width
+		&& Left.bIsStretchable == Right.bIsStretchable && Left.bAppearsInOverflow == Right.bAppearsInOverflow
+		&& Left.bWasClipped == Right.bWasClipped;
+}
+
+std::ostream& operator<<(std::ostream& Stream, const UE::Slate::FClippingInfo& Value)
+{
+	// Format booleans as "true" and "false" instead of 1 and 0.
+	Stream << std::boolalpha;
+
+	Stream << "FClippingInfo3("
+		   << "Widget=" << Value.Widget.Get() << ", ";
+	Stream << "FMenuEntryResizeParams(" << Value.ResizeParams.ClippingPriority.Get() << ", "
+		   << Value.ResizeParams.AllowClipping.Get() << ", " << Value.ResizeParams.VisibleInOverflow.Get() << "), ";
+	Stream << "X=" << Value.X << ", ";
+	Stream << "Width=" << Value.Width << ", ";
+	Stream << "bIsStretchable=" << Value.bIsStretchable << ", ";
+	Stream << "bAppearsInOverflow=" << Value.bAppearsInOverflow << ", ";
+	Stream << "bWasClipped=" << Value.bWasClipped;
+	Stream << ")";
+	return Stream;
+}
+
+} // namespace UE::Slate
+
+TEST_CASE("Runtime::Slate::PrioritizedResizing works on empty input", "[Slate]")
+{
+	constexpr double AllottedWidth = 200.0;
+	constexpr double WrapButtonWidth = 40.0;
+	const FMargin WrapButtonPadding = FMargin(4.0f, 0);
+	constexpr int32 WrapButtonIndex = -1;
+	TArray<UE::Slate::FClippingInfo> ClippingInfos;
+
+	double WrapButtonX = 0.0;
+	UE::Slate::PrioritizedResize(
+		AllottedWidth, WrapButtonWidth, WrapButtonPadding, WrapButtonIndex, ClippingInfos, WrapButtonX
+	);
+
+	CHECK(ClippingInfos.IsEmpty());
+}
+
+TEST_CASE("Runtime::Slate::PrioritizedResizing clips no widgets when there's plenty of space", "[Slate]")
+{
+	constexpr double AllottedWidth = 500.0;
+	constexpr double WrapButtonWidth = 40.0;
+	const FMargin WrapButtonPadding = FMargin(4.0f, 0);
+	constexpr int32 WrapButtonIndex = -1;
+	TArray<UE::Slate::FClippingInfo> ClippingInfos = {
+		UE::Slate::FClippingInfo{ .X = 0.0,  .Width = 20.0  },
+		UE::Slate::FClippingInfo{ .X = 30.0, .Width = 100.0 },
+	};
+	double WrapButtonX = 0.0;
+
+	UE::Slate::PrioritizedResize(
+		AllottedWidth, WrapButtonWidth, WrapButtonPadding, WrapButtonIndex, ClippingInfos, WrapButtonX
+	);
+
+	for (const UE::Slate::FClippingInfo& Info : ClippingInfos)
+	{
+		CHECK_FALSE(Info.bWasClipped);
+	}
+}
+
+TEST_CASE("Runtime::Slate::PrioritizedResizing with plenty of space does not modify input with default resize params", "[Slate]")
+{
+	constexpr double AllottedWidth = 1000.0;
+	constexpr double WrapButtonWidth = 40.0;
+	const FMargin WrapButtonPadding = FMargin(4.0f, 0);
+	constexpr int32 WrapButtonIndex = -1;
+	const TArray<UE::Slate::FClippingInfo> OriginalClippingInfos = {
+		UE::Slate::FClippingInfo{ .X = 0.0,   .Width = 20.0  },
+        UE::Slate::FClippingInfo{ .X = 20.0,  .Width = 100.0 },
+		UE::Slate::FClippingInfo{ .X = 120.0, .Width = 30.0  },
+        UE::Slate::FClippingInfo{ .X = 150.0, .Width = 40.0  },
+		UE::Slate::FClippingInfo{ .X = 190.0, .Width = 200.0 },
+        UE::Slate::FClippingInfo{ .X = 390.0, .Width = 15.0  },
+	};
+	double WrapButtonX = 0.0;
+
+	TArray<UE::Slate::FClippingInfo> InOutClippingInfos = OriginalClippingInfos;
+	UE::Slate::PrioritizedResize(
+		AllottedWidth, WrapButtonWidth, WrapButtonPadding, WrapButtonIndex, InOutClippingInfos, WrapButtonX
+	);
+
+	CHECK(OriginalClippingInfos.Num() == InOutClippingInfos.Num());
+
+	for (int32 Index = 0; Index < OriginalClippingInfos.Num(); ++Index)
+	{
+		const UE::Slate::FClippingInfo Original = OriginalClippingInfos[Index];
+		const UE::Slate::FClippingInfo Output = InOutClippingInfos[Index];
+
+		CHECK(Original == Output);
+	}
+}
+
+TEST_CASE("Runtime::Slate::PrioritizedResizing with constrained size and two widgets clips the clippable one", "[Slate]")
+{
+	constexpr double AllottedWidth = 90.0;
+	constexpr double WrapButtonWidth = 40.0;
+	const FMargin WrapButtonPadding = FMargin(4.0f, 0);
+	constexpr int32 WrapButtonIndex = -1;
+	TArray<UE::Slate::FClippingInfo> InOutClippingInfos = {
+		UE::Slate::FClippingInfo{ .ResizeParams = FMenuEntryResizeParams{ .AllowClipping = true }, .X = 0.0,  .Width = 50.0 },
+		UE::Slate::FClippingInfo{
+								  .ResizeParams = FMenuEntryResizeParams{ .AllowClipping = false }, .X = 50.0, .Width = 50.0 },
+	};
+	double WrapButtonX = 0.0;
+
+	UE::Slate::PrioritizedResize(
+		AllottedWidth, WrapButtonWidth, WrapButtonPadding, WrapButtonIndex, InOutClippingInfos, WrapButtonX
+	);
+
+	CHECK(InOutClippingInfos[0].bWasClipped == true);
+	CHECK(InOutClippingInfos[1].bWasClipped == false);
+}
+
+TEST_CASE("Runtime::Slate::PrioritizedResizing superfluous space after clipping is distributed to stretching widgets", "[Slate]")
 {
-	CHECK(true);
+	constexpr double WrapButtonWidth = 0.0;
+	const FMargin WrapButtonPadding = FMargin(0.0f, 0);
+	constexpr int32 WrapButtonIndex = -1;
+	const TArray<UE::Slate::FClippingInfo> OriginalClippingInfos = {
+		UE::Slate::FClippingInfo{ .X = 0.0, .Width = 50.0 },
+		UE::Slate::FClippingInfo{ .X = 50.0, .Width = 50.0, .bIsStretchable = true },
+		UE::Slate::FClippingInfo{ .X = 100.0, .Width = 50.0 },
+		UE::Slate::FClippingInfo{ .X = 150.0, .Width = 50.0, .bIsStretchable = true },
+		UE::Slate::FClippingInfo{ .X = 200.0, .Width = 50.0 },
+		UE::Slate::FClippingInfo{ .X = 250.0, .Width = 50.0 },
+	};
+	double WrapButtonX = 0.0;
+
+	// The tests following this rely on the total width of all widgets to be 300 pixels.
+	{
+		double TotalWidthOfWidgets = 0.0;
+		for (const UE::Slate::FClippingInfo& Info : OriginalClippingInfos)
+		{
+			TotalWidthOfWidgets += Info.Width;
+		}
+
+		CHECK(TotalWidthOfWidgets == 300.0);
+	}
+
+	SECTION("When all widgets don't fit")
+	{
+		constexpr double AllottedWidth = 295.0;
+
+		TArray<UE::Slate::FClippingInfo> InOutClippingInfos = OriginalClippingInfos;
+		UE::Slate::PrioritizedResize(
+			AllottedWidth, WrapButtonWidth, WrapButtonPadding, WrapButtonIndex, InOutClippingInfos, WrapButtonX
+		);
+
+		SECTION("Only the last one is clipped")
+		{
+			for (int Index = 0; Index < InOutClippingInfos.Num() - 1; ++Index)
+			{
+				INFO("Index=" << Index);
+				CHECK(InOutClippingInfos[Index].bWasClipped == false);
+			}
+
+			CHECK(InOutClippingInfos.Last().bWasClipped);
+		}
+
+		SECTION("The non-stretching widgets keep their original width")
+		{
+			CHECK(InOutClippingInfos[0].Width == 50.0);
+			// Stretching widget
+			CHECK(InOutClippingInfos[2].Width == 50.0);
+			// Stretching widget
+			CHECK(InOutClippingInfos[4].Width == 50.0);
+			CHECK(InOutClippingInfos[5].Width == 50.0);
+		}
+
+		SECTION("The stretching widgets are expanded")
+		{
+			// Non-stretching widget
+			CHECK(InOutClippingInfos[1].Width > 51.0);
+			// Non-stretching widget
+			CHECK(InOutClippingInfos[3].Width > 51.0);
+			// Non-stretching widget
+			// Non-stretching widget
+		}
+	}
+}
+
+TEST_CASE("Runtime::Slate::PrioritizedResizing all superfluous space is eaten by stretching widgets", "[Slate]")
+{
+	constexpr double WrapButtonWidth = 0.0;
+	const FMargin WrapButtonPadding = FMargin(0.0f, 0);
+	constexpr int32 WrapButtonIndex = -1;
+	const TArray<UE::Slate::FClippingInfo> OriginalClippingInfos = {
+		UE::Slate::FClippingInfo{ .X = 0.0, .Width = 50.0 },
+		UE::Slate::FClippingInfo{ .X = 50.0, .Width = 50.0, .bIsStretchable = true },
+		UE::Slate::FClippingInfo{ .X = 100.0, .Width = 50.0 },
+		UE::Slate::FClippingInfo{ .X = 150.0, .Width = 50.0, .ResizeParams{ .ClippingPriority = -100 } }, // Clips  first.
+		UE::Slate::FClippingInfo{ .X = 200.0, .Width = 50.0 },
+	};
+	double WrapButtonX = 0.0;
+
+	SECTION("When all widgets don't fit")
+	{
+		constexpr double AllottedWidth = 245.0;
+
+		TArray<UE::Slate::FClippingInfo> InOutClippingInfos = OriginalClippingInfos;
+		UE::Slate::PrioritizedResize(
+			AllottedWidth, WrapButtonWidth, WrapButtonPadding, WrapButtonIndex, InOutClippingInfos, WrapButtonX
+		);
+
+		SECTION("The stretching widget eats all space")
+		{
+			CHECK(InOutClippingInfos[1].Width == AllottedWidth - 50 - 50 - 50);
+		}
+
+		SECTION("The widget to the right of the stretching widget is moved to the right")
+		{
+			CHECK(InOutClippingInfos[2].X == AllottedWidth - 50 - 50);
+		}
+	}
+}
+
+TEST_CASE("Runtime::Slate::PrioritizedResizing sorts zero-width widgets identical to input regardless of clipping priority", "[Slate]")
+{
+	constexpr double WrapButtonWidth = 0.0;
+	const FMargin WrapButtonPadding = FMargin(0.0f, 0);
+	constexpr int32 WrapButtonIndex = -1;
+
+	const TArray<UE::Slate::FClippingInfo> OriginalClippingInfos = {
+		UE::Slate::FClippingInfo{ .X = 0, .Width = 0,   .ResizeParams{ .ClippingPriority = 2 }, .bIsStretchable = true },
+		UE::Slate::FClippingInfo{ .X = 0, .Width = 0,   .ResizeParams{ .ClippingPriority = 4 }, .bIsStretchable = true },
+		UE::Slate::FClippingInfo{ .X = 0, .Width = 100, .ResizeParams{ .ClippingPriority = 3 }, .bIsStretchable = true },
+		UE::Slate::FClippingInfo{ .X = 0, .Width = 0,   .ResizeParams{ .ClippingPriority = 1 }, .bIsStretchable = true },
+	};
+	double WrapButtonX = 0.0;
+
+	SECTION("When size is contrained")
+	{
+		constexpr double AllottedWidth = 50.0;
+
+		TArray<UE::Slate::FClippingInfo> InOutClippingInfos = OriginalClippingInfos;
+		UE::Slate::PrioritizedResize(
+			AllottedWidth, WrapButtonWidth, WrapButtonPadding, WrapButtonIndex, InOutClippingInfos, WrapButtonX
+		);
+
+		SECTION("The widget sorting hasn't changed")
+		{
+			CHECK(InOutClippingInfos[0].ResizeParams.ClippingPriority.Get() == 2);
+			CHECK(InOutClippingInfos[1].ResizeParams.ClippingPriority.Get() == 4);
+			CHECK(InOutClippingInfos[2].ResizeParams.ClippingPriority.Get() == 3);
+			CHECK(InOutClippingInfos[3].ResizeParams.ClippingPriority.Get() == 1);
+		}
+	}
+}
+
+TEST_CASE("Runtime::Slate::PrioritizedResizing positions a wrap button at index 1 snugly", "[Slate]")
+{
+	constexpr double WrapButtonWidth = 40;
+	const FMargin WrapButtonPadding = FMargin(0.0f, 0);
+	constexpr int32 WrapButtonIndex = 1;
+
+	const TArray<UE::Slate::FClippingInfo> OriginalClippingInfos = {
+		UE::Slate::FClippingInfo{ .X = 0,   .Width = 50 },
+		UE::Slate::FClippingInfo{ .X = 50,  .Width = 50 },
+		UE::Slate::FClippingInfo{ .X = 100, .Width = 50 },
+		UE::Slate::FClippingInfo{ .X = 150, .Width = 50 },
+	};
+	double WrapButtonX = 0.0;
+
+	SECTION("When size is contrained")
+	{
+		constexpr double AllottedWidth = 190;
+
+		TArray<UE::Slate::FClippingInfo> InOutClippingInfos = OriginalClippingInfos;
+		UE::Slate::PrioritizedResize(
+			AllottedWidth, WrapButtonWidth, WrapButtonPadding, WrapButtonIndex, InOutClippingInfos, WrapButtonX
+		);
+
+		SECTION("The wrap button is positioned snugly with the non-clipped widgets")
+		{
+			CHECK(InOutClippingInfos[0].X == 0);
+			CHECK(InOutClippingInfos[0].Width == 50);
+
+			CHECK(WrapButtonX == 50);
+
+			CHECK(InOutClippingInfos[1].X == 90);
+			CHECK(InOutClippingInfos[1].Width == 50);
+
+			CHECK(InOutClippingInfos[2].X == 140);
+			CHECK(InOutClippingInfos[2].Width == 50);
+		}
+	}
 }
diff --git a/Engine/Source/Runtime/Slate/Private/Framework/MultiBox/MultiBox.cpp b/Engine/Source/Runtime/Slate/Private/Framework/MultiBox/MultiBox.cpp
--- a/Engine/Source/Runtime/Slate/Private/Framework/MultiBox/MultiBox.cpp
+++ b/Engine/Source/Runtime/Slate/Private/Framework/MultiBox/MultiBox.cpp
@@ -273,6 +273,16 @@ void FMultiBlock::SetVisibilityOverride(TAttribute<EVisibility> InVisibilityOver
 	VisibilityOverride = InVisibilityOverride;
 }
 
+FMenuEntryResizeParams FMultiBlock::GetResizeParams() const
+{
+	return ResizeParams.Get();
+}
+
+void FMultiBlock::SetResizeParams(const TAttribute<FMenuEntryResizeParams>& InResizeParams)
+{
+	ResizeParams = InResizeParams;
+}
+
 /**
  * Constructor
  *
@@ -1082,11 +1092,26 @@ void SMultiBoxWidget::BuildMultiBoxWidget()
 	case EMultiBoxType::ToolBar:
 	case EMultiBoxType::SlimHorizontalToolBar:
 		{
-			MainWidget = HorizontalBox = ClippedHorizontalBox = SNew(SClippingHorizontalBox)
-				.OnWrapButtonClicked(FOnGetContent::CreateSP(this, &SMultiBoxWidget::OnWrapButtonClicked))
-				.IsFocusable(MultiBox->bIsFocusable)
-				.StyleSet(StyleSet)
-				.StyleName(StyleName);
+			MainWidget = HorizontalBox = ClippedHorizontalBox =
+				SNew(SClippingHorizontalBox)
+					.OnWrapButtonClicked(FOnGetContent::CreateSP(this, &SMultiBoxWidget::OnWrapButtonClicked))
+					.OnGetWidgetResizeParams(FOnGetWidgetResizeParams::CreateLambda(
+						[this](const TSharedRef<SWidget>& InWidget) -> FMenuEntryResizeParams
+						{
+							for (const TPair<TSharedPtr<SWidget>, TSharedPtr<const FMultiBlock>>& Pair : FinalMultiBoxWidgets)
+							{
+								if (InWidget == Pair.Key)
+								{
+									return Pair.Value->GetResizeParams();
+								}
+							}
+
+							return FMenuEntryResizeParams();
+						}
+					))
+					.IsFocusable(MultiBox->bIsFocusable)
+					.StyleSet(StyleSet)
+					.StyleName(StyleName);
 		}
 		break;
 	case EMultiBoxType::VerticalToolBar:
diff --git a/Engine/Source/Runtime/Slate/Private/Framework/MultiBox/MultiBoxBuilder.cpp b/Engine/Source/Runtime/Slate/Private/Framework/MultiBox/MultiBoxBuilder.cpp
--- a/Engine/Source/Runtime/Slate/Private/Framework/MultiBox/MultiBoxBuilder.cpp
+++ b/Engine/Source/Runtime/Slate/Private/Framework/MultiBox/MultiBoxBuilder.cpp
@@ -466,7 +466,8 @@ void FToolBarBuilder::AddToolBarButton(
 	FName InTutorialHighlightName,
 	FNewMenuDelegate InCustomMenuDelegate,
 	TAttribute<EVisibility> InVisibilityOverride,
-	TAttribute<FText> InToolbarLabelOverride
+	TAttribute<FText> InToolbarLabelOverride,
+	const TAttribute<FMenuEntryResizeParams>& InResizeParams
 )
 {
 	ApplySectionBeginning();
@@ -488,6 +489,7 @@ void FToolBarBuilder::AddToolBarButton(
 	NewToolBarButtonBlock->SetStyleNameOverride(CurrentStyleOverride);
 	NewToolBarButtonBlock->SetCustomMenuDelegate(InCustomMenuDelegate);
 	NewToolBarButtonBlock->SetVisibilityOverride(InVisibilityOverride);
+	NewToolBarButtonBlock->SetResizeParams(InResizeParams);
 
 	MultiBox->AddMultiBlock( NewToolBarButtonBlock );
 
@@ -549,7 +551,8 @@ void FToolBarBuilder::AddToolBarButton(
 	const EUserInterfaceActionType UserInterfaceActionType,
 	FName InTutorialHighlightName,
 	TAttribute<EVisibility> InVisibilityOverride,
-	TAttribute<FText> InToolbarLabelOverride
+	TAttribute<FText> InToolbarLabelOverride,
+	const TAttribute<FMenuEntryResizeParams>& InResizeParams
 )
 {
 	ApplySectionBeginning();
@@ -570,6 +573,7 @@ void FToolBarBuilder::AddToolBarButton(
 	NewToolBarButtonBlock->SetTutorialHighlightName(GenerateTutorialIdentifierName(TutorialHighlightName, InTutorialHighlightName, nullptr, MultiBox->GetBlocks().Num()));
 	NewToolBarButtonBlock->SetStyleNameOverride(CurrentStyleOverride);
 	NewToolBarButtonBlock->SetVisibilityOverride(InVisibilityOverride);
+	NewToolBarButtonBlock->SetResizeParams(InResizeParams);
 
 	MultiBox->AddMultiBlock( NewToolBarButtonBlock );
 
@@ -587,7 +591,8 @@ void FToolBarBuilder::AddComboButton(
 	TAttribute<EVisibility> InVisibilityOverride,
 	TAttribute<FText> InToolbarLabelOverride,
 	TAttribute<EMenuPlacement> InPlacementOverride,
-	const EUserInterfaceActionType InUserInterfaceActionType
+	const EUserInterfaceActionType InUserInterfaceActionType,
+	const TAttribute<FMenuEntryResizeParams>& InResizeParams
 )
 {
 	ApplySectionBeginning();
@@ -605,6 +610,7 @@ void FToolBarBuilder::AddComboButton(
 	NewToolBarComboButtonBlock->SetTutorialHighlightName(GenerateTutorialIdentifierName(TutorialHighlightName, InTutorialHighlightName, nullptr, MultiBox->GetBlocks().Num()));
 	NewToolBarComboButtonBlock->SetStyleNameOverride(CurrentStyleOverride);
 	NewToolBarComboButtonBlock->SetVisibilityOverride(InVisibilityOverride);
+	NewToolBarComboButtonBlock->SetResizeParams(InResizeParams);
 
 	MultiBox->AddMultiBlock( NewToolBarComboButtonBlock );
 }
@@ -699,7 +705,8 @@ void FToolBarBuilder::AddWidget(
 	FName InTutorialHighlightName,
 	bool bInSearchable,
 	FNewMenuDelegate InCustomMenuDelegate,
-	TAttribute<EVisibility> InVisibilityOverride
+	TAttribute<EVisibility> InVisibilityOverride,
+	const TAttribute<FMenuEntryResizeParams>& InResizeParams
 )
 {
 	FMenuEntryStyleParams StyleParams = InStyleParams;
@@ -708,7 +715,9 @@ void FToolBarBuilder::AddWidget(
 		StyleParams.VerticalAlignment = VAlign_Center; // Default to centered, versions prior to 5.5 decide alignment based on label presence
 	}
 
-	AddWidgetInternal(InWidget, StyleParams, InTutorialHighlightName, bInSearchable, InCustomMenuDelegate, InVisibilityOverride);
+	AddWidgetInternal(
+		InWidget, StyleParams, InTutorialHighlightName, bInSearchable, InCustomMenuDelegate, InVisibilityOverride, InResizeParams
+	);
 }
 
 void FToolBarBuilder::AddWidgetInternal(
@@ -717,7 +726,8 @@ void FToolBarBuilder::AddWidgetInternal(
 	FName InTutorialHighlightName,
 	bool bInSearchable,
 	FNewMenuDelegate InCustomMenuDelegate,
-	TAttribute<EVisibility> InVisibilityOverride
+	TAttribute<EVisibility> InVisibilityOverride,
+	const TAttribute<FMenuEntryResizeParams>& InResizeParams
 )
 {
 	ApplySectionBeginning();
@@ -741,9 +751,12 @@ void FToolBarBuilder::AddWidgetInternal(
 	NewWidgetBlock->SetSearchable(bInSearchable);
 	NewWidgetBlock->SetCustomMenuDelegate(InCustomMenuDelegate);
 	NewWidgetBlock->SetVisibilityOverride(InVisibilityOverride);
+	NewWidgetBlock->SetResizeParams(InResizeParams);
 }
 
-void FToolBarBuilder::AddSeparator(FName InExtensionHook, TAttribute<EVisibility> InVisibilityOverride)
+void FToolBarBuilder::AddSeparator(
+	FName InExtensionHook, TAttribute<EVisibility> InVisibilityOverride, const TAttribute<FMenuEntryResizeParams>& InResizeParams
+)
 {
 	ApplySectionBeginning();
 
@@ -752,6 +765,7 @@ void FToolBarBuilder::AddSeparator(FName InExtensionHook, TAttribute<EVisibility
 	TSharedRef<FToolBarSeparatorBlock> NewSeparatorBlock = MakeShared<FToolBarSeparatorBlock>(InExtensionHook);
 	NewSeparatorBlock->SetStyleNameOverride(CurrentStyleOverride);
 	NewSeparatorBlock->SetVisibilityOverride(InVisibilityOverride);
+	NewSeparatorBlock->SetResizeParams(InResizeParams);
 
 	MultiBox->AddMultiBlock(NewSeparatorBlock);
 
diff --git a/Engine/Source/Runtime/Slate/Private/Framework/MultiBox/SClippingHorizontalBox.cpp b/Engine/Source/Runtime/Slate/Private/Framework/MultiBox/SClippingHorizontalBox.cpp
--- a/Engine/Source/Runtime/Slate/Private/Framework/MultiBox/SClippingHorizontalBox.cpp
+++ b/Engine/Source/Runtime/Slate/Private/Framework/MultiBox/SClippingHorizontalBox.cpp
@@ -9,87 +9,331 @@
 #include "Widgets/Images/SImage.h"
 #include "Widgets/Input/SComboButton.h"
 
-void SClippingHorizontalBox::OnArrangeChildren( const FGeometry& AllottedGeometry, FArrangedChildren& ArrangedChildren ) const
+namespace UE::Slate
 {
-	// If WrapButton hasn't been initialized, that means AddWrapButton() hasn't 
-	// been called and this method isn't going to behave properly
-	check(WrapButton.IsValid());
 
-	SHorizontalBox::OnArrangeChildren(AllottedGeometry, ArrangedChildren);
+void PrioritizedResize(
+	double InAllottedWidth,
+	double InWrapButtonWidth,
+	const FMargin& InWrapButtonPadding,
+	int32 InWrapButtonIndex,
+	TArray<FClippingInfo>& InOutClippingInfos,
+	double& OutWrapButtonX
+)
+{
+	double WidthOfAllChildren = 0.0;
+	double WidthOfAllChildrenWhenWrapButtonIsNeeded = 0.0;
+	double NonClippingWidgetWidths = 0.0;
+	for (const UE::Slate::FClippingInfo& Info : InOutClippingInfos)
+	{
+		WidthOfAllChildren += Info.Width;
 
-	// Clear our previously clipped widgets. We'll repopulate this below.
-	ClippedWidgets.Reset();
+		// Find the combined width of widgets that either appear in the overflow menu or never clips, i.e. the widgets
+		// that compete with the wrap button for space. Any other widgets than these we'll just clip before adding a wrap button.
+		if (Info.ResizeParams.VisibleInOverflow.Get(FMenuEntryResizeParams::DefaultVisibleInOverflow)
+			|| !Info.ResizeParams.AllowClipping.Get(FMenuEntryResizeParams::DefaultAllowClipping))
+		{
+			WidthOfAllChildrenWhenWrapButtonIsNeeded += Info.Width;
+		}
 
-	// Remove children that are clipped by the allotted geometry
-	const int32 NumChildren = ArrangedChildren.Num(); 
-	int32 IndexClippedAt = NumChildren;
-	for (int32 ChildIdx = NumChildren - 2; ChildIdx >= 0; --ChildIdx)
+		// Find the combined width of widgets that should never clip.
+		if (!Info.ResizeParams.AllowClipping.Get(FMenuEntryResizeParams::DefaultAllowClipping))
+		{
+			NonClippingWidgetWidths += Info.Width;
+		}
+	}
+
+	const bool bNeedsClipping = WidthOfAllChildren > InAllottedWidth;
+	const bool bNeedsWrapButton = WidthOfAllChildrenWhenWrapButtonIsNeeded > InAllottedWidth;
+
+	// Early out if we don't need to clip.
+	if (!bNeedsClipping)
 	{
-		const FArrangedWidget& CurWidget = ArrangedChildren[ChildIdx];
-		// Ceil (Minus a tad for float precision) to ensure contents are not a sub-pixel larger than the box, which will create an unnecessary wrap button
-		FVector2D AbsWidgetPos = CurWidget.Geometry.LocalToAbsolute(CurWidget.Geometry.GetLocalSize());
-		FVector2D AbsBoxPos = FVector2D(AllottedGeometry.AbsolutePosition) + AllottedGeometry.GetLocalSize() * AllottedGeometry.Scale;
-		const FVector2D WidgetPosRounded = FVector2D(FMath::CeilToInt(AbsWidgetPos.X - KINDA_SMALL_NUMBER), FMath::CeilToInt(AbsWidgetPos.Y - KINDA_SMALL_NUMBER));
-		const FVector2D BoxPosRounded = FVector2D(FMath::CeilToInt(AbsBoxPos.X - KINDA_SMALL_NUMBER), FMath::CeilToInt(AbsBoxPos.Y - KINDA_SMALL_NUMBER));
-		if (WidgetPosRounded.X > BoxPosRounded.X)
+		return;
+	}
+
+	// Clip widgets in priority order.
+	{
+		// Construct a priority-sorted array of pointers into the clipping infos. Work on this copy
+		// to not mess up the original sorting of the input widgets.
+		TArray<UE::Slate::FClippingInfo*> PrioritySortedClippingInfos;
+		for (UE::Slate::FClippingInfo& Info : InOutClippingInfos)
+		{
+			PrioritySortedClippingInfos.Add(&Info);
+		}
+		PrioritySortedClippingInfos.StableSort(
+			[](const UE::Slate::FClippingInfo& A, const UE::Slate::FClippingInfo& B) -> bool
+			{
+				return A.ResizeParams.ClippingPriority.Get(FMenuEntryResizeParams::DefaultClippingPriority)
+					 > B.ResizeParams.ClippingPriority.Get(FMenuEntryResizeParams::DefaultClippingPriority);
+			}
+		);
+
+		// Iterate our arranged children in resize-priority order, accumulate their width, and start removing widgets
+		// once their combined width is larger than our allotted width. Start counting at the width required by
+		// non-clipping widgets to make sure we start clipping clippable widgets immediately once non-clippable widths
+		// and the accumulator sum up to our allotted width.
 		{
-			// Insert the clipped widget at the start of the array to keep the array sorted with
-			// the left-most clipped widget first and the right-most clipped widget last.
-			ClippedWidgets.Insert(CurWidget.Widget.ToWeakPtr(), 0);
+			double WidthAccumulator = NonClippingWidgetWidths;
+			// If we need a wrap button, clip widgets as if the wrap button is already taking up space.
+			if (bNeedsWrapButton)
+			{
+				WidthAccumulator += InWrapButtonWidth + InWrapButtonPadding.Left + InWrapButtonPadding.Right;
+			}
+
+			for (UE::Slate::FClippingInfo* Info : PrioritySortedClippingInfos)
+			{
+				// Don't clip non-clipping widgets and also don't count their width, we've already done that above.
+				if (!Info->ResizeParams.AllowClipping.Get(FMenuEntryResizeParams::DefaultAllowClipping))
+				{
+					continue;
+				}
 
-			ArrangedChildren.Remove(ChildIdx);
-			IndexClippedAt = ChildIdx;
+				WidthAccumulator += Info->Width;
+
+				// We are wider than our allotted width, so mark this widget for clipping.
+				if (WidthAccumulator > InAllottedWidth)
+				{
+					Info->bWasClipped = true;
+					Info->bAppearsInOverflow =
+						Info->ResizeParams.VisibleInOverflow.Get(FMenuEntryResizeParams::DefaultVisibleInOverflow);
+				}
+			}
 		}
 	}
 
-	if (IndexClippedAt == NumChildren)
+	// Sort blocks by X position. Use stable sort to keep the original order if priorities are the same.
+	InOutClippingInfos.StableSort(
+		[](const UE::Slate::FClippingInfo& A, const UE::Slate::FClippingInfo& B) -> bool
+		{
+			return A.X < B.X;
+		}
+	);
+
+	// Iterate widgets left to right and if clipped widgets are found, translate all subsequent widgets to the left to compensate.
+	double WidthOfRemovedWidgets = 0.0;
+	for (UE::Slate::FClippingInfo& Info : InOutClippingInfos)
 	{
-		// Note do not increment NumClippedChildren here. This is to remove the wrap button
-		// None of the children are being clipped, so remove the wrap button
-		ArrangedChildren.Remove(ArrangedChildren.Num() - 1);
+		if (Info.bWasClipped)
+		{
+			WidthOfRemovedWidgets += Info.Width;
+			continue;
+		}
+		Info.X -= WidthOfRemovedWidgets;
 	}
-	else if (ArrangedChildren.Num() > 0)
+
+	// Here we want to EXPAND again to fill up the whole post-clip toolbar if possible. The reason is because we've
+	// reduced the width of the toolbar in "integer steps", i.e. by clipping whole widgets, and in general we'll have
+	// some slack space left.
+	//
+	// An example of widgets to expand are variable-sized spacers that were previously shrunk to make things fit.
 	{
-		// Insert and right align the wrap button, for occlusion checking and insertion below
-		FArrangedWidget& ArrangedWrapButton = ArrangedChildren[ArrangedChildren.Num() - 1];
-		FGeometry& WrapButtonGeometry = ArrangedWrapButton.Geometry;
+		// Count our children that can stretch that were not clipped above.
+		const int32 NumStretchWidgets = InOutClippingInfos
+											.FilterByPredicate(
+												[](const UE::Slate::FClippingInfo& Info)
+												{
+													return Info.bIsStretchable && !Info.bWasClipped;
+												}
+											)
+											.Num();
 
-		if (const bool bHasSpaceForWrapButton = WrapButtonWidth <= AllottedGeometry.GetLocalSize().X)
+		if (NumStretchWidgets > 0)
 		{
-			const float AdjustedWrapButtonWidth = FMath::Min(AllottedGeometry.GetLocalSize().X, WrapButtonWidth);
+			// Distribute our extra space to our stretching children.
+			const double ExtraSpace =
+				InAllottedWidth - WidthOfAllChildren
+				- (bNeedsWrapButton ? InWrapButtonWidth + InWrapButtonPadding.Left + InWrapButtonPadding.Right : 0.0)
+				+ WidthOfRemovedWidgets;
+			const double ExtraSpacePerStretchChild = ExtraSpace / NumStretchWidgets;
+			{
+				double StretchSpaceAdded = 0.0;
+				for (UE::Slate::FClippingInfo& Info : InOutClippingInfos)
+				{
+					if (Info.bWasClipped)
+					{
+						continue;
+					}
 
-			const FVector2D WrapButtonSize = FVector2D(AdjustedWrapButtonWidth, WrapButtonGeometry.GetLocalSize().Y);
-			WrapButtonGeometry = AllottedGeometry.MakeChild(
-				WrapButtonSize,
-				FSlateLayoutTransform(AllottedGeometry.GetLocalSize() - WrapButtonSize));
+					Info.X += StretchSpaceAdded;
+					if (Info.bIsStretchable)
+					{
+						Info.Width += ExtraSpacePerStretchChild;
+						StretchSpaceAdded += ExtraSpacePerStretchChild;
+					}
+				}
+			}
+		}
+	}
 
-			const int32 WrapButtonXPosition = FMath::TruncToInt(WrapButtonGeometry.AbsolutePosition.X);
+	// Add the wrap button. Move widgets to make space for it and then position it.
+	const bool bHasSpaceForWrapButton = InWrapButtonWidth <= InAllottedWidth;
+	if (bNeedsWrapButton && bHasSpaceForWrapButton)
+	{
+		// Use an array that excludes clipped widgets for sanity while indexing. Each element points
+		// into InOutClippingInfos, so we don't need to consolidate the arrays later.
+		TArray<FClippingInfo*> NonClippedInfos;
+		for (FClippingInfo& Info : InOutClippingInfos)
+		{
+			if (!Info.bWasClipped)
+			{
+				NonClippedInfos.Add(&Info);
+			}
+		}
+
+		const bool bIsIndexingFromLeft = InWrapButtonIndex >= 0;
+
+		const int32 StolenIndex = FMath::Clamp(
+			bIsIndexingFromLeft ? InWrapButtonIndex : NonClippedInfos.Num() + InWrapButtonIndex, 0, NonClippedInfos.Num() - 1
+		);
+
+		// Sum up the width of widgets to the left of the wrap button index.
+		double WidthOfWidgetsBeforeStolenIndex = 0.0;
+		for (int32 Index = 0; Index < StolenIndex; ++Index)
+		{
+			WidthOfWidgetsBeforeStolenIndex += NonClippedInfos[Index]->Width;
+		}
 
-			// Further remove any children that the wrap button overlaps with
-			for (int32 ChildIdx = IndexClippedAt - 1; ChildIdx >= 0; --ChildIdx)
+		OutWrapButtonX = 0.0;
+		int32 IndexFromWhichToPushToTheRight = 0;
+		// If we're indexing from the left, we want to add the wrap button to the left of the widget originally at the stolen index.
+		if (bIsIndexingFromLeft)
+		{
+			OutWrapButtonX = WidthOfWidgetsBeforeStolenIndex + InWrapButtonPadding.Left;
+			IndexFromWhichToPushToTheRight = StolenIndex;
+		}
+		// If we're indexing from the right, we want to add the wrap button to the right of the widget originally at the stolen index.
+		else
+		{
+			const double WidthOfWidgetAtStolenIndex =
+				StolenIndex < NonClippedInfos.Num() ? NonClippedInfos[StolenIndex]->Width : 0.0;
+
+			OutWrapButtonX = WidthOfWidgetsBeforeStolenIndex + WidthOfWidgetAtStolenIndex + InWrapButtonPadding.Left;
+			IndexFromWhichToPushToTheRight = StolenIndex + 1;
+		}
+
+		// Special case: if the wrap button is set to appear at the right-most index (WrapButtonIndex = -1), then always
+		// position it to the right, allowing additional space to appear to the left of the wrap button.
+		if (InWrapButtonIndex == -1)
+		{
+			OutWrapButtonX = InAllottedWidth - InWrapButtonWidth - InWrapButtonPadding.Right;
+		}
+
+		const double PaddedWidthOfWrapButton = InWrapButtonWidth + InWrapButtonPadding.Left + InWrapButtonPadding.Right;
+		for (int32 Index = IndexFromWhichToPushToTheRight; !NonClippedInfos.IsEmpty() && Index < NonClippedInfos.Num();
+			 ++Index)
+		{
+			NonClippedInfos[Index]->X += PaddedWidthOfWrapButton;
+		}
+	}
+}
+
+} // namespace UE::Slate
+
+void SClippingHorizontalBox::OnArrangeChildren(const FGeometry& AllottedGeometry, FArrangedChildren& ArrangedChildren) const
+{
+	// If WrapButton hasn't been initialized, that means AddWrapButton() hasn't
+	// been called and this method isn't going to behave properly
+	check(WrapButton.IsValid());
+
+	SHorizontalBox::OnArrangeChildren(AllottedGeometry, ArrangedChildren);
+
+	// Clear our previously clipped widgets. We'll repopulate them below.
+	ClippedWidgets.Reset();
+
+	// Build clipping info for all children.
+	TArray<UE::Slate::FClippingInfo> ClippingInfos;
+	for (int i = 0; i < ArrangedChildren.Num() - 1 /* Skip the wrap button. */; ++i)
+	{
+		const FArrangedWidget& Child = ArrangedChildren[i];
+
+		UE::Slate::FClippingInfo& Info = ClippingInfos.AddDefaulted_GetRef();
+
+		Info.Widget = Child.Widget;
+		if (OnGetWidgetResizeParams.IsBound())
+		{
+			Info.ResizeParams = OnGetWidgetResizeParams.Execute(Child.Widget);
+		}
+		Info.X = Child.Geometry.GetLocalPositionAtCoordinates(FVector2f::ZeroVector).X;
+		Info.Width = Child.Geometry.GetLocalSize().X;
+
+		// Set Info.bIsStretchable by finding if the slot of Info.Widget uses SizeRule_StretchContent.
+		for (int32 ChildrenIndex = Children.Num() - 1; ChildrenIndex >= 0; --ChildrenIndex)
+		{
+			if (Info.Widget == Children[ChildrenIndex].GetWidget())
 			{
-				const FArrangedWidget& CurWidget = ArrangedChildren[ChildIdx];
-				if (FMath::TruncToInt(CurWidget.Geometry.AbsolutePosition.X + CurWidget.Geometry.GetLocalSize().X * CurWidget.Geometry.Scale) > WrapButtonXPosition)
+				if (Children[ChildrenIndex].GetSizeRule() == FSizeParam::SizeRule_Stretch
+					|| Children[ChildrenIndex].GetSizeRule() == FSizeParam::SizeRule_StretchContent)
 				{
-					// Insert the clipped widget at the start of the array to keep the array sorted with
-					// the left-most clipped widget first and the right-most clipped widget last.
-					ClippedWidgets.Insert(CurWidget.Widget.ToWeakPtr(), 0);
-
-					ArrangedChildren.Remove(ChildIdx);
+					Info.bIsStretchable = true;
+					break;
 				}
 			}
 		}
-		else // No space left for anything including WrapButton
+	}
+
+	const double InAllottedWidth = AllottedGeometry.GetLocalSize().X;
+	const FToolBarStyle& ToolBarStyle = StyleSet->GetWidgetStyle<FToolBarStyle>(StyleName);
+	double WrapButtonX = 0.0;
+	UE::Slate::PrioritizedResize(
+		InAllottedWidth, WrapButtonWidth, ToolBarStyle.WrapButtonPadding, WrapButtonIndex, ClippingInfos, WrapButtonX
+	);
+
+	const int32 NumClippedWidgets = ClippingInfos
+										.FilterByPredicate(
+											[](const UE::Slate::FClippingInfo& Info)
+											{
+												return Info.bWasClipped;
+											}
+										)
+										.Num();
+
+	if (NumClippedWidgets == 0)
+	{
+		// None of the children are being clipped, so remove the wrap button and early out.
+		ArrangedChildren.Remove(ArrangedChildren.Num() - 1);
+
+		return;
+	}
+
+	for (const UE::Slate::FClippingInfo& Info : ClippingInfos)
+	{
+		if (Info.bAppearsInOverflow)
 		{
-			for (int32 ChildIdx = 0; ChildIdx < ArrangedChildren.Num() - 1 /* Exclude the wrap button. */; ++ChildIdx)
+			ClippedWidgets.Add(Info.Widget.ToWeakPtr());
+		}
+	}
+
+	// Position the wrap button.
+	{
+		FGeometry& WrapButtonGeometry = ArrangedChildren[ArrangedChildren.Num() - 1].Geometry;
+		const FVector2D WrapButtonSize = FVector2D(WrapButtonWidth, WrapButtonGeometry.GetLocalSize().Y);
+		WrapButtonGeometry = AllottedGeometry.MakeChild(WrapButtonSize, FSlateLayoutTransform(FVector2f(WrapButtonX, 0)));
+	}
+
+	// Finally, position all children using the clipping information.
+	for (int32 Index = 0; Index < ArrangedChildren.Num() - 1; ++Index)
+	{
+		const UE::Slate::FClippingInfo* const Info = ClippingInfos.FindByPredicate(
+			[Widget = ArrangedChildren[Index].Widget](const UE::Slate::FClippingInfo& Info)
 			{
-				const FArrangedWidget& CurWidget = ArrangedChildren[ChildIdx];
-				// Insert the clipped widget at the start of the array to keep the array sorted with
-				// the left-most clipped widget first and the right-most clipped widget last.
-				ClippedWidgets.Insert(CurWidget.Widget.ToWeakPtr(), 0);
+				return Info.Widget == Widget;
 			}
-			ArrangedChildren.Empty();
+		);
+
+		check(Info);
+
+		// Remove the widget if it was clipped.
+		if (Info->bWasClipped)
+		{
+			ArrangedChildren.Remove(Index);
+			--Index;
+			continue;
 		}
+
+		FGeometry& Geometry = ArrangedChildren[Index].Geometry;
+		const FVector2D WrapButtonSize = FVector2D(Info->Width, Geometry.GetLocalSize().Y);
+		Geometry = AllottedGeometry.MakeChild(WrapButtonSize, FSlateLayoutTransform(FVector2f(Info->X, 0)));
 	}
 }
 
@@ -116,6 +360,7 @@ void SClippingHorizontalBox::Construct( const FArguments& InArgs )
 	StyleSet = InArgs._StyleSet;
 	StyleName = InArgs._StyleName;
 	bIsFocusable = InArgs._IsFocusable;
+	OnGetWidgetResizeParams = InArgs._OnGetWidgetResizeParams;
 }
 
 void SClippingHorizontalBox::AddWrapButton()
@@ -143,6 +388,7 @@ void SClippingHorizontalBox::AddWrapButton()
 	// Perform a prepass to get a valid DesiredSize value below
 	WrapButton->SlatePrepass(1.0f);
 	WrapButtonWidth = WrapButton->GetDesiredSize().X;
+	WrapButtonIndex = ToolBarStyle.WrapButtonIndex;
 
 	// Add the wrap button
 	AddSlot()
diff --git a/Engine/Source/Runtime/Slate/Private/Framework/MultiBox/SClippingHorizontalBox.h b/Engine/Source/Runtime/Slate/Private/Framework/MultiBox/SClippingHorizontalBox.h
--- a/Engine/Source/Runtime/Slate/Private/Framework/MultiBox/SClippingHorizontalBox.h
+++ b/Engine/Source/Runtime/Slate/Private/Framework/MultiBox/SClippingHorizontalBox.h
@@ -3,6 +3,7 @@
 #pragma once
 
 #include "CoreMinimal.h"
+#include "Framework/MultiBox/MultiBoxDefs.h"
 #include "Framework/SlateDelegates.h"
 #include "Styling/CoreStyle.h"
 #include "Widgets/DeclarativeSyntaxSupport.h"
@@ -13,6 +14,33 @@ class FPaintArgs;
 class FSlateWindowElementList;
 class SComboButton;
 
+DECLARE_DELEGATE_RetVal_OneParam(FMenuEntryResizeParams, FOnGetWidgetResizeParams, const TSharedRef<SWidget>&);
+
+namespace UE::Slate
+{
+
+struct FClippingInfo
+{
+	TSharedPtr<SWidget> Widget;
+	FMenuEntryResizeParams ResizeParams;
+	double X = 0.0;
+	double Width = 0.0;
+	bool bIsStretchable = false;
+	bool bAppearsInOverflow = false;
+	bool bWasClipped = false;
+};
+
+void PrioritizedResize(
+	double InAllottedWidth,
+	double InWrapButtonWidth,
+	const FMargin& InWrapButtonPadding,
+	int32 InWrapButtonIndex,
+	TArray<FClippingInfo>& InOutClippingInfos,
+	double& OutWrapButtonX
+);
+
+}; // namespace UE::Slate
+
 /** Specialized control for handling the clipping of toolbars and menubars */
 class SClippingHorizontalBox : public SHorizontalBox
 {
@@ -25,6 +53,7 @@ public:
 		SLATE_ARGUMENT(const ISlateStyle*, StyleSet)
 		SLATE_ARGUMENT(FName, StyleName)
 		SLATE_ARGUMENT(bool, IsFocusable)
+		SLATE_ARGUMENT(FOnGetWidgetResizeParams, OnGetWidgetResizeParams)
 	SLATE_END_ARGS()
 
 	/** SWidget interface */
@@ -49,7 +78,6 @@ private:
 	void OnWrapButtonOpenChanged(bool bIsOpen);
 	EActiveTimerReturnType UpdateWrapButtonStatus(double CurrentTime, float DeltaTime);
 
-private:
 	/** The button that is displayed when a toolbar or menubar is clipped */
 	TSharedPtr<SComboButton> WrapButton;
 
@@ -62,10 +90,13 @@ private:
 	bool bIsFocusable = true;
 
 	float WrapButtonWidth = 0.0f; // Fixed width, initialized after button creation
+	int32 WrapButtonIndex = -1; // 0 is left-most index, -1 is right-most index (Python style).
 
 	/** The style to use */
 	const ISlateStyle* StyleSet = nullptr;
 	FName StyleName;
 
+	FOnGetWidgetResizeParams OnGetWidgetResizeParams;
+
 	mutable TArray<TWeakPtr<SWidget>> ClippedWidgets;
 };
diff --git a/Engine/Source/Runtime/Slate/Public/Framework/MultiBox/MultiBox.h b/Engine/Source/Runtime/Slate/Public/Framework/MultiBox/MultiBox.h
--- a/Engine/Source/Runtime/Slate/Public/Framework/MultiBox/MultiBox.h
+++ b/Engine/Source/Runtime/Slate/Public/Framework/MultiBox/MultiBox.h
@@ -251,6 +251,12 @@ public:
 	/** Set the visibility override for this block. */
 	void SetVisibilityOverride(TAttribute<EVisibility> InVisibilityOverride);
 
+	/** Get the resize parameters which determine how this block behaves during toolbar resizing.*/
+	FMenuEntryResizeParams GetResizeParams() const;
+
+	/** Set the resize parameters which determine how this block behaves during toolbar resizing.*/
+	void SetResizeParams(const TAttribute<FMenuEntryResizeParams>& InResizeParams);
+
 private:
 	/**
 	 * Allocates a widget for this type of MultiBlock.  Override this in derived classes.
@@ -310,6 +316,9 @@ private:
 
 	/** Visibility override which, if set and returning anything but Visible, takes priority over other methods of determining visibility. */
 	TAttribute<EVisibility> VisibilityOverride;
+
+	/** Resize parameters which determine how this block behaves during toolbar resizing. */
+	TAttribute<FMenuEntryResizeParams> ResizeParams;
 };
 
 
diff --git a/Engine/Source/Runtime/Slate/Public/Framework/MultiBox/MultiBoxBuilder.h b/Engine/Source/Runtime/Slate/Public/Framework/MultiBox/MultiBoxBuilder.h
--- a/Engine/Source/Runtime/Slate/Public/Framework/MultiBox/MultiBoxBuilder.h
+++ b/Engine/Source/Runtime/Slate/Public/Framework/MultiBox/MultiBoxBuilder.h
@@ -559,6 +559,7 @@ public:
 	 * @param	InCustomMenuDelegate	Optional custom menu delegate for cases where the toolbar is compressed into a menu
 	 * @param	InVisibilityOverride	Optional visibility override which, if set and returning anything but Visible, overrides the visibility of the toolbar button.
 	 * @param	InToolbarLabelOverride	Optional label override for when this block appears in a toolbar. If omitted, then the action's label will be used instead.
+	 * @param	InResizeParams			Optional resize parameters to control how this entry behaves during toolbar resizing.
 	 */
 	SLATE_API void AddToolBarButton(
 		const TSharedPtr<const FUICommandInfo> InCommand,
@@ -569,7 +570,8 @@ public:
 		FName InTutorialHighlightName = NAME_None,
 		FNewMenuDelegate InCustomMenuDelegate = FNewMenuDelegate(),
 		TAttribute<EVisibility> InVisibilityOverride = TAttribute<EVisibility>(),
-		TAttribute<FText> InToolbarLabelOverride = TAttribute<FText>()
+		TAttribute<FText> InToolbarLabelOverride = TAttribute<FText>(),
+		const TAttribute<FMenuEntryResizeParams>& InResizeParams = TAttribute<FMenuEntryResizeParams>()
 	);
 
 	/**
@@ -584,6 +586,7 @@ public:
 	 * @param	InTutorialHighlightName	Name to identify this widget and highlight during tutorials
 	 * @param	InVisibilityOverride	Optional visibility override which, if set and returning anything but Visible, overrides the visibility of the toolbar button.
 	 * @param	InToolbarLabelOverride	Optional label override for when this block appears in a toolbar. If omitted, then the action's label will be used instead.
+	 * @param	InResizeParams			Optional resize parameters to control how this entry behaves during toolbar resizing.
 	 */
 	SLATE_API void AddToolBarButton(
 		const FUIAction& InAction,
@@ -594,7 +597,9 @@ public:
 		const EUserInterfaceActionType UserInterfaceActionType = EUserInterfaceActionType::Button,
 		FName InTutorialHighlightName = NAME_None,
 		TAttribute<EVisibility> InVisibilityOverride = TAttribute<EVisibility>(),
-		TAttribute<FText> InToolbarLabelOverride = TAttribute<FText>() );
+		TAttribute<FText> InToolbarLabelOverride = TAttribute<FText>(),
+		const TAttribute<FMenuEntryResizeParams>& InResizeParams = TAttribute<FMenuEntryResizeParams>()
+	);
 
 	/**
 	 * Adds a combo button
@@ -610,6 +615,7 @@ public:
 	 * @param	InToolbarLabelOverride		Optional label override for when this block appears in a toolbar.  If omitted, then the action's label will be used instead.
 	 * @param	InPlacementOverride			Optional override placement to customize menu placement once opened via e.g. a toolbar menu button.
 	 * @param	InUserInterfaceActionType	Optional user interface action type to use if a command isn't bound. If a command is bound, we use the action type associated with the command instead.
+	 * @param	InResizeParams			Optional resize parameters to control how this entry behaves during toolbar resizing.
 	 */
 	SLATE_API void AddComboButton(
 		const FUIAction& InAction,
@@ -622,7 +628,8 @@ public:
 		TAttribute<EVisibility> InVisibilityOverride = TAttribute<EVisibility>(),
 		TAttribute<FText> InToolbarLabelOverride = TAttribute<FText>(),
 		TAttribute<EMenuPlacement> InPlacementOverride = TAttribute<EMenuPlacement>(),
-		const EUserInterfaceActionType InUserInterfaceActionType = EUserInterfaceActionType::Button
+		const EUserInterfaceActionType InUserInterfaceActionType = EUserInterfaceActionType::Button,
+		const TAttribute<FMenuEntryResizeParams>& InResizeParams = TAttribute<FMenuEntryResizeParams>()
 	);
 
 	/**
@@ -672,6 +679,7 @@ public:
 	 * @param	bInSearchable			If true, widget will be searchable (default == true)
 	 * @param	InCustomMenuDelegate	Optional custom menu delegate for cases where the toolbar is compressed into a menu
 	 * @param	InVisibilityOverride	Optional visibility override which, if set and returning anything but Visible, overrides the visibility of the widget.
+	 * @param	InResizeParams			Optional resize parameters to control how this entry behaves during toolbar resizing.
 	 */
 	SLATE_API void AddWidget(
 		TSharedRef<SWidget> InWidget,
@@ -679,7 +687,8 @@ public:
 		FName InTutorialHighlightName = NAME_None,
 		bool bInSearchable = true,
 		FNewMenuDelegate InCustomMenuDelegate = FNewMenuDelegate(),
-		TAttribute<EVisibility> InVisibilityOverride = TAttribute<EVisibility>()
+		TAttribute<EVisibility> InVisibilityOverride = TAttribute<EVisibility>(),
+		const TAttribute<FMenuEntryResizeParams>& InResizeParams = TAttribute<FMenuEntryResizeParams>()
 	);
 
 	/**
@@ -688,9 +697,12 @@ public:
 	 * @param	InExtensionHook			Optional extensibility hook.
 	 * @param	InVisibilityOverride	Optional visibility override which, if set and returning anything but Visible,
 	 * overrides the visibility of the separator.
+	 * @param	InResizeParams			Optional resize parameters to control how this entry behaves during toolbar resizing.
 	 */
 	SLATE_API void AddSeparator(
-		FName InExtensionHook = NAME_None, TAttribute<EVisibility> InVisibilityOverride = TAttribute<EVisibility>()
+		FName InExtensionHook = NAME_None,
+		TAttribute<EVisibility> InVisibilityOverride = TAttribute<EVisibility>(),
+		const TAttribute<FMenuEntryResizeParams>& InResizeParams = TAttribute<FMenuEntryResizeParams>()
 	);
 
 	/**
@@ -748,7 +760,8 @@ protected:
 		FName InTutorialHighlightName = NAME_None,
 		bool bInSearchable = true,
 		FNewMenuDelegate InCustomMenuDelegate = FNewMenuDelegate(),
-		TAttribute<EVisibility> InVisibilityOverride = TAttribute<EVisibility>()
+		TAttribute<EVisibility> InVisibilityOverride = TAttribute<EVisibility>(),
+		const TAttribute<FMenuEntryResizeParams>& InResizeParams = TAttribute<FMenuEntryResizeParams>()
 	);
 
 private:
diff --git a/Engine/Source/Runtime/Slate/Public/Framework/MultiBox/MultiBoxDefs.h b/Engine/Source/Runtime/Slate/Public/Framework/MultiBox/MultiBoxDefs.h
--- a/Engine/Source/Runtime/Slate/Public/Framework/MultiBox/MultiBoxDefs.h
+++ b/Engine/Source/Runtime/Slate/Public/Framework/MultiBox/MultiBoxDefs.h
@@ -168,3 +168,22 @@ struct FMenuEntryStyleParams
 	/** (Optionally) set the desired height override */
 	TOptional<float> DesiredHeight;
 };
+
+struct FMenuEntryResizeParams
+{
+	// Keep default values separate from the attributes. This allows the attributes to act like unset optionals by default.
+	static constexpr int32 DefaultClippingPriority = 0;
+	static constexpr bool DefaultAllowClipping = true;
+	static constexpr bool DefaultVisibleInOverflow = true;
+
+	/** The priority of this entry during resizing (default is 0). A higher priority relative to other entries keeps
+	 * the entry visible for longer as size is constrained during toolbar resizing. */
+	TAttribute<int32> ClippingPriority;
+
+	/** If true (default), allow this entry to be clipped during resizing of toolbars. If false, this entry will never
+	 * be clipped and always stay in the toolbar. */
+	TAttribute<bool> AllowClipping;
+
+	/** If true (default), this entry will be visible in a toolbar overflow menu. If false, this entry will disappear when it overflows. */
+	TAttribute<bool> VisibleInOverflow;
+};
diff --git a/Engine/Source/Runtime/SlateCore/Private/Styling/ToolBarStyle.cpp b/Engine/Source/Runtime/SlateCore/Private/Styling/ToolBarStyle.cpp
--- a/Engine/Source/Runtime/SlateCore/Private/Styling/ToolBarStyle.cpp
+++ b/Engine/Source/Runtime/SlateCore/Private/Styling/ToolBarStyle.cpp
@@ -29,6 +29,8 @@ FToolBarStyle::FToolBarStyle()
 	, BlockPadding(0)
 	, IndentedBlockPadding(0)
 	, BackgroundPadding(0)
+	, WrapButtonPadding(0)
+	, WrapButtonIndex(-1)
 	, IconSize(16, 16)
 	, bShowLabels(true)
 {}
diff --git a/Engine/Source/Runtime/SlateCore/Public/Styling/ToolBarStyle.h b/Engine/Source/Runtime/SlateCore/Public/Styling/ToolBarStyle.h
--- a/Engine/Source/Runtime/SlateCore/Public/Styling/ToolBarStyle.h
+++ b/Engine/Source/Runtime/SlateCore/Public/Styling/ToolBarStyle.h
@@ -125,6 +125,22 @@ struct FToolBarStyle : public FSlateWidgetStyle
 	FMargin BackgroundPadding;
 	FToolBarStyle& SetBackgroundPadding(const FMargin& InMargin) { BackgroundPadding = InMargin; return *this; }
 
+	UPROPERTY(EditAnywhere, BlueprintReadWrite, Category = Appearance)
+	FMargin WrapButtonPadding;
+	FToolBarStyle& SetWrapButtonPadding(const FMargin& InWrapButtonPadding)
+	{
+		WrapButtonPadding = InWrapButtonPadding;
+		return *this;
+	}
+
+	UPROPERTY(EditAnywhere, BlueprintReadWrite, Category = Appearance)
+	int32 WrapButtonIndex;
+	FToolBarStyle& SetWrapButtonIndex(int32 InWrapButtonIndex)
+	{
+		WrapButtonIndex = InWrapButtonIndex;
+		return *this;
+	}
+
 	UPROPERTY(EditAnywhere, BlueprintReadWrite, Category = Appearance)
 	FDeprecateSlateVector2D IconSize;
 	FToolBarStyle& SetIconSize(const UE::Slate::FDeprecateVector2DParameter& InIconSize) { IconSize = InIconSize; return *this; }
