Default SSpinBox::PreventThrottling to false

* This was the default behaviour when PreventThrottling was introduced in 4.27

diff --git a/Engine/Source/Runtime/Slate/Public/Widgets/Input/SSpinBox.h b/Engine/Source/Runtime/Slate/Public/Widgets/Input/SSpinBox.h
--- a/Engine/Source/Runtime/Slate/Public/Widgets/Input/SSpinBox.h
+++ b/Engine/Source/Runtime/Slate/Public/Widgets/Input/SSpinBox.h
@@ -97,7 +97,7 @@ public:
 		, _MinDesiredWidth(0.0f)
 		, _Justification(ETextJustify::Left)
 		, _KeyboardType(Keyboard_Default)
-		, _PreventThrottling(true)
+		, _PreventThrottling(false)
 	{}
 
 	/** The style used to draw this spinbox */
