Slate: Menu Anchors now properly only set focus for the provided FocusUserIndex when using EPopupMethod::CreateNewWindow

- PushMenu now takes a `FocusUserIndex` that is propagated to the menu's window and checked before setting focus
- Fixes cases where having a 2nd local player would unexpectedly force-close menus in the editor

diff --git a/Engine/Source/Runtime/Slate/Private/Framework/Application/MenuStack.cpp b/Engine/Source/Runtime/Slate/Private/Framework/Application/MenuStack.cpp
--- a/Engine/Source/Runtime/Slate/Private/Framework/Application/MenuStack.cpp
+++ b/Engine/Source/Runtime/Slate/Private/Framework/Application/MenuStack.cpp
@@ -204,11 +204,11 @@ namespace MenuStackInternal
 	FSimpleDelegate MenuStackPushDebuggingInfo;
 }	// anon namespace
 
-TSharedRef<IMenu> FMenuStack::Push(const FWidgetPath& InOwnerPath, const TSharedRef<SWidget>& InContent, const UE::Slate::FDeprecateVector2DParameter& SummonLocation, const FPopupTransitionEffect& TransitionEffect, const bool bFocusImmediately, const UE::Slate::FDeprecateVector2DParameter& SummonLocationSize, TOptional<EPopupMethod> InMethod, const bool bIsCollapsedByParent, const bool bEnablePerPixelTransparency)
+TSharedRef<IMenu> FMenuStack::Push(const FWidgetPath& InOwnerPath, const TSharedRef<SWidget>& InContent, const UE::Slate::FDeprecateVector2DParameter& SummonLocation, const FPopupTransitionEffect& TransitionEffect, const bool bFocusImmediately, const UE::Slate::FDeprecateVector2DParameter& SummonLocationSize, TOptional<EPopupMethod> InMethod, const bool bIsCollapsedByParent, const bool bEnablePerPixelTransparency, const int32 FocusUserIndex)
 {
 	// We want to ensure that when the window is restored, we restore the current keyboard focus, 
 	// but only if it is valid, otherwise we could end up clearing a previously valid path.
-	TSharedPtr<SWidget> FocusedWidget = FSlateApplication::Get().GetKeyboardFocusedWidget();
+	TSharedPtr<SWidget> FocusedWidget = (FocusUserIndex == INDEX_NONE) ? FSlateApplication::Get().GetKeyboardFocusedWidget() : FSlateApplication::Get().GetUserFocusedWidget(FocusUserIndex);
 	if (FocusedWidget.IsValid())
 	{
 		InOwnerPath.GetWindow()->SetWidgetToFocusOnActivate(FocusedWidget);
@@ -261,17 +261,17 @@ TSharedRef<IMenu> FMenuStack::Push(const FWidgetPath& InOwnerPath, const TShared
 #endif
 
 	TGuardValue<bool> Guard(bHostWindowGuard, true);
-	return PushInternal(ParentMenu, InContent, Anchor, TransitionEffect, bFocusImmediately, ActiveMethod.GetShouldThrottle(), bIsCollapsedByParent, bEnablePerPixelTransparency);
+	return PushInternal(ParentMenu, InContent, Anchor, TransitionEffect, bFocusImmediately, ActiveMethod.GetShouldThrottle(), bIsCollapsedByParent, bEnablePerPixelTransparency, FocusUserIndex);
 }
 
-TSharedRef<IMenu> FMenuStack::Push(const TSharedPtr<IMenu>& InParentMenu, const TSharedRef<SWidget>& InContent, const UE::Slate::FDeprecateVector2DParameter& SummonLocation, const FPopupTransitionEffect& TransitionEffect, const bool bFocusImmediately, const UE::Slate::FDeprecateVector2DParameter& SummonLocationSize, const bool bIsCollapsedByParent, const bool bEnablePerPixelTransparency)
+TSharedRef<IMenu> FMenuStack::Push(const TSharedPtr<IMenu>& InParentMenu, const TSharedRef<SWidget>& InContent, const UE::Slate::FDeprecateVector2DParameter& SummonLocation, const FPopupTransitionEffect& TransitionEffect, const bool bFocusImmediately, const UE::Slate::FDeprecateVector2DParameter& SummonLocationSize, const bool bIsCollapsedByParent, const bool bEnablePerPixelTransparency, const int32 FocusUserIndex)
 {
 	check(Stack.Contains(InParentMenu));
 	check(HostWindow.IsValid());
 
 	FSlateRect Anchor(SummonLocation, SummonLocation + SummonLocationSize);
 
-	return PushInternal(InParentMenu, InContent, Anchor, TransitionEffect, bFocusImmediately, EShouldThrottle::Yes, bIsCollapsedByParent, bEnablePerPixelTransparency);
+	return PushInternal(InParentMenu, InContent, Anchor, TransitionEffect, bFocusImmediately, EShouldThrottle::Yes, bIsCollapsedByParent, bEnablePerPixelTransparency, FocusUserIndex);
 }
 
 TSharedRef<IMenu> FMenuStack::PushHosted(const FWidgetPath& InOwnerPath, const TSharedRef<IMenuHost>& InMenuHost, const TSharedRef<SWidget>& InContent, TSharedPtr<SWidget>& OutWrappedContent, const FPopupTransitionEffect& TransitionEffect, EShouldThrottle ShouldThrottle, const bool bIsCollapsedByParent)
@@ -321,12 +321,13 @@ TSharedRef<IMenu> FMenuStack::PushHosted(const TSharedPtr<IMenu>& InParentMenu,
 	return OutMenu;
 }
 
-TSharedRef<IMenu> FMenuStack::PushInternal(const TSharedPtr<IMenu>& InParentMenu, const TSharedRef<SWidget>& InContent, FSlateRect Anchor, const FPopupTransitionEffect& TransitionEffect, const bool bFocusImmediately, EShouldThrottle ShouldThrottle, const bool bIsCollapsedByParent, const bool bEnablePerPixelTransparency)
+TSharedRef<IMenu> FMenuStack::PushInternal(const TSharedPtr<IMenu>& InParentMenu, const TSharedRef<SWidget>& InContent, FSlateRect Anchor, const FPopupTransitionEffect& TransitionEffect, const bool bFocusImmediately, EShouldThrottle ShouldThrottle, const bool bIsCollapsedByParent, const bool bEnablePerPixelTransparency, const int32 FocusUserIndex)
 {
 	FPrePushArgs PrePushArgs;
 	PrePushArgs.Content = InContent;
 	PrePushArgs.Anchor = Anchor;
 	PrePushArgs.TransitionEffect = TransitionEffect;
+	PrePushArgs.FocusUserIndex = FocusUserIndex;
 	PrePushArgs.bFocusImmediately = bFocusImmediately;
 	PrePushArgs.bIsCollapsedByParent = bIsCollapsedByParent;
 
@@ -338,8 +339,8 @@ TSharedRef<IMenu> FMenuStack::PushInternal(const TSharedPtr<IMenu>& InParentMenu
 
 	// Menu object creation stage
 	TSharedRef<FMenuBase> OutMenu = ActiveMethod.GetPopupMethod() == EPopupMethod::CreateNewWindow
-		? PushNewWindow(InParentMenu, PrePushResults, bEnablePerPixelTransparency)
-		: PushPopup(InParentMenu, PrePushResults);
+		? PushNewWindow(InParentMenu, PrePushResults, bEnablePerPixelTransparency, FocusUserIndex)
+		: PushPopup(InParentMenu, PrePushResults, FocusUserIndex);
 
 	// Post-push stage
 	//   Updates the stack and content map member variables
@@ -454,13 +455,20 @@ FMenuStack::FPrePushResults FMenuStack::PrePush(const FPrePushArgs& InArgs)
 	// Release the mouse so that context can be properly restored upon closing menus.  See CL 1411833 before changing this.
 	if (InArgs.bFocusImmediately)
 	{
-		FSlateApplication::Get().ReleaseAllPointerCapture();
+		if (InArgs.FocusUserIndex == INDEX_NONE)
+		{
+			FSlateApplication::Get().ReleaseAllPointerCapture();
+		}
+		else
+		{
+			FSlateApplication::Get().ReleaseAllPointerCapture(InArgs.FocusUserIndex);
+		}
 	}
 
 	return OutResults;
 }
 
-TSharedRef<FMenuBase> FMenuStack::PushNewWindow(TSharedPtr<IMenu> InParentMenu, const FPrePushResults& InPrePushResults, const bool bEnablePerPixelTransparency)
+TSharedRef<FMenuBase> FMenuStack::PushNewWindow(TSharedPtr<IMenu> InParentMenu, const FPrePushResults& InPrePushResults, const bool bEnablePerPixelTransparency, const int32 FocusUserIndex)
 {
 	check(ActiveMethod.GetPopupMethod() == EPopupMethod::CreateNewWindow);
 
@@ -486,6 +494,7 @@ TSharedRef<FMenuBase> FMenuStack::PushNewWindow(TSharedPtr<IMenu> InParentMenu,
 		.InitialOpacity(InitialWindowOpacity)
 		.SupportsTransparency(Transparency)
 		.FocusWhenFirstShown(InPrePushResults.bFocusImmediately)
+		.FocusUserIndex(FocusUserIndex)
 		.ActivationPolicy(InPrePushResults.bFocusImmediately ? EWindowActivationPolicy::Always : EWindowActivationPolicy::Never)
 		[
 			InPrePushResults.WrappedContent.ToSharedRef()
@@ -528,7 +537,7 @@ TSharedRef<FMenuBase> FMenuStack::PushNewWindow(TSharedPtr<IMenu> InParentMenu,
 	return Menu;
 }
 
-TSharedRef<FMenuBase> FMenuStack::PushPopup(TSharedPtr<IMenu> InParentMenu, const FPrePushResults& InPrePushResults)
+TSharedRef<FMenuBase> FMenuStack::PushPopup(TSharedPtr<IMenu> InParentMenu, const FPrePushResults& InPrePushResults, const int32 FocusUserIndex)
 {
 	check(ActiveMethod.GetPopupMethod() == EPopupMethod::UseCurrentWindow);
 
@@ -545,7 +554,14 @@ TSharedRef<FMenuBase> FMenuStack::PushPopup(TSharedPtr<IMenu> InParentMenu, cons
 
 	if (InPrePushResults.bFocusImmediately && InPrePushResults.WidgetToFocus.IsValid())
 	{
-		FSlateApplication::Get().SetKeyboardFocus(InPrePushResults.WidgetToFocus, EFocusCause::SetDirectly);
+		if (FocusUserIndex == INDEX_NONE)
+		{
+			FSlateApplication::Get().SetKeyboardFocus(InPrePushResults.WidgetToFocus, EFocusCause::SetDirectly);
+		}
+		else
+		{
+			FSlateApplication::Get().SetUserFocus(FocusUserIndex, InPrePushResults.WidgetToFocus, EFocusCause::SetDirectly);
+		}
 	}
 
 	return Menu;
diff --git a/Engine/Source/Runtime/Slate/Private/Framework/Application/SlateApplication.cpp b/Engine/Source/Runtime/Slate/Private/Framework/Application/SlateApplication.cpp
--- a/Engine/Source/Runtime/Slate/Private/Framework/Application/SlateApplication.cpp
+++ b/Engine/Source/Runtime/Slate/Private/Framework/Application/SlateApplication.cpp
@@ -2220,28 +2220,30 @@ TSharedRef<SWindow> FSlateApplication::AddWindowAsNativeChild( TSharedRef<SWindo
 	return InSlateWindow;
 }
 
-TSharedPtr<IMenu> FSlateApplication::PushMenu(const TSharedRef<SWidget>& InParentWidget, const FWidgetPath& InOwnerPath, const TSharedRef<SWidget>& InContent, const UE::Slate::FDeprecateVector2DParameter& SummonLocation, const FPopupTransitionEffect& TransitionEffect, const bool bFocusImmediately, const UE::Slate::FDeprecateVector2DParameter& SummonLocationSize, TOptional<EPopupMethod> Method, const bool bIsCollapsedByParent)
+TSharedPtr<IMenu> FSlateApplication::PushMenu(const TSharedRef<SWidget>& InParentWidget, const FWidgetPath& InOwnerPath, const TSharedRef<SWidget>& InContent, const UE::Slate::FDeprecateVector2DParameter& SummonLocation, const FPopupTransitionEffect& TransitionEffect, const bool bFocusImmediately, const UE::Slate::FDeprecateVector2DParameter& SummonLocationSize, TOptional<EPopupMethod> Method, const bool bIsCollapsedByParent, const int32 FocusUserIndex)
 {
+	constexpr bool bEnablePerPixelTransparency = false;
 	// Caller supplied a valid path? Pass it to the menu stack.
 	if (InOwnerPath.IsValid())
 	{
-		return MenuStack.Push(InOwnerPath, InContent, SummonLocation, TransitionEffect, bFocusImmediately, SummonLocationSize, Method, bIsCollapsedByParent);
+		return MenuStack.Push(InOwnerPath, InContent, SummonLocation, TransitionEffect, bFocusImmediately, SummonLocationSize, Method, bIsCollapsedByParent, bEnablePerPixelTransparency, FocusUserIndex);
 	}
 
 	// If the caller doesn't specify a valid event path we'll generate one from InParentWidget
 	FWidgetPath WidgetPath;
 	if (GeneratePathToWidgetUnchecked(InParentWidget, WidgetPath))
 	{
-		return MenuStack.Push(WidgetPath, InContent, SummonLocation, TransitionEffect, bFocusImmediately, SummonLocationSize, Method, bIsCollapsedByParent);
+		return MenuStack.Push(WidgetPath, InContent, SummonLocation, TransitionEffect, bFocusImmediately, SummonLocationSize, Method, bIsCollapsedByParent, bEnablePerPixelTransparency, FocusUserIndex);
 	}
 
 	UE_LOG(LogSlate, Warning, TEXT("Menu could not be pushed.  A path to the parent widget(%s) could not be found"), *InParentWidget->ToString());
 	return TSharedPtr<IMenu>();
 }
 
-TSharedPtr<IMenu> FSlateApplication::PushMenu(const TSharedPtr<IMenu>& InParentMenu, const TSharedRef<SWidget>& InContent, const UE::Slate::FDeprecateVector2DParameter& SummonLocation, const FPopupTransitionEffect& TransitionEffect, const bool bFocusImmediately, const UE::Slate::FDeprecateVector2DParameter& SummonLocationSize, const bool bIsCollapsedByParent)
+TSharedPtr<IMenu> FSlateApplication::PushMenu(const TSharedPtr<IMenu>& InParentMenu, const TSharedRef<SWidget>& InContent, const UE::Slate::FDeprecateVector2DParameter& SummonLocation, const FPopupTransitionEffect& TransitionEffect, const bool bFocusImmediately, const UE::Slate::FDeprecateVector2DParameter& SummonLocationSize, const bool bIsCollapsedByParent, const int32 FocusUserIndex)
 {
-	return MenuStack.Push(InParentMenu, InContent, SummonLocation, TransitionEffect, bFocusImmediately, SummonLocationSize, bIsCollapsedByParent);
+	constexpr bool bEnablePerPixelTransparency = false;
+	return MenuStack.Push(InParentMenu, InContent, SummonLocation, TransitionEffect, bFocusImmediately, SummonLocationSize, bIsCollapsedByParent, bEnablePerPixelTransparency, FocusUserIndex);
 }
 
 TSharedPtr<IMenu> FSlateApplication::PushHostedMenu(const TSharedRef<SWidget>& InParentWidget, const FWidgetPath& InOwnerPath, const TSharedRef<IMenuHost>& InMenuHost, const TSharedRef<SWidget>& InContent, TSharedPtr<SWidget>& OutWrappedContent, const FPopupTransitionEffect& TransitionEffect, EShouldThrottle ShouldThrottle, const bool bIsCollapsedByParent)
@@ -2849,6 +2851,23 @@ bool FSlateApplication::SetUserFocus(const uint32 InUserIndex, const FWidgetPath
 	return User && SetUserFocus(*User, InFocusPath, InCause);
 }
 
+bool FSlateApplication::SetUserFocusAllowingDescendantFocus(const uint32 InUserIndex, const FWidgetPath& InFocusPath, const EFocusCause InCause)
+{
+	if (TSharedPtr<FSlateUser> User = GetUser(InUserIndex))
+	{
+		const TSharedRef<SWidget>& FocusWidget = InFocusPath.Widgets.Last().Widget;
+		if (User->GetWeakFocusPath().ContainsWidget(&FocusWidget.Get()))
+		{
+			// A descendant is already focused
+			return true;
+		}
+
+		return SetUserFocus(*User, InFocusPath, InCause);
+	}
+
+	return false;
+}
+
 bool FSlateApplication::SetUserFocus(FSlateUser& User, const FWidgetPath& InFocusPath, const EFocusCause InCause)
 {
 	if (InFocusPath.IsValid())
diff --git a/Engine/Source/Runtime/Slate/Private/Widgets/Input/SMenuAnchor.cpp b/Engine/Source/Runtime/Slate/Private/Widgets/Input/SMenuAnchor.cpp
--- a/Engine/Source/Runtime/Slate/Private/Widgets/Input/SMenuAnchor.cpp
+++ b/Engine/Source/Runtime/Slate/Private/Widgets/Input/SMenuAnchor.cpp
@@ -384,7 +384,7 @@ void SMenuAnchor::SetIsOpen( bool InIsOpen, const bool bFocusMenu, const int32 F
 						if (MethodInUse.GetPopupMethod() == EPopupMethod::CreateNewWindow)
 						{
 							// Open the pop-up
-							TSharedPtr<IMenu> NewMenu = FSlateApplication::Get().PushMenu(AsShared(), MyWidgetPath, MenuContentRef, NewPosition, TransitionEffect, bFocusMenu, NewWindowSize, MethodInUse.GetPopupMethod(), bIsCollapsedByParent);
+							TSharedPtr<IMenu> NewMenu = FSlateApplication::Get().PushMenu(AsShared(), MyWidgetPath, MenuContentRef, NewPosition, TransitionEffect, bFocusMenu, NewWindowSize, MethodInUse.GetPopupMethod(), bIsCollapsedByParent, FocusUserIndex);
 							
 							if (ensure(NewMenu.IsValid()))
 							{
diff --git a/Engine/Source/Runtime/Slate/Public/Framework/Application/MenuStack.h b/Engine/Source/Runtime/Slate/Public/Framework/Application/MenuStack.h
--- a/Engine/Source/Runtime/Slate/Public/Framework/Application/MenuStack.h
+++ b/Engine/Source/Runtime/Slate/Public/Framework/Application/MenuStack.h
@@ -86,8 +86,9 @@ public:
 	 * @param InMethod				An optional popup method that will override the default method for the widgets in InOwnerPath
 	 * @param bIsCollapsedByParent	Is this menu collapsed when a parent menu receives focus/activation? If false, only focus/activation outside the entire stack will auto collapse it.
 	 * @param bEnablePerPixelTransparency Does the menu's content require per pixel transparency?
+	 * @param FocusUserIndex		The index of the user that should focus the menu. INDEX_NONE for all users.
 	 */
-	TSharedRef<IMenu> Push(const FWidgetPath& InOwnerPath, const TSharedRef<SWidget>& InContent, const UE::Slate::FDeprecateVector2DParameter& SummonLocation, const FPopupTransitionEffect& TransitionEffect, const bool bFocusImmediately = true, const UE::Slate::FDeprecateVector2DParameter& SummonLocationSize = FVector2f::ZeroVector, TOptional<EPopupMethod> InMethod = TOptional<EPopupMethod>(), const bool bIsCollapsedByParent = true, const bool bEnablePerPixelTransparency = false);
+	TSharedRef<IMenu> Push(const FWidgetPath& InOwnerPath, const TSharedRef<SWidget>& InContent, const UE::Slate::FDeprecateVector2DParameter& SummonLocation, const FPopupTransitionEffect& TransitionEffect, const bool bFocusImmediately = true, const UE::Slate::FDeprecateVector2DParameter& SummonLocationSize = FVector2f::ZeroVector, TOptional<EPopupMethod> InMethod = TOptional<EPopupMethod>(), const bool bIsCollapsedByParent = true, const bool bEnablePerPixelTransparency = false, const int32 FocusUserIndex = INDEX_NONE);
 	
 	/**
 	 * Pushes a new child menu onto the stack.
@@ -101,8 +102,9 @@ public:
 	 * @param SummonLocationSize	An optional size around the summon location which describes an area in which the menu may not appear
 	 * @param bIsCollapsedByParent	Is this menu collapsed when a parent menu receives focus/activation? If false, only focus/activation outside the entire stack will auto collapse it.
 	 * @param bEnablePerPixelTransparency Does the menu's content require per pixel transparency?
+	 * @param FocusUserIndex		The index of the user that should focus the menu. INDEX_NONE for all users.
 	 */
-	TSharedRef<IMenu> Push(const TSharedPtr<IMenu>& InParentMenu, const TSharedRef<SWidget>& InContent, const UE::Slate::FDeprecateVector2DParameter& SummonLocation, const FPopupTransitionEffect& TransitionEffect, const bool bFocusImmediately = true, const UE::Slate::FDeprecateVector2DParameter& SummonLocationSize = FVector2f::ZeroVector, const bool bIsCollapsedByParent = true, const bool bEnablePerPixelTransparency = false);
+	TSharedRef<IMenu> Push(const TSharedPtr<IMenu>& InParentMenu, const TSharedRef<SWidget>& InContent, const UE::Slate::FDeprecateVector2DParameter& SummonLocation, const FPopupTransitionEffect& TransitionEffect, const bool bFocusImmediately = true, const UE::Slate::FDeprecateVector2DParameter& SummonLocationSize = FVector2f::ZeroVector, const bool bIsCollapsedByParent = true, const bool bEnablePerPixelTransparency = false, const int32 FocusUserIndex = INDEX_NONE);
 
 	/**
 	 * Pushes a new menu onto the stack that is drawn by an external host widget.
@@ -264,6 +266,7 @@ private:
 		TSharedPtr<SWidget> Content;
 		FSlateRect Anchor;
 		FPopupTransitionEffect TransitionEffect;
+		int32 FocusUserIndex = INDEX_NONE;
 		bool bFocusImmediately;
 		bool bIsCollapsedByParent;
 	};
@@ -302,10 +305,11 @@ private:
 	 * @param TransitionEffect		Animation to use when the popup appears
 	 * @param bFocusImmediately		Should the popup steal focus when shown?
 	 * @param bIsCollapsedByParent	Is this menu collapsed when a parent menu receives focus/activation? If false, only focus/activation outside the entire stack will auto collapse it.
+	 * @param FocusUserIndex		The index of the user that should focus the menu. INDEX_NONE for all users.
 	 *
 	 * @return			The newly created menu.
 	 */
-	TSharedRef<IMenu> PushInternal(const TSharedPtr<IMenu>& InParentMenu, const TSharedRef<SWidget>& InContent, FSlateRect Anchor, const FPopupTransitionEffect& TransitionEffect, const bool bFocusImmediately, EShouldThrottle ShouldThrottle, const bool bIsCollapsedByParent, const bool bEnablePerPixelTransparency);
+	TSharedRef<IMenu> PushInternal(const TSharedPtr<IMenu>& InParentMenu, const TSharedRef<SWidget>& InContent, FSlateRect Anchor, const FPopupTransitionEffect& TransitionEffect, const bool bFocusImmediately, EShouldThrottle ShouldThrottle, const bool bIsCollapsedByParent, const bool bEnablePerPixelTransparency, const int32 FocusUserIndex = INDEX_NONE);
 
 	/**
 	 * This is the actual menu object creation method for FMenuInWindow menus.
@@ -313,10 +317,11 @@ private:
 	 *
 	 * @param InParentMenu			The parent menu for this menu
 	 * @param InPrePushResults		The results of the pre-push stage
+	 * @param FocusUserIndex		The index of the user that should focus the window. INDEX_NONE for all users.
 	 *
 	 * @return			The newly created FMenuInWindow menu.
 	 */
-	TSharedRef<FMenuBase> PushNewWindow(TSharedPtr<IMenu> InParentMenu, const FPrePushResults& InPrePushResults, const bool bEnablePerPixelTransparency);
+	TSharedRef<FMenuBase> PushNewWindow(TSharedPtr<IMenu> InParentMenu, const FPrePushResults& InPrePushResults, const bool bEnablePerPixelTransparency, const int32 FocusUserIndex = INDEX_NONE);
 
 	/**
 	 * This is the actual menu object creation method for FMenuInPopup menus.
@@ -324,10 +329,11 @@ private:
 	 *
 	 * @param InParentMenu			The parent menu for this menu
 	 * @param InPrePushResults		The results of the pre-push stage
+	 * @param FocusUserIndex		The index of the user that should focus the popup. INDEX_NONE for all users.
 	 *
 	 * @return			The newly created FMenuInPopup menu.
 	 */
-	TSharedRef<FMenuBase> PushPopup(TSharedPtr<IMenu> InParentMenu, const FPrePushResults& InPrePushResults);
+	TSharedRef<FMenuBase> PushPopup(TSharedPtr<IMenu> InParentMenu, const FPrePushResults& InPrePushResults, const int32 FocusUserIndex = INDEX_NONE);
 
 	/**
 	 * This is the post-push stage of menu creation.
diff --git a/Engine/Source/Runtime/Slate/Public/Framework/Application/SlateApplication.h b/Engine/Source/Runtime/Slate/Public/Framework/Application/SlateApplication.h
--- a/Engine/Source/Runtime/Slate/Public/Framework/Application/SlateApplication.h
+++ b/Engine/Source/Runtime/Slate/Public/Framework/Application/SlateApplication.h
@@ -444,8 +444,9 @@ public:
 	 * @param SummonLocationSize	An optional rect which describes an area in which the menu may not appear
 	 * @param Method				An optional popup method override. If not set, the widgets in the InOwnerPath will be queried for this.
 	 * @param bIsCollapsedByParent	Is this menu collapsed when a parent menu receives focus/activation? If false, only focus/activation outside the entire stack will auto collapse it.
+	 * @param FocusUserIndex		The index of the user that should focus the menu. INDEX_NONE for all users.
 	 */
-	SLATE_API TSharedPtr<IMenu> PushMenu(const TSharedRef<SWidget>& InParentWidget, const FWidgetPath& InOwnerPath, const TSharedRef<SWidget>& InContent, const UE::Slate::FDeprecateVector2DParameter& SummonLocation, const FPopupTransitionEffect& TransitionEffect, const bool bFocusImmediately = true, const UE::Slate::FDeprecateVector2DParameter& SummonLocationSize = FVector2f::ZeroVector, TOptional<EPopupMethod> Method = TOptional<EPopupMethod>(), const bool bIsCollapsedByParent = true);
+	SLATE_API TSharedPtr<IMenu> PushMenu(const TSharedRef<SWidget>& InParentWidget, const FWidgetPath& InOwnerPath, const TSharedRef<SWidget>& InContent, const UE::Slate::FDeprecateVector2DParameter& SummonLocation, const FPopupTransitionEffect& TransitionEffect, const bool bFocusImmediately = true, const UE::Slate::FDeprecateVector2DParameter& SummonLocationSize = FVector2f::ZeroVector, TOptional<EPopupMethod> Method = TOptional<EPopupMethod>(), const bool bIsCollapsedByParent = true, const int32 FocusUserIndex = INDEX_NONE);
 
 	/**
 	 * Creates a new Menu and adds it to the menu stack under the specified parent menu.
@@ -458,8 +459,9 @@ public:
 	 * @param bFocusImmediately		Should the popup steal focus when shown?
 	 * @param SummonLocationSize	An optional rect which describes an area in which the menu may not appear
 	 * @param bIsCollapsedByParent	Is this menu collapsed when a parent menu receives focus/activation? If false, only focus/activation outside the entire stack will auto collapse it.
+	 * @param FocusUserIndex		The index of the user that should focus the menu. INDEX_NONE for all users.
 	 */
-	SLATE_API TSharedPtr<IMenu> PushMenu(const TSharedPtr<IMenu>& InParentMenu, const TSharedRef<SWidget>& InContent, const UE::Slate::FDeprecateVector2DParameter& SummonLocation, const FPopupTransitionEffect& TransitionEffect, const bool bFocusImmediately = true, const UE::Slate::FDeprecateVector2DParameter& SummonLocationSize = FVector2f::ZeroVector, const bool bIsCollapsedByParent = true);
+	SLATE_API TSharedPtr<IMenu> PushMenu(const TSharedPtr<IMenu>& InParentMenu, const TSharedRef<SWidget>& InContent, const UE::Slate::FDeprecateVector2DParameter& SummonLocation, const FPopupTransitionEffect& TransitionEffect, const bool bFocusImmediately = true, const UE::Slate::FDeprecateVector2DParameter& SummonLocationSize = FVector2f::ZeroVector, const bool bIsCollapsedByParent = true, const int32 FocusUserIndex = INDEX_NONE);
 
 	/**
 	 * Creates a new hosted Menu and adds it to the menu stack.
@@ -1614,6 +1616,7 @@ public:
 	SLATE_API virtual void RequestDestroyWindow( TSharedRef<SWindow> WindowToDestroy ) override;
 	SLATE_API virtual bool SetKeyboardFocus( const FWidgetPath& InFocusPath, const EFocusCause InCause ) override;
 	SLATE_API virtual bool SetUserFocus(const uint32 InUserIndex, const FWidgetPath& InFocusPath, const EFocusCause InCause) override;
+	SLATE_API virtual bool SetUserFocusAllowingDescendantFocus(const uint32 InUserIndex, const FWidgetPath& InFocusPath, const EFocusCause InCause) override;
 	SLATE_API virtual void SetAllUserFocus(const FWidgetPath& InFocusPath, const EFocusCause InCause) override;
 	SLATE_API virtual void SetAllUserFocusAllowingDescendantFocus(const FWidgetPath& InFocusPath, const EFocusCause InCause) override;
 	SLATE_API virtual TSharedPtr<SWidget> GetUserFocusedWidget(uint32 UserIndex) const override;
diff --git a/Engine/Source/Runtime/SlateCore/Private/Widgets/SWindow.cpp b/Engine/Source/Runtime/SlateCore/Private/Widgets/SWindow.cpp
--- a/Engine/Source/Runtime/SlateCore/Private/Widgets/SWindow.cpp
+++ b/Engine/Source/Runtime/SlateCore/Private/Widgets/SWindow.cpp
@@ -250,6 +250,7 @@ void SWindow::Construct(const FArguments& InArgs)
 	this->bIsPopupWindow = InArgs._IsPopupWindow;
 	this->bIsTopmostWindow = InArgs._IsTopmostWindow;
 	this->bFocusWhenFirstShown = InArgs._FocusWhenFirstShown;
+	this->FocusUserIndex = InArgs._FocusUserIndex;
 	this->bHasOSWindowBorder = InArgs._UseOSWindowBorder;
 	this->bHasCloseButton = InArgs._HasCloseButton;
 	this->bHasMinimizeButton = InArgs._SupportsMinimize;
@@ -1674,7 +1675,7 @@ bool SWindow::OnIsActiveChanged( const FWindowActivateEvent& ActivateEvent )
 		else if (SupportsKeyboardFocus())
 		{
 			// If we have no specific widget to focus then cache the currently focused widget so we can restore its focus when we regain focus
-			WidgetFocusedOnDeactivate = FSlateApplicationBase::Get().GetKeyboardFocusedWidget();
+			WidgetFocusedOnDeactivate = (FocusUserIndex == INDEX_NONE) ? FSlateApplicationBase::Get().GetKeyboardFocusedWidget() : FSlateApplicationBase::Get().GetUserFocusedWidget(FocusUserIndex);
 			if (!WidgetFocusedOnDeactivate.IsValid())
 			{
 				WidgetFocusedOnDeactivate = FSlateApplicationBase::Get().GetUserFocusedWidget(0);
@@ -1696,7 +1697,14 @@ bool SWindow::OnIsActiveChanged( const FWindowActivateEvent& ActivateEvent )
 				FWidgetPath WidgetToFocusPath;
 				if (FSlateWindowHelper::FindPathToWidget(JustThisWindow, PinnedWidgetToFocus.ToSharedRef(), WidgetToFocusPath))
 				{
-					FSlateApplicationBase::Get().SetAllUserFocus(WidgetToFocusPath, EFocusCause::WindowActivate);
+					if (FocusUserIndex == INDEX_NONE)
+					{
+						FSlateApplicationBase::Get().SetAllUserFocus(WidgetToFocusPath, EFocusCause::WindowActivate);
+					}
+					else
+					{
+						FSlateApplicationBase::Get().SetUserFocus(FocusUserIndex, WidgetToFocusPath, EFocusCause::WindowActivate);
+					}
 				}
 			}
 
@@ -1708,7 +1716,14 @@ bool SWindow::OnIsActiveChanged( const FWindowActivateEvent& ActivateEvent )
 				TSharedRef<SWidget> WindowWidgetToFocus = WidgetFocusedOnDeactivate.IsValid() ? WidgetFocusedOnDeactivate.Pin().ToSharedRef() : AsShared();
 				if (FSlateWindowHelper::FindPathToWidget(JustThisWindow, WindowWidgetToFocus, WindowWidgetPath))
 				{
-					FSlateApplicationBase::Get().SetAllUserFocusAllowingDescendantFocus(WindowWidgetPath, EFocusCause::WindowActivate);
+					if (FocusUserIndex == INDEX_NONE)
+					{
+						FSlateApplicationBase::Get().SetAllUserFocusAllowingDescendantFocus(WindowWidgetPath, EFocusCause::WindowActivate);
+					}
+					else
+					{
+						FSlateApplicationBase::Get().SetUserFocusAllowingDescendantFocus(FocusUserIndex, WindowWidgetPath, EFocusCause::WindowActivate);
+					}
 				}
 			}
 		}
@@ -2042,6 +2057,7 @@ SWindow::SWindow()
 	, bShouldShowWindowContentDuringOverlay( false )
 	, ExpectedMaxWidth( INDEX_NONE )
 	, ExpectedMaxHeight( INDEX_NONE )
+	, FocusUserIndex( INDEX_NONE )
 	, TitleBar()
 	, bIsDrawingEnabled( true )
 	
diff --git a/Engine/Source/Runtime/SlateCore/Public/Application/SlateApplicationBase.h b/Engine/Source/Runtime/SlateCore/Public/Application/SlateApplicationBase.h
--- a/Engine/Source/Runtime/SlateCore/Public/Application/SlateApplicationBase.h
+++ b/Engine/Source/Runtime/SlateCore/Public/Application/SlateApplicationBase.h
@@ -428,7 +428,7 @@ public:
 	/**
 	 * Sets keyboard focus to the specified widget.  The widget must be allowed to receive keyboard focus.
 	 *
-	 * @param  InWidget WidgetPath to the Widget to being focused
+	 * @param  InFocusPath WidgetPath to the Widget to being focused
 	 * @param InCause The reason that keyboard focus is changing
 	 * @return true if the widget is now focused, false otherwise.
 	 */
@@ -438,16 +438,26 @@ public:
 	 * Sets user focus to the specified widget.  The widget must be allowed to receive focus.
 	 *
 	 * @param InUserIndex Index of the user that we want to change the focus of.
-	 * @param InWidget WidgetPath to the Widget to being focused.
+	 * @param InFocusPath WidgetPath to the Widget to being focused.
 	 * @param InCause The reason that focus is changing.
 	 * @return true if the widget is now focused, false otherwise.
 	 */
 	virtual bool SetUserFocus(const uint32 InUserIndex, const FWidgetPath& InFocusPath, const EFocusCause InCause) = 0;
 
+	/**
+	 * Sets the focus for the specified user to the specified widget unless that user is focused on a descendant.  The widget must be allowed to receive focus.
+	 *
+	 * @param InUserIndex Index of the user that we want to change the focus of.
+	 * @param InFocusPath WidgetPath to the Widget to being focused.
+	 * @param InCause The reason that focus is changing.
+	 * @return true if the widget or a descendant is now focused, false otherwise.
+	 */
+	virtual bool SetUserFocusAllowingDescendantFocus(const uint32 InUserIndex, const FWidgetPath& InFocusPath, const EFocusCause InCause) = 0;
+
 	/**
 	 * Sets the focus for all users to the specified widget.  The widget must be allowed to receive focus.
 	 *
-	 * @param InWidget WidgetPath to the Widget to being focused.
+	 * @param InFocusPath WidgetPath to the Widget to being focused.
 	 * @param InCause The reason that focus is changing.
 	 */
 	virtual void SetAllUserFocus(const FWidgetPath& InFocusPath, const EFocusCause InCause) = 0;
diff --git a/Engine/Source/Runtime/SlateCore/Public/Widgets/SWindow.h b/Engine/Source/Runtime/SlateCore/Public/Widgets/SWindow.h
--- a/Engine/Source/Runtime/SlateCore/Public/Widgets/SWindow.h
+++ b/Engine/Source/Runtime/SlateCore/Public/Widgets/SWindow.h
@@ -261,6 +261,7 @@ public:
 		, _IsPopupWindow( false )
 		, _IsTopmostWindow( false )
 		, _FocusWhenFirstShown(true)
+		, _FocusUserIndex( INDEX_NONE )
 		, _ActivationPolicy( EWindowActivationPolicy::Always )
 		, _UseOSWindowBorder( false )
 		, _HasCloseButton( true )
@@ -325,6 +326,9 @@ public:
 		/** Should this window be focused immediately after it is shown? */
 		SLATE_ARGUMENT( bool, FocusWhenFirstShown )
 
+		/** The index of the user that should focus this window on activation. INDEX_NONE for all users */
+		SLATE_ARGUMENT( int32, FocusUserIndex )
+
 		/** When should this window be activated upon being shown? */
 		SLATE_ARGUMENT( EWindowActivationPolicy, ActivationPolicy )
 
@@ -1345,6 +1349,9 @@ protected:
 	/** The expected maximum height of the window.  May be used for performance optimization when bSizeWillChangeOften is set. */
 	int32 ExpectedMaxHeight;
 
+	/** The index of the user that should focus this window on activation. INDEX_NONE for all users */
+	int32 FocusUserIndex;
+
 	// The window title bar.
 	TSharedPtr<IWindowTitleBar> TitleBar;
 
