Numeric Entry Box: Add argument so users can toggle throttling on it's internal SSpinBox. This is needed since it broke animation tools with the
https://p4-swarm.epicgames.net/reviews/37524342

diff --git a/Engine/Plugins/Animation/ControlRig/Source/ControlRigEditor/Private/EditMode/ControlsProxyDetailCustomization.cpp b/Engine/Plugins/Animation/ControlRig/Source/ControlRigEditor/Private/EditMode/ControlsProxyDetailCustomization.cpp
--- a/Engine/Plugins/Animation/ControlRig/Source/ControlRigEditor/Private/EditMode/ControlsProxyDetailCustomization.cpp
+++ b/Engine/Plugins/Animation/ControlRig/Source/ControlRigEditor/Private/EditMode/ControlsProxyDetailCustomization.cpp
@@ -724,6 +724,7 @@ TSharedRef<SWidget> FAnimDetailValueCustomization::MakeDoubleWidget(
 				.MinSliderValue(SliderMinValue)
 				.MaxSliderValue(SliderMaxValue)
 				.SliderExponent(SliderExponent)
+				.PreventThrottling(true)
 				.LinearDeltaSensitivity(1.0)
 				.Delta(Delta)
 				.LabelPadding(FMargin(3))
diff --git a/Engine/Source/Runtime/Slate/Public/Widgets/Input/SNumericEntryBox.h b/Engine/Source/Runtime/Slate/Public/Widgets/Input/SNumericEntryBox.h
--- a/Engine/Source/Runtime/Slate/Public/Widgets/Input/SNumericEntryBox.h
+++ b/Engine/Source/Runtime/Slate/Public/Widgets/Input/SNumericEntryBox.h
@@ -93,6 +93,7 @@ public:
 		, _MaxSliderValue(NumericType(100))
 		, _SliderExponent(1.f)
 		, _AllowWheel(true)
+		, _PreventThrottling(false)
 		, _BroadcastValueChangesPerKey(false)
 		, _MinDesiredValueWidth(0.f)
 		, _DisplayToggle(false)
@@ -163,6 +164,8 @@ public:
 		SLATE_ATTRIBUTE(NumericType, SliderExponentNeutralValue )
 		/** Whether this spin box should have mouse wheel feature enabled, defaults to true */
 		SLATE_ARGUMENT( bool, AllowWheel )
+		/** If refresh requests for the viewport should happen for all value changes **/
+		SLATE_ARGUMENT(bool, PreventThrottling)
 		/** True to broadcast every time we type */
 		SLATE_ARGUMENT( bool, BroadcastValueChangesPerKey)
 		/** Step to increment or decrement the value by when scrolling the mouse wheel. If not specified will determine automatically */
@@ -271,6 +274,7 @@ public:
 				.SliderExponent(InArgs._SliderExponent)
 				.SliderExponentNeutralValue(InArgs._SliderExponentNeutralValue)
 				.EnableWheel(InArgs._AllowWheel)
+				.PreventThrottling(InArgs._PreventThrottling)
 				.BroadcastValueChangesPerKey(InArgs._BroadcastValueChangesPerKey)
 				.WheelStep(InArgs._WheelStep)
 				.OnBeginSliderMovement(InArgs._OnBeginSliderMovement)
