Fix scrolling issue in STableViewBase

This CL fixes an issue with how we compute the scrollbar position when scrolling in a STableViewBase list.

VisibleSizeFraction is the fraction of the visible part of the whole list, in the [0.0, 1.0] range (normalized).

But in the case where we are observing a single row which is larger than the available visible area (ie.: ReGenerateResults.ExactNumLinesOnScreen < 1.0f), we compute it like this:

const double VisibleSizeFraction = AllottedGeometry.GetLocalSize().Y / ReGenerateResults.LengthOfGeneratedItems;

* AllottedGeometry.GetLocalSize().Y is the size (height) in pixels of the list widget.
* ReGenerateResults.LengthOfGeneratedItems is the size (height) in pixels of the PANEL geometry (single row), so the ratio we get is normalized in the PANEL, but not in the list, so we actually have to divide it by NumItemLines.

We were already doing this for the regular case (else).

#rnx

#tests Tested in PIE.

diff --git a/Engine/Source/Runtime/Slate/Private/Widgets/Views/STableViewBase.cpp b/Engine/Source/Runtime/Slate/Private/Widgets/Views/STableViewBase.cpp
--- a/Engine/Source/Runtime/Slate/Private/Widgets/Views/STableViewBase.cpp
+++ b/Engine/Source/Runtime/Slate/Private/Widgets/Views/STableViewBase.cpp
@@ -351,8 +351,8 @@ void STableViewBase::Tick( const FGeometry& AllottedGeometry, const double InCur
 			{
 				if (ReGenerateResults.ExactNumLinesOnScreen < 1.0f)
 				{
-					// We are be observing a single row which is larger than the available visible area, so we should calculate thumb size based on that
-					const double VisibleSizeFraction = AllottedGeometry.GetLocalSize().Y / ReGenerateResults.LengthOfGeneratedItems;
+					// We are observing a single row which is larger than the available visible area, so we should calculate thumb size based on that
+					const double VisibleSizeFraction = (Orientation == Orient_Vertical ? AllottedGeometry.GetLocalSize().Y : AllottedGeometry.GetLocalSize().X) / ReGenerateResults.LengthOfGeneratedItems / NumItemLines;
 					const double ThumbSizeFraction = FMath::Min(VisibleSizeFraction, 1.0);
 					const double OffsetFraction = CurrentScrollOffset / NumItemsBeingObserved;
 					ScrollBar->SetState( OffsetFraction, ThumbSizeFraction );
