Spinbox: Fix min desired width getter for spinbox
#12453

diff --git a/Engine/Source/Runtime/Slate/Private/Widgets/Input/SSpinBox.cpp b/Engine/Source/Runtime/Slate/Private/Widgets/Input/SSpinBox.cpp
--- a/Engine/Source/Runtime/Slate/Private/Widgets/Input/SSpinBox.cpp
+++ b/Engine/Source/Runtime/Slate/Private/Widgets/Input/SSpinBox.cpp
@@ -815,7 +815,7 @@ void SSpinBox<NumericType>::SetSliderExponent(const TAttribute<float>& InSliderE
 
 template <typename NumericType>
 float SSpinBox<NumericType>::GetMinDesiredWidth() const
-{ return SliderExponent.Get(); }
+{ return MinDesiredWidth.Get(); }
 
 template <typename NumericType>
 void SSpinBox<NumericType>::SetMinDesiredWidth(const TAttribute<float>& InMinDesiredWidth)
