Avoid adding a wrap button when the contents of an SClippingHorizontalBox is a subpixel larger than the box

diff --git a/Engine/Source/Runtime/Slate/Private/Framework/MultiBox/SClippingHorizontalBox.cpp b/Engine/Source/Runtime/Slate/Private/Framework/MultiBox/SClippingHorizontalBox.cpp
--- a/Engine/Source/Runtime/Slate/Private/Framework/MultiBox/SClippingHorizontalBox.cpp
+++ b/Engine/Source/Runtime/Slate/Private/Framework/MultiBox/SClippingHorizontalBox.cpp
@@ -34,10 +34,11 @@ void PrioritizedResize(
 		}
 	}
 
-	const bool bNeedsClipping = WidthOfAllChildren > InAllottedWidth;
-
 	// Early out if we don't need to clip.
-	if (!bNeedsClipping)
+	// Round to avoid adding a wrap button if the contents are a subpixel larger than the box.
+	if (const bool bNeedsClipping = FMath::CeilToInt(WidthOfAllChildren - UE_KINDA_SMALL_NUMBER)
+								  > FMath::CeilToInt(InAllottedWidth - UE_KINDA_SMALL_NUMBER);
+		!bNeedsClipping)
 	{
 		return;
 	}
@@ -73,7 +74,7 @@ void PrioritizedResize(
 		// non-clipping widgets to make sure we start clipping clippable widgets immediately once non-clippable widths
 		// and the accumulator sum up to our allotted width.
 		//
-		// If we find that we need a wrap button, restart the clipping loop and include the width of the warp button in
+		// If we find that we need a wrap button, restart the clipping loop and include the width of the wrap button in
 		// the initial accumulator width.
 		for (int32 NumClippingAttemptsLeft = 1; NumClippingAttemptsLeft > 0; --NumClippingAttemptsLeft)
 		{
