Fix nullptr access on Mac when a SWindow is destroyed before a deferred event  that might need to access the DPIScaleFactor is processed
Change a ensure to regular check on Mac because a deferred MouseUp event on Mac might be executed when the SWindow is already Dismissed
This problem may happen when selecting a new Default Editor map in Project Settings

#tests Local testing

diff --git a/Engine/Source/Runtime/ApplicationCore/Private/Mac/MacWindow.cpp b/Engine/Source/Runtime/ApplicationCore/Private/Mac/MacWindow.cpp
--- a/Engine/Source/Runtime/ApplicationCore/Private/Mac/MacWindow.cpp
+++ b/Engine/Source/Runtime/ApplicationCore/Private/Mac/MacWindow.cpp
@@ -540,7 +540,11 @@ float FMacWindow::GetDPIScaleFactor() const
 {
 	// @TODO: UE-227397
 	// Currently using cached property to avoid triggering a MainThreadChecker issue.
-	return FPlatformApplicationMisc::IsHighDPIModeEnabled() ? (float)WindowHandle->CachedBackingScaleFactor : 1.0f;
+	if (WindowHandle && FPlatformApplicationMisc::IsHighDPIModeEnabled())
+	{
+		return (float)WindowHandle->CachedBackingScaleFactor;
+	}
+	return 1.0f;
 }
 
 void FMacWindow::SetNativeWindowButtonsVisibility(bool bVisible)
diff --git a/Engine/Source/Runtime/Slate/Private/Framework/Application/SlateApplication.cpp b/Engine/Source/Runtime/Slate/Private/Framework/Application/SlateApplication.cpp
--- a/Engine/Source/Runtime/Slate/Private/Framework/Application/SlateApplication.cpp
+++ b/Engine/Source/Runtime/Slate/Private/Framework/Application/SlateApplication.cpp
@@ -5338,7 +5338,12 @@ FReply FSlateApplication::RoutePointerUpEvent(const FWidgetPath& WidgetsUnderPoi
 	if (SlateUser->HasCapture(PointerEvent.GetPointerIndex()))
 	{
 		FWidgetPath MouseCaptorPath = SlateUser->GetCaptorPath(PointerEvent.GetPointerIndex(), FWeakWidgetPath::EInterruptedPathHandling::Truncate, &PointerEvent);
+#if PLATFORM_MAC
+		// Because of Deferred event on Mac, by the time this event is processed, the window with the mouse capture might have been dismissed before the mouse up is processed
+		if ( MouseCaptorPath.Widgets.Num() > 0 ) 
+#else
 		if ( ensureMsgf(MouseCaptorPath.Widgets.Num() > 0, TEXT("A window had a widget with mouse capture. That entire window has been dismissed before the mouse up could be processed.")) )
+#endif
 		{
 			// Switch worlds widgets in the current path
 			FScopedSwitchWorldHack SwitchWorld( MouseCaptorPath );
@@ -7775,4 +7780,4 @@ void VerifyWidgetLayerId(const TSharedRef<SWindow>& WindowToDraw)
 }
 
 } //namespace
-#endif //WITH_SLATE_DEBUGGING
\ No newline at end of file
+#endif //WITH_SLATE_DEBUGGING
