ColorPicker: Clamp mouse position when interacting with color wheel and spectrum to the underlying widget geometry.
Update default mouse sensitivity when interacting with the color wheel and spectrum.
#rnx

diff --git a/Engine/Source/Runtime/Slate/Private/Widgets/Colors/SColorSpectrum.cpp b/Engine/Source/Runtime/Slate/Private/Widgets/Colors/SColorSpectrum.cpp
--- a/Engine/Source/Runtime/Slate/Private/Widgets/Colors/SColorSpectrum.cpp
+++ b/Engine/Source/Runtime/Slate/Private/Widgets/Colors/SColorSpectrum.cpp
@@ -151,13 +151,20 @@ void SColorSpectrum::ProcessMouseAction(const FGeometry& MyGeometry, const FPoin
 	FVector2f LocalMouseCoordinate;
 	if (bDragging)
 	{
-		FVector2f Delta = MouseEvent.GetCursorDelta();
+		constexpr float SpectrumSensitivity = 0.4f;
+		FVector2f Delta = MouseEvent.GetCursorDelta() * SpectrumSensitivity;
 		if (MouseEvent.IsControlDown())
 		{
 			Delta *= CtrlMultiplier.Get();
 		}
 
 		LocalMouseCoordinate = LastSpectrumPosition + Delta;
+
+		// Clamp mouse position to spectrum geometry
+		const FVector2f SpectrumSize = MyGeometry.GetLocalSize();
+		LocalMouseCoordinate.X = FMath::Clamp(LocalMouseCoordinate.X, 0.0f, SpectrumSize.X);
+		LocalMouseCoordinate.Y = FMath::Clamp(LocalMouseCoordinate.Y, 0.0f, SpectrumSize.Y);
+
 		LastSpectrumPosition = LocalMouseCoordinate;
 	}
 	else
diff --git a/Engine/Source/Runtime/Slate/Private/Widgets/Colors/SColorWheel.cpp b/Engine/Source/Runtime/Slate/Private/Widgets/Colors/SColorWheel.cpp
--- a/Engine/Source/Runtime/Slate/Private/Widgets/Colors/SColorWheel.cpp
+++ b/Engine/Source/Runtime/Slate/Private/Widgets/Colors/SColorWheel.cpp
@@ -151,13 +151,27 @@ bool SColorWheel::ProcessMouseAction(const FGeometry& MyGeometry, const FPointer
 	FVector2f LocalMouseCoordinate;
 	if (bDragging)
 	{
-		FVector2f Delta = MouseEvent.GetCursorDelta();
+		constexpr float WheelSensitivity = 0.35f;
+		FVector2f Delta = MouseEvent.GetCursorDelta() * WheelSensitivity;
 		if (MouseEvent.IsControlDown())
 		{
 			Delta *= CtrlMultiplier.Get();
 		}
 
 		LocalMouseCoordinate = LastWheelPosition + Delta;
+
+		// Clamp mouse position to the circle geometry
+		const float CircleRadius = MyGeometry.GetLocalSize().X / 2.0f;
+
+		const FVector2f CirclePosition = LocalMouseCoordinate - FVector2f(CircleRadius);
+		const float DistanceFromCenter = FMath::Sqrt((CirclePosition.X * CirclePosition.X) + (CirclePosition.Y * CirclePosition.Y));
+
+		if (DistanceFromCenter > CircleRadius)
+		{
+			const float Angle = FMath::Atan2(CirclePosition.Y, CirclePosition.X);
+			LocalMouseCoordinate = FVector2f(FMath::Cos(Angle), FMath::Sin(Angle)) * CircleRadius + FVector2f(CircleRadius);
+		}
+
 		LastWheelPosition = LocalMouseCoordinate;
 	}
 	else
