FTabManager : Call RequestSavePersistentLayout() from both OnTabOpening and OnTabClosing

diff --git a/Engine/Source/Runtime/Slate/Private/Framework/Docking/TabManager.cpp b/Engine/Source/Runtime/Slate/Private/Framework/Docking/TabManager.cpp
--- a/Engine/Source/Runtime/Slate/Private/Framework/Docking/TabManager.cpp
+++ b/Engine/Source/Runtime/Slate/Private/Framework/Docking/TabManager.cpp
@@ -2410,11 +2410,13 @@ void FTabManager::OnTabRelocated( const TSharedRef<SDockTab>& RelocatedTab, cons
 void FTabManager::OnTabOpening( const TSharedRef<SDockTab>& TabBeingOpened )
 {
 	UpdateStats();
+
+	RequestSavePersistentLayout();
 }
 
 void FTabManager::OnTabClosing( const TSharedRef<SDockTab>& TabBeingClosed )
 {
-
+	RequestSavePersistentLayout();
 }
 
 void FTabManager::OnTabManagerClosing()
@@ -2927,6 +2929,8 @@ void FGlobalTabmanager::OnTabRelocated( const TSharedRef<SDockTab>& RelocatedTab
 
 void FGlobalTabmanager::OnTabClosing( const TSharedRef<SDockTab>& TabBeingClosed )
 {
+	FTabManager::OnTabClosing(TabBeingClosed);
+
 	// Is this a major tab that contained a Sub TabManager?
 	// If so, need to properly close the sub tab manager
 	const int32 TabManagerBeingClosedIndex = SubTabManagers.IndexOfByPredicate(FindByTab(TabBeingClosed));
