[Slate] Labels are collapsed when empty in vertical toolbars

diff --git a/Engine/Source/Runtime/Slate/Private/Framework/MultiBox/SToolBarButtonBlock.cpp b/Engine/Source/Runtime/Slate/Private/Framework/MultiBox/SToolBarButtonBlock.cpp
--- a/Engine/Source/Runtime/Slate/Private/Framework/MultiBox/SToolBarButtonBlock.cpp
+++ b/Engine/Source/Runtime/Slate/Private/Framework/MultiBox/SToolBarButtonBlock.cpp
@@ -246,36 +246,37 @@ void SToolBarButtonBlock::BuildMultiBlockWidget(const ISlateStyle* StyleSet, con
 	IconWidget->AddLayer(TAttribute<const FSlateBrush*>(this, &SToolBarButtonBlock::GetOverlayIconBrush));
 	const bool bIsSlimHorizontalUniformToolBar = MultiBox->GetType() == EMultiBoxType::SlimHorizontalUniformToolBar;
 
+	const TSharedRef<STextBlock> TextBlock =
+	SNew(STextBlock)
+		// Collapse empty labels to prevent them from taking up visible space.
+		.Visibility_Lambda(
+			[WeakBlock = SharedThis(this).ToWeakPtr(), ActualLabel]() -> EVisibility
+			{
+				// Check first if the label is empty, and if so collapse it.
+				if (ActualLabel.IsSet() && ActualLabel.Get().IsEmpty())
+				{
+					return EVisibility::Collapsed;
+				}
+				// Only now check the set override.
+				else if (TSharedPtr<SToolBarButtonBlock> Block = WeakBlock.Pin())
+				{
+					return Block->LabelVisibility.Get();
+				}
+				else
+				{
+					return EVisibility::Visible;
+				}
+			}
+		)
+		.Text(ActualLabel)
+		.TextStyle(&ToolBarStyle.LabelStyle); // Smaller font for tool tip labels
+
 	// Create the content for our button
 	TSharedRef<SWidget> ButtonContent = SNullWidget::NullWidget;
 	if (MultiBox->GetType() == EMultiBoxType::SlimHorizontalToolBar 
 		|| bIsSlimHorizontalUniformToolBar)
 	{
 		const FVector2f IconSize = ToolBarStyle.IconSize;
-		const TSharedRef<STextBlock> TextBlock =
-			SNew(STextBlock)
-				// Collapse empty labels to prevent them from taking up visible space.
-				.Visibility_Lambda(
-					[WeakBlock = SharedThis(this).ToWeakPtr(), ActualLabel]() -> EVisibility
-					{
-						// Check first if the label is empty, and if so collapse it.
-						if (ActualLabel.IsSet() && ActualLabel.Get().IsEmpty())
-						{
-							return EVisibility::Collapsed;
-						}
-						// Only now check the set override.
-						else if (TSharedPtr<SToolBarButtonBlock> Block = WeakBlock.Pin())
-						{
-							return Block->LabelVisibility.Get();
-						}
-						else
-						{
-							return EVisibility::Visible;
-						}
-					}
-				)
-				.Text(ActualLabel)
-				.TextStyle(&ToolBarStyle.LabelStyle); // Smaller font for tool tip labels
 
 		if (bIsSlimHorizontalUniformToolBar)
 		{
@@ -310,26 +311,42 @@ void SToolBarButtonBlock::BuildMultiBlockWidget(const ISlateStyle* StyleSet, con
 	{
 		const FMargin IconPadding = !LabelVisibility.IsSet() ? ToolBarStyle.IconPadding :
 		(LabelVisibility.Get() == EVisibility::Collapsed ? ToolBarStyle.IconPadding : ToolBarStyle.IconPaddingWithVisibleLabel);
-		const TSharedRef ContentVBox = 
+
+		// Use a delegate rather than static value, to account for the possibility that LabelVisiblity changes
+		TAttribute<FMargin> IconPaddingAttribute = TAttribute<FMargin>::Create(
+			TAttribute<FMargin>::FGetter::CreateLambda(
+				[TextBlock, InitialPadding = IconPadding]() -> FMargin
+				{
+					FMargin IconPaddingValue = InitialPadding;
+
+					// Icon Padding may use a bottom value appropriate for label separation, rather than for the button bounds,
+					// so if the label is empty, we instead use the top padding which will be more appropriate for the button bounds.
+					if (TextBlock->GetText().IsEmpty())
+					{
+						IconPaddingValue.Bottom = InitialPadding.Top;
+					}
+
+					return IconPaddingValue;
+				}));
+
+		const TSharedRef<SVerticalBox> ContentVBox =
 			SNew(SVerticalBox)
-					// Icon image
-					+ SVerticalBox::Slot()
-					.AutoHeight()
-					.Padding(IconPadding)
-					.HAlign(HAlign_Center)	// Center the icon horizontally, so that large labels don't stretch out the artwork
-					[
-						IconWidget
-					]
-					// Label text
-					+ SVerticalBox::Slot().AutoHeight()
-				.Padding(ToolBarStyle.LabelPadding)
-					.HAlign(HAlign_Center)	// Center the label text horizontally
-					[
-						SNew(STextBlock)
-						.Visibility(LabelVisibility)
-						.Text(ActualLabel)
-						.TextStyle(&ToolBarStyle.LabelStyle)	// Smaller font for tool tip labels
-					];
+			// Icon image
+			+ SVerticalBox::Slot()
+			.AutoHeight()
+			.Padding(IconPaddingAttribute)
+			.HAlign(HAlign_Center)	// Center the icon horizontally, so that large labels don't stretch out the artwork
+			[
+				IconWidget
+			]
+			// Label text
+			+ SVerticalBox::Slot()
+			.AutoHeight()
+			.Padding(ToolBarStyle.LabelPadding)
+			.HAlign(HAlign_Center)	// Center the label text horizontally
+			[
+				TextBlock
+			];
 
 		ButtonContent =
 			SNew(SHorizontalBox)
