ColorBlock: Add new display mode type to draw a block with the opaque color on the left and the color with alpha on the right.
Change default orientation of a square color block to be vertical.

diff --git a/Engine/Source/Runtime/Slate/Private/Widgets/Colors/SColorBlock.cpp b/Engine/Source/Runtime/Slate/Private/Widgets/Colors/SColorBlock.cpp
--- a/Engine/Source/Runtime/Slate/Private/Widgets/Colors/SColorBlock.cpp
+++ b/Engine/Source/Runtime/Slate/Private/Widgets/Colors/SColorBlock.cpp
@@ -80,6 +80,10 @@ int32 SColorBlock::OnPaint(const FPaintArgs& Args, const FGeometry& AllottedGeom
 		MakeSection(GradientStops, FVector2D::ZeroVector, AllottedGeometry.GetLocalSize() * 0.5f, InColor, InWidgetStyle, false);
 		MakeSection(GradientStops, AllottedGeometry.GetLocalSize() * 0.5f, AllottedGeometry.GetLocalSize(), InColor, InWidgetStyle, true);
 		break;
+	case EColorBlockAlphaDisplayMode::SeparateReverse:
+		MakeSection(GradientStops, FVector2D::ZeroVector, AllottedGeometry.GetLocalSize() * 0.5f, InColor, InWidgetStyle, true);
+		MakeSection(GradientStops, AllottedGeometry.GetLocalSize() * 0.5f, AllottedGeometry.GetLocalSize(), InColor, InWidgetStyle, false);
+		break;
 	case EColorBlockAlphaDisplayMode::Ignore:
 		MakeSection(GradientStops, FVector2D::ZeroVector, AllottedGeometry.GetLocalSize(), InColor, InWidgetStyle, true);
 		break;
@@ -90,7 +94,7 @@ int32 SColorBlock::OnPaint(const FPaintArgs& Args, const FGeometry& AllottedGeom
 		LayerId+1,
 		AllottedGeometry.ToPaintGeometry(),
 		MoveTemp(GradientStops),
-		(AllottedGeometry.GetLocalSize().X > AllottedGeometry.GetLocalSize().Y) ? Orient_Vertical : Orient_Horizontal,
+		(AllottedGeometry.GetLocalSize().X >= AllottedGeometry.GetLocalSize().Y) ? Orient_Vertical : Orient_Horizontal,
 		DrawEffects,
 		(FVector4f)GradientCornerRadius.Get()
 	);
diff --git a/Engine/Source/Runtime/Slate/Public/Widgets/Colors/SColorBlock.h b/Engine/Source/Runtime/Slate/Public/Widgets/Colors/SColorBlock.h
--- a/Engine/Source/Runtime/Slate/Public/Widgets/Colors/SColorBlock.h
+++ b/Engine/Source/Runtime/Slate/Public/Widgets/Colors/SColorBlock.h
@@ -18,6 +18,8 @@ enum class EColorBlockAlphaDisplayMode : uint8
 	Combined,
 	// The color block is split into in half. The left half draws the color with opacity and the right half draws without any opacity
 	Separate,
+	// The color block is split into in half. The left half draws the color without any opacity and the right half draws with opacity
+	SeparateReverse,
 	// Alpha is omitted from display
 	Ignore,
 };
