[Slate] Added  DragHandleSize override to SSplitter, allowing separate visual and hit detection for drag operations, and not just hover detection

diff --git a/Engine/Source/Runtime/Slate/Private/Widgets/Layout/SSplitter.cpp b/Engine/Source/Runtime/Slate/Private/Widgets/Layout/SSplitter.cpp
--- a/Engine/Source/Runtime/Slate/Private/Widgets/Layout/SSplitter.cpp
+++ b/Engine/Source/Runtime/Slate/Private/Widgets/Layout/SSplitter.cpp
@@ -76,6 +76,10 @@ void SSplitter::Construct( const SSplitter::FArguments& InArgs )
 	ResizeMode = InArgs._ResizeMode;
 	PhysicalSplitterHandleSize = InArgs._PhysicalSplitterHandleSize;
 	HitDetectionSplitterHandleSize = InArgs._HitDetectionSplitterHandleSize;
+	DragDetectionSplitterHandleSize = FMath::Clamp(
+		InArgs._DragDetectionSplitterHandleSize.Get(PhysicalSplitterHandleSize), // fallback/default to PhysicalSplitterHandleSize
+		PhysicalSplitterHandleSize,
+		HitDetectionSplitterHandleSize); // must be between physical and hover sizes (inclusive)
 	Orientation = InArgs._Orientation;
 	HoveredHandleIndex = INDEX_NONE;
 	HighlightedHandleIndex.Assign(*this, InArgs._HighlightedHandleIndex);
@@ -455,9 +459,12 @@ FReply SSplitter::OnMouseMove( const FGeometry& MyGeometry, const FPointerEvent&
 		// Hit test which handle we are hovering over.
 		int32 PreviousHoveredHandleIndex = HoveredHandleIndex;
 
+		float DistanceFromHandleCenter = 0.0f;
 		HoveredHandleIndex = (Orientation == Orient_Horizontal)
-			? GetHandleBeingResizedFromMousePosition<Orient_Horizontal>( PhysicalSplitterHandleSize, HitDetectionSplitterHandleSize, LocalMousePosition, LayoutChildren )
-			: GetHandleBeingResizedFromMousePosition<Orient_Vertical>( PhysicalSplitterHandleSize, HitDetectionSplitterHandleSize, LocalMousePosition, LayoutChildren );
+			? GetHandleBeingResizedFromMousePosition<Orient_Horizontal>( PhysicalSplitterHandleSize, HitDetectionSplitterHandleSize, LocalMousePosition, LayoutChildren, DistanceFromHandleCenter )
+			: GetHandleBeingResizedFromMousePosition<Orient_Vertical>( PhysicalSplitterHandleSize, HitDetectionSplitterHandleSize, LocalMousePosition, LayoutChildren, DistanceFromHandleCenter );
+
+		FReply Reply = FReply::Unhandled();
 
 		if (HoveredHandleIndex != INDEX_NONE)
 		{
@@ -465,6 +472,18 @@ FReply SSplitter::OnMouseMove( const FGeometry& MyGeometry, const FPointerEvent&
 			{
 				HoveredHandleIndex = INDEX_NONE;
 			}
+			else if (!HasMouseCapture() && FMath::FloorToInt(DistanceFromHandleCenter) <= DragDetectionSplitterHandleSize / 2.0f)
+			{
+				Reply.CaptureMouse(AsShared());
+			}
+		}
+		else
+		{
+			// Release if no hits, and we have capture
+			if (Reply.GetMouseCaptor().IsValid() == false && HasMouseCapture())
+			{
+				Reply.ReleaseMouseCapture();
+			}
 		}
 
 		if (HoveredHandleIndex != PreviousHoveredHandleIndex)
@@ -472,7 +491,7 @@ FReply SSplitter::OnMouseMove( const FGeometry& MyGeometry, const FPointerEvent&
 			OnHandleHovered.ExecuteIfBound(HoveredHandleIndex);
 		}
 
-		return FReply::Unhandled();
+		return Reply;
 	}
 
 }
@@ -529,10 +548,11 @@ FCursorReply SSplitter::OnCursorQuery( const FGeometry& MyGeometry, const FPoint
 
 	TArray<FLayoutGeometry> LayoutChildren = ArrangeChildrenForLayout(MyGeometry);
 
-	// Hit test which handle we are hovering over.		
+	// Hit test which handle we are hovering over.
+	float MouseDistanceFromHandle = 0.0f;
 	const int32 CurrentHoveredHandleIndex = (Orientation == Orient_Horizontal)
-		? GetHandleBeingResizedFromMousePosition<Orient_Horizontal>( PhysicalSplitterHandleSize, HitDetectionSplitterHandleSize, LocalMousePosition, LayoutChildren )
-		: GetHandleBeingResizedFromMousePosition<Orient_Vertical>( PhysicalSplitterHandleSize, HitDetectionSplitterHandleSize, LocalMousePosition, LayoutChildren );
+		? GetHandleBeingResizedFromMousePosition<Orient_Horizontal>( PhysicalSplitterHandleSize, HitDetectionSplitterHandleSize, LocalMousePosition, LayoutChildren, MouseDistanceFromHandle )
+		: GetHandleBeingResizedFromMousePosition<Orient_Vertical>( PhysicalSplitterHandleSize, HitDetectionSplitterHandleSize, LocalMousePosition, LayoutChildren, MouseDistanceFromHandle );
 
 	if (CurrentHoveredHandleIndex != INDEX_NONE &&
 		Children[CurrentHoveredHandleIndex].CanBeResized())
@@ -805,7 +825,7 @@ float SSplitter::ClampChild(const FSlot& ChildSlot, float ProposedSize) const
 
 
 template<EOrientation SplitterOrientation>
-int32 SSplitter::GetHandleBeingResizedFromMousePosition( float InPhysicalSplitterHandleSize, float HitDetectionSplitterHandleSize, FVector2D LocalMousePos, const TArray<FLayoutGeometry>& ChildGeometries )
+int32 SSplitter::GetHandleBeingResizedFromMousePosition( float InPhysicalSplitterHandleSize, float HitDetectionSplitterHandleSize, FVector2D LocalMousePos, const TArray<FLayoutGeometry>& ChildGeometries, float& OutDistanceFromHandle )
 {
 	const int32 AxisIndex = (SplitterOrientation == Orient_Horizontal) ? 0 : 1;
 	const float HalfHitDetectionSplitterHandleSize = ( HitDetectionSplitterHandleSize / 2 );
@@ -818,13 +838,16 @@ int32 SSplitter::GetHandleBeingResizedFromMousePosition( float InPhysicalSplitte
 		FVector2D NextChildOffset = ChildGeometries[ChildIndex].GetOffsetInParentSpace();
 		float PrevBound = PrevChildRect.GetTopLeft().Component(AxisIndex) + PrevChildRect.GetSize().Component(AxisIndex) - HalfHitDetectionSplitterHandleSize + HalfPhysicalSplitterHandleSize;
 		float NextBound = NextChildOffset.Component(AxisIndex) + HalfHitDetectionSplitterHandleSize - HalfPhysicalSplitterHandleSize;
+		float MousePosOnAxis = LocalMousePos.Component(AxisIndex);
 
-		if ( LocalMousePos.Component(AxisIndex) > PrevBound && LocalMousePos.Component(AxisIndex) < NextBound )
+		if ( MousePosOnAxis > PrevBound && MousePosOnAxis < NextBound )
 		{
+			OutDistanceFromHandle = FMath::Abs(MousePosOnAxis - NextChildOffset.Component(AxisIndex));
 			return ChildIndex - 1;
 		}
 	}
 
+	OutDistanceFromHandle = 0.0f;
 	return INDEX_NONE;
 }
 
diff --git a/Engine/Source/Runtime/Slate/Public/Widgets/Layout/SSplitter.h b/Engine/Source/Runtime/Slate/Public/Widgets/Layout/SSplitter.h
--- a/Engine/Source/Runtime/Slate/Public/Widgets/Layout/SSplitter.h
+++ b/Engine/Source/Runtime/Slate/Public/Widgets/Layout/SSplitter.h
@@ -196,6 +196,8 @@ public:
 
 		SLATE_ARGUMENT( float, HitDetectionSplitterHandleSize )
 
+		SLATE_ARGUMENT( TOptional<float>, DragDetectionSplitterHandleSize )
+
 		SLATE_ARGUMENT( float, MinimumSlotHeight )
 
 		SLATE_ATTRIBUTE( int32, HighlightedHandleIndex )
@@ -354,11 +356,12 @@ protected:
 	 * @return The index of the handle being hovered, or INDEX_NONE if we are not hovering a handle.
 	 */
 	template<EOrientation SplitterOrientation>
-	static int32 GetHandleBeingResizedFromMousePosition(  float PhysicalSplitterHandleSize, float HitDetectionSplitterHandleSize, FVector2D LocalMousePos, const TArray<FLayoutGeometry>& ChildGeometries );
+	static int32 GetHandleBeingResizedFromMousePosition(  float PhysicalSplitterHandleSize, float HitDetectionSplitterHandleSize, FVector2D LocalMousePos, const TArray<FLayoutGeometry>& ChildGeometries, float& OutDistanceFromHandle );
 
 	TPanelChildren< FSlot > Children;
 
 	int32 HoveredHandleIndex;
+	int32 DragHandleIndex;
 	TSlateAttribute<int32, EInvalidateWidgetReason::Paint> HighlightedHandleIndex;
 	bool bIsResizing;
 	EOrientation Orientation;
@@ -374,6 +377,7 @@ protected:
 	/** The thickness of the grip area that the user uses to resize a splitter */
 	float PhysicalSplitterHandleSize;
 	float HitDetectionSplitterHandleSize;
+	float DragDetectionSplitterHandleSize;
 
 	const FSplitterStyle* Style;
 };
