[Slate] Fixed Vertical ToolBar separator orientation

diff --git a/Engine/Source/Editor/EditorStyle/Private/StarshipStyle.cpp b/Engine/Source/Editor/EditorStyle/Private/StarshipStyle.cpp
--- a/Engine/Source/Editor/EditorStyle/Private/StarshipStyle.cpp
+++ b/Engine/Source/Editor/EditorStyle/Private/StarshipStyle.cpp
@@ -1783,8 +1783,9 @@ void FStarshipEditorStyle::FStyle::SetupGeneralStyles()
 
 		Set("FVerticalToolBar.ToggleButton", CheckBoxStyle);
 
-		FVerticalToolBarStyle.SetButtonPadding(FMargin(0.0f, 4.0f, 0.0f, 0.f));
+		FVerticalToolBarStyle.SetButtonPadding(FMargin(0.0f, 4.0f));
 		FVerticalToolBarStyle.SetToggleButtonStyle(CheckBoxStyle);
+		FVerticalToolBarStyle.SetSeparatorPadding(FMargin(-5.0f, 4.0f));
 		Set("FVerticalToolBar", FVerticalToolBarStyle);
 	}
 
diff --git a/Engine/Source/Runtime/Slate/Private/Framework/MultiBox/SToolBarSeparatorBlock.cpp b/Engine/Source/Runtime/Slate/Private/Framework/MultiBox/SToolBarSeparatorBlock.cpp
--- a/Engine/Source/Runtime/Slate/Private/Framework/MultiBox/SToolBarSeparatorBlock.cpp
+++ b/Engine/Source/Runtime/Slate/Private/Framework/MultiBox/SToolBarSeparatorBlock.cpp
@@ -53,19 +53,39 @@ void SToolBarSeparatorBlock::BuildMultiBlockWidget(const ISlateStyle* StyleSet,
 {
 	const FToolBarStyle& ToolBarStyle = StyleSet->GetWidgetStyle<FToolBarStyle>(StyleName);
 
-	ChildSlot
-	[
-		SNew( SHorizontalBox )
-		+SHorizontalBox::Slot()
-		.AutoWidth()
-		.Padding(ToolBarStyle.SeparatorPadding)
+	const TSharedRef<const FMultiBox> MultiBox = OwnerMultiBoxWidget.Pin()->GetMultiBox();
+	if (MultiBox->GetType() == EMultiBoxType::VerticalToolBar)
+	{
+		ChildSlot
 		[
-			SNew(SSeparator)
+			SNew(SVerticalBox)
+			+ SVerticalBox::Slot()
+			.AutoHeight()
+			.Padding(ToolBarStyle.SeparatorPadding)
+			[
+				SNew(SSeparator)
+				.Orientation(Orient_Horizontal)
+				.Thickness(ToolBarStyle.SeparatorThickness)
+				.SeparatorImage(&ToolBarStyle.SeparatorBrush)
+			]
+		];
+	}
+	else
+	{
+		ChildSlot
+		[
+			SNew(SHorizontalBox)
+			+ SHorizontalBox::Slot()
+			.AutoWidth()
+			.Padding(ToolBarStyle.SeparatorPadding)
+			[
+				SNew(SSeparator)
 				.Orientation(Orient_Vertical)
 				.Thickness(ToolBarStyle.SeparatorThickness)
 				.SeparatorImage(&ToolBarStyle.SeparatorBrush)
-		]
-	];
+			]
+		];
+	}
 
 	// Add this widget to the search list of the multibox and hide it
 	OwnerMultiBoxWidget.Pin()->AddElement(this->AsWidget(), FText::GetEmpty(), MultiBlock->GetSearchable());
