Media Player Slider Scrubbing improvements

- [Slate] Implementing bPreventThrottle for SSlider
- Media Player Slider uses the Scrub function instead of Seek.

diff --git a/Engine/Plugins/Media/MediaPlayerEditor/Source/MediaPlayerEditor/Private/Widgets/SMediaPlayerSlider.cpp b/Engine/Plugins/Media/MediaPlayerEditor/Source/MediaPlayerEditor/Private/Widgets/SMediaPlayerSlider.cpp
--- a/Engine/Plugins/Media/MediaPlayerEditor/Source/MediaPlayerEditor/Private/Widgets/SMediaPlayerSlider.cpp
+++ b/Engine/Plugins/Media/MediaPlayerEditor/Source/MediaPlayerEditor/Private/Widgets/SMediaPlayerSlider.cpp
@@ -29,6 +29,7 @@ void SMediaPlayerSlider::Construct(const FArguments& InArgs, const TArrayView<TW
 		.Orientation(Orient_Horizontal)
 		.SliderBarColor(FLinearColor::Transparent)
 		.Style(InArgs._Style)
+		.PreventThrottling(true)
 	];
 }
 END_SLATE_FUNCTION_BUILD_OPTIMIZATION
@@ -97,7 +98,7 @@ void SMediaPlayerSlider::Seek(float InPlaybackPosition)
 
 			if (!ScrubberSlider->HasMouseCapture() || MediaPlayer->SupportsScrubbing())
 			{
-				MediaPlayer->Seek(MediaPlayer->GetDuration() * InPlaybackPosition);
+				MediaPlayer->Scrub(MediaPlayer->GetDuration() * InPlaybackPosition);
 				Entry.LastScrubValue = Entry.ScrubValue;
 			}
 		}
diff --git a/Engine/Source/Runtime/Slate/Private/Widgets/Input/SSlider.cpp b/Engine/Source/Runtime/Slate/Private/Widgets/Input/SSlider.cpp
--- a/Engine/Source/Runtime/Slate/Private/Widgets/Input/SSlider.cpp
+++ b/Engine/Source/Runtime/Slate/Private/Widgets/Input/SSlider.cpp
@@ -60,6 +60,7 @@ void SSlider::Construct( const SSlider::FArguments& InDeclaration )
 	OnControllerCaptureBegin = InDeclaration._OnControllerCaptureBegin;
 	OnControllerCaptureEnd = InDeclaration._OnControllerCaptureEnd;
 	OnValueChanged = InDeclaration._OnValueChanged;
+	bPreventThrottling = InDeclaration._PreventThrottling;
 
 	bControllerInputCaptured = false;
 }
@@ -299,8 +300,13 @@ FReply SSlider::OnMouseButtonDown( const FGeometry& MyGeometry, const FPointerEv
 		
 		// Release capture for controller/keyboard when switching to mouse.
 		ResetControllerState();
-		
-		return FReply::Handled().CaptureMouse(SharedThis(this));
+
+		FReply ReturnReply = FReply::Handled().CaptureMouse(SharedThis(this));
+		if (bPreventThrottling)
+		{
+			ReturnReply.PreventThrottling();
+		}
+		return ReturnReply;
 	}
 
 	return FReply::Unhandled();
diff --git a/Engine/Source/Runtime/Slate/Public/Widgets/Input/SSlider.h b/Engine/Source/Runtime/Slate/Public/Widgets/Input/SSlider.h
--- a/Engine/Source/Runtime/Slate/Public/Widgets/Input/SSlider.h
+++ b/Engine/Source/Runtime/Slate/Public/Widgets/Input/SSlider.h
@@ -42,6 +42,7 @@ public:
 		, _OnMouseCaptureBegin()
 		, _OnMouseCaptureEnd()
 		, _OnValueChanged()
+		, _PreventThrottling(false)
 		{
 		}
 
@@ -98,6 +99,9 @@ public:
 		/** Called when the value is changed by the slider. */
 		SLATE_EVENT( FOnFloatValueChanged, OnValueChanged )
 
+		/** If refresh requests for the viewport should happen for all value changes **/
+		SLATE_ARGUMENT(bool, PreventThrottling)
+
 	SLATE_END_ARGS()
 
 	SLATE_API SSlider();
@@ -268,6 +272,12 @@ protected:
 	/** When true, this slider will be keyboard focusable. Defaults to false. */
 	bool bIsFocusable;
 
+	/*
+	 * Holds whether or not to prevent throttling during mouse capture
+	 * When true, the viewport will be updated with every single change to the value during dragging
+	 */
+	bool bPreventThrottling = false;
+
 private:
 
 	// Resets controller input state. Fires delegates.
