[Slate] Toolbars can choose wether to allow wrap buttons or not (allowed by default)

#rnx

diff --git a/Engine/Source/Developer/ToolMenus/Private/ToolMenus.cpp b/Engine/Source/Developer/ToolMenus/Private/ToolMenus.cpp
--- a/Engine/Source/Developer/ToolMenus/Private/ToolMenus.cpp
+++ b/Engine/Source/Developer/ToolMenus/Private/ToolMenus.cpp
@@ -2635,6 +2635,7 @@ TSharedRef<SWidget> UToolMenus::GenerateWidget(UToolMenu* GeneratedMenu)
 		FToolBarBuilder ToolbarBuilder(GeneratedMenu->MenuType, GeneratedMenu->Context.CommandList, GeneratedMenu->MenuName, GeneratedMenu->Context.GetAllExtenders(), GeneratedMenu->bToolBarForceSmallIcons);
 		ToolbarBuilder.SetExtendersEnabled(GeneratedMenu->bExtendersEnabled);
 		ToolbarBuilder.SetIsFocusable(GeneratedMenu->bToolBarIsFocusable);
+		ToolbarBuilder.SetAllowWrapButton(GeneratedMenu->bAllowToolBarWrapButton);
 
 		if (GeneratedMenu->StyleName != NAME_None)
 		{
@@ -2856,6 +2857,7 @@ bool UToolMenus::RefreshMenuWidget(const FName InName, FGeneratedToolMenuWidget&
 		FToolBarBuilder ToolbarBuilder(GeneratedMenu->MenuType, GeneratedMenu->Context.CommandList, GeneratedMenu->MenuName, GeneratedMenu->Context.GetAllExtenders(), GeneratedMenu->bToolBarForceSmallIcons);
 		ToolbarBuilder.SetExtendersEnabled(GeneratedMenu->bExtendersEnabled);
 		ToolbarBuilder.SetIsFocusable(GeneratedMenu->bToolBarIsFocusable);
+		ToolbarBuilder.SetAllowWrapButton(GeneratedMenu->bAllowToolBarWrapButton);
 
 		if (GeneratedMenu->StyleName != NAME_None)
 		{
diff --git a/Engine/Source/Developer/ToolMenus/Public/ToolMenu.h b/Engine/Source/Developer/ToolMenus/Public/ToolMenu.h
--- a/Engine/Source/Developer/ToolMenus/Public/ToolMenu.h
+++ b/Engine/Source/Developer/ToolMenus/Public/ToolMenu.h
@@ -190,6 +190,10 @@ public:
 	UPROPERTY(EditAnywhere, BlueprintReadWrite, Category = "ToolBar")
 	bool bToolBarIsFocusable;
 
+	/** Whether to (ever) allow the wrap button/overflow menu. Defauls to the style. */
+	UPROPERTY(EditAnywhere, BlueprintReadWrite, Category = "ToolBar")
+	TOptional<bool> bAllowToolBarWrapButton;
+
 	UPROPERTY(EditAnywhere, BlueprintReadWrite, Category = "ToolBar")
 	bool bToolBarForceSmallIcons;
 
diff --git a/Engine/Source/Runtime/Slate/Private/Framework/MultiBox/MultiBox.cpp b/Engine/Source/Runtime/Slate/Private/Framework/MultiBox/MultiBox.cpp
--- a/Engine/Source/Runtime/Slate/Private/Framework/MultiBox/MultiBox.cpp
+++ b/Engine/Source/Runtime/Slate/Private/Framework/MultiBox/MultiBox.cpp
@@ -1094,6 +1094,7 @@ void SMultiBoxWidget::BuildMultiBoxWidget()
 		{
 			MainWidget = HorizontalBox = ClippedHorizontalBox =
 				SNew(SClippingHorizontalBox)
+					.AllowWrapButton(MultiBox->bAllowWrapButton)
 					.OnWrapButtonClicked(FOnGetContent::CreateSP(this, &SMultiBoxWidget::OnWrapButtonClicked))
 					.OnGetWidgetResizeParams(FOnGetWidgetResizeParams::CreateLambda(
 						[this](const TSharedRef<SWidget>& InWidget) -> FMenuEntryResizeParams
diff --git a/Engine/Source/Runtime/Slate/Private/Framework/MultiBox/MultiBoxBuilder.cpp b/Engine/Source/Runtime/Slate/Private/Framework/MultiBox/MultiBoxBuilder.cpp
--- a/Engine/Source/Runtime/Slate/Private/Framework/MultiBox/MultiBoxBuilder.cpp
+++ b/Engine/Source/Runtime/Slate/Private/Framework/MultiBox/MultiBoxBuilder.cpp
@@ -457,6 +457,11 @@ void FToolBarBuilder::SetIsFocusable(bool bInIsFocusable)
 	MultiBox->bIsFocusable = bIsFocusable;
 }
 
+void FToolBarBuilder::SetAllowWrapButton(const TOptional<bool>& bInAllowWrapButton) const
+{
+	MultiBox->bAllowWrapButton = bInAllowWrapButton;
+}
+
 void FToolBarBuilder::AddToolBarButton(
 	const TSharedPtr<const FUICommandInfo> InCommand,
 	FName InExtensionHook,
diff --git a/Engine/Source/Runtime/Slate/Private/Framework/MultiBox/SClippingHorizontalBox.cpp b/Engine/Source/Runtime/Slate/Private/Framework/MultiBox/SClippingHorizontalBox.cpp
--- a/Engine/Source/Runtime/Slate/Private/Framework/MultiBox/SClippingHorizontalBox.cpp
+++ b/Engine/Source/Runtime/Slate/Private/Framework/MultiBox/SClippingHorizontalBox.cpp
@@ -288,9 +288,20 @@ void SClippingHorizontalBox::OnArrangeChildren(const FGeometry& AllottedGeometry
 	const double InAllottedWidth = AllottedGeometry.GetLocalSize().X;
 	const FToolBarStyle& ToolBarStyle = StyleSet->GetWidgetStyle<FToolBarStyle>(StyleName);
 	double WrapButtonX = 0.0;
-	UE::Slate::PrioritizedResize(
-		InAllottedWidth, WrapButtonWidth, ToolBarStyle.WrapButtonPadding, WrapButtonIndex, ClippingInfos, WrapButtonX
-	);
+
+	if (bAllowWrapButton)
+	{
+		UE::Slate::PrioritizedResize(
+			InAllottedWidth, WrapButtonWidth, ToolBarStyle.WrapButtonPadding, WrapButtonIndex, ClippingInfos, WrapButtonX
+		);
+	}
+	else
+	{
+		// We don't care about WrapButton stuff here, so use placeholder values.
+		UE::Slate::PrioritizedResize(
+			InAllottedWidth, 0.0, FMargin(0.0f), 0, ClippingInfos, WrapButtonX
+		);
+	}
 
 	const int32 NumClippedWidgets = ClippingInfos
 										.FilterByPredicate(
@@ -374,6 +385,9 @@ void SClippingHorizontalBox::Construct( const FArguments& InArgs )
 	StyleName = InArgs._StyleName;
 	bIsFocusable = InArgs._IsFocusable;
 	OnGetWidgetResizeParams = InArgs._OnGetWidgetResizeParams;
+
+	const FToolBarStyle& ToolBarStyle = StyleSet->GetWidgetStyle<FToolBarStyle>(StyleName);
+	bAllowWrapButton = InArgs._AllowWrapButton.Get(ToolBarStyle.bAllowWrapButton);
 }
 
 void SClippingHorizontalBox::AddWrapButton()
@@ -400,7 +414,7 @@ void SClippingHorizontalBox::AddWrapButton()
 
 	// Perform a prepass to get a valid DesiredSize value below
 	WrapButton->SlatePrepass(1.0f);
-	WrapButtonWidth = WrapButton->GetDesiredSize().X;
+	WrapButtonWidth = bAllowWrapButton ? WrapButton->GetDesiredSize().X : 0.0f;
 	WrapButtonIndex = ToolBarStyle.WrapButtonIndex;
 
 	// Add the wrap button
diff --git a/Engine/Source/Runtime/Slate/Private/Framework/MultiBox/SClippingHorizontalBox.h b/Engine/Source/Runtime/Slate/Private/Framework/MultiBox/SClippingHorizontalBox.h
--- a/Engine/Source/Runtime/Slate/Private/Framework/MultiBox/SClippingHorizontalBox.h
+++ b/Engine/Source/Runtime/Slate/Private/Framework/MultiBox/SClippingHorizontalBox.h
@@ -46,9 +46,11 @@ class SClippingHorizontalBox : public SHorizontalBox
 {
 public:
 	SLATE_BEGIN_ARGS(SClippingHorizontalBox) 
-		: _StyleSet(&FCoreStyle::Get())
+		: _AllowWrapButton(true)
+		, _StyleSet(&FCoreStyle::Get())
 		, _StyleName(NAME_None)
 		, _IsFocusable(true) { }
+		SLATE_ARGUMENT(TOptional<bool>, AllowWrapButton)
 		SLATE_ARGUMENT(FOnGetContent, OnWrapButtonClicked)
 		SLATE_ARGUMENT(const ISlateStyle*, StyleSet)
 		SLATE_ARGUMENT(FName, StyleName)
@@ -81,6 +83,9 @@ private:
 	/** The button that is displayed when a toolbar or menubar is clipped */
 	TSharedPtr<SComboButton> WrapButton;
 
+	/** Whether or not to (ever) produce a wrap button */
+	bool bAllowWrapButton = true;
+
 	/** Callback for when the wrap button is clicked */
 	FOnGetContent OnWrapButtonClicked;
 
diff --git a/Engine/Source/Runtime/Slate/Public/Framework/MultiBox/MultiBox.h b/Engine/Source/Runtime/Slate/Public/Framework/MultiBox/MultiBox.h
--- a/Engine/Source/Runtime/Slate/Public/Framework/MultiBox/MultiBox.h
+++ b/Engine/Source/Runtime/Slate/Public/Framework/MultiBox/MultiBox.h
@@ -479,6 +479,9 @@ public:
 	/** Whether the MultiBox can be focused. */
 	bool bIsFocusable;
 
+	/** Whether the WrapButton/Overflow can be shown (if applicable). Defaults to the style. */
+	TOptional<bool> bAllowWrapButton;
+
 	/* Returns the last command list used */
 	const TSharedPtr<const FUICommandList> GetLastCommandList() const { return CommandLists.Num() > 0 ? CommandLists.Last() : nullptr; }
 
diff --git a/Engine/Source/Runtime/Slate/Public/Framework/MultiBox/MultiBoxBuilder.h b/Engine/Source/Runtime/Slate/Public/Framework/MultiBox/MultiBoxBuilder.h
--- a/Engine/Source/Runtime/Slate/Public/Framework/MultiBox/MultiBoxBuilder.h
+++ b/Engine/Source/Runtime/Slate/Public/Framework/MultiBox/MultiBoxBuilder.h
@@ -540,6 +540,9 @@ public:
 
 	SLATE_API void SetIsFocusable(bool bInIsFocusable);
 
+	/** Whether to allow the wrap button/overflow menu, if applicable. Defaults to the style. */
+	SLATE_API void SetAllowWrapButton(const TOptional<bool>& bInAllowWrapButton) const;
+
 	/**
 	 * Adds a tool bar button
 	 *
diff --git a/Engine/Source/Runtime/SlateCore/Public/Styling/ToolBarStyle.h b/Engine/Source/Runtime/SlateCore/Public/Styling/ToolBarStyle.h
--- a/Engine/Source/Runtime/SlateCore/Public/Styling/ToolBarStyle.h
+++ b/Engine/Source/Runtime/SlateCore/Public/Styling/ToolBarStyle.h
@@ -141,6 +141,11 @@ struct FToolBarStyle : public FSlateWidgetStyle
 		return *this;
 	}
 
+	/** Set to false if the wrap button should never be shown (even if entries are clipped) */
+	UPROPERTY(EditAnywhere, BlueprintReadWrite, Category=Appearance)
+	bool bAllowWrapButton = true;
+	FToolBarStyle& SetAllowWrapButton(bool bInAllowWrapButton) { bAllowWrapButton = bInAllowWrapButton; return *this; }
+
 	UPROPERTY(EditAnywhere, BlueprintReadWrite, Category = Appearance)
 	FDeprecateSlateVector2D IconSize;
 	FToolBarStyle& SetIconSize(const UE::Slate::FDeprecateVector2DParameter& InIconSize) { IconSize = InIconSize; return *this; }
