[Viewport Toolbar] WidgetBlocks Entries created with FToolMenuEntry::InitWidget cannot currently show an icon

Works with widget based entries created with FToolMenuEntry::InitWidget
Also adds padding so that the label for these items aligns with labels from other items.

diff --git a/Engine/Source/Developer/ToolMenus/Private/ToolMenus.cpp b/Engine/Source/Developer/ToolMenus/Private/ToolMenus.cpp
--- a/Engine/Source/Developer/ToolMenus/Private/ToolMenus.cpp
+++ b/Engine/Source/Developer/ToolMenus/Private/ToolMenus.cpp
@@ -367,7 +367,7 @@ public:
 			else
 			{
 				Block.WidgetData.StyleParams.bNoIndent = Block.WidgetData.bNoIndent;
-				MenuBuilder.AddWidget(Widget.ToSharedRef(), LabelToDisplay.Get(), Block.WidgetData.StyleParams, Block.WidgetData.bSearchable, Block.ToolTip.Get());
+				MenuBuilder.AddWidget(Widget.ToSharedRef(), LabelToDisplay.Get(), Block.WidgetData.StyleParams, Block.WidgetData.bSearchable, Block.ToolTip.Get(), Block.Icon);
 			}
 		}
 		else
diff --git a/Engine/Source/Runtime/Slate/Private/Framework/MultiBox/MultiBoxBuilder.cpp b/Engine/Source/Runtime/Slate/Private/Framework/MultiBox/MultiBoxBuilder.cpp
--- a/Engine/Source/Runtime/Slate/Private/Framework/MultiBox/MultiBoxBuilder.cpp
+++ b/Engine/Source/Runtime/Slate/Private/Framework/MultiBox/MultiBoxBuilder.cpp
@@ -358,11 +358,11 @@ void FMenuBuilder::AddWidget( TSharedRef<SWidget> InWidget, const FText& Label,
 	return AddWidget(InWidget, Label, StyleParams, bInSearchable, InToolTipText);
 }
 
-void FMenuBuilder::AddWidget(const TSharedRef<SWidget>& InWidget, const FText& InLabel, const FMenuEntryStyleParams& InStyleParams, const bool bInSearchable, const TAttribute<FText>& InToolTipText)
+void FMenuBuilder::AddWidget(const TSharedRef<SWidget>& InWidget, const FText& InLabel, const FMenuEntryStyleParams& InStyleParams, const bool bInSearchable, const TAttribute<FText>& InToolTipText, const TAttribute<FSlateIcon>& InIcon)
 {
 	ApplySectionBeginning();
 
-	TSharedRef<FWidgetBlock> NewWidgetBlock = MakeShared<FWidgetBlock>(InWidget, InLabel, InToolTipText, InStyleParams);
+	TSharedRef<FWidgetBlock> NewWidgetBlock = MakeShared<FWidgetBlock>(InWidget, InLabel, InToolTipText, InStyleParams, InIcon);
 	NewWidgetBlock->SetSearchable( bInSearchable );
 
 	MultiBox->AddMultiBlock( NewWidgetBlock );
diff --git a/Engine/Source/Runtime/Slate/Private/Framework/MultiBox/SWidgetBlock.cpp b/Engine/Source/Runtime/Slate/Private/Framework/MultiBox/SWidgetBlock.cpp
--- a/Engine/Source/Runtime/Slate/Private/Framework/MultiBox/SWidgetBlock.cpp
+++ b/Engine/Source/Runtime/Slate/Private/Framework/MultiBox/SWidgetBlock.cpp
@@ -5,12 +5,14 @@
 #include "Framework/MultiBox/MultiBoxBuilder.h"
 #include "Styling/ToolBarStyle.h"
 #include "Widgets/SToolTip.h"
+#include "Widgets/Images/SImage.h"
 
-FWidgetBlock::FWidgetBlock(TSharedRef<SWidget> InContent, const FText& InLabel, const TAttribute<FText>& InToolTipText, const FMenuEntryStyleParams& InStyleParams)
+FWidgetBlock::FWidgetBlock(TSharedRef<SWidget> InContent, const FText& InLabel, const TAttribute<FText>& InToolTipText, const FMenuEntryStyleParams& InStyleParams, const TAttribute<FSlateIcon>& InIcon)
 	: FMultiBlock(nullptr, nullptr, NAME_None, EMultiBlockType::Widget)
     , ContentWidget(InContent)
 	, Label(InLabel)
 	, ToolTipText(InToolTipText)
+	, Icon(InIcon)
 	, StyleParams(InStyleParams)
 {
 }
@@ -167,6 +169,23 @@ void SWidgetBlock::BuildMultiBlockWidget(const ISlateStyle* StyleSet, const FNam
 
 	static constexpr float IconRightPadding = 4.f;
 
+	// If we were supplied an image than go ahead and use that, otherwise we use a null widget
+	TSharedRef<SWidget> IconWidget = SNullWidget::NullWidget;
+	if (WidgetBlock->Icon.IsSet())
+	{
+		const FSlateIcon& ActualIcon = WidgetBlock->Icon.Get();
+		const FSlateBrush* IconBrush = ActualIcon.GetIcon();
+		if (IconBrush->GetResourceName() != NAME_None)
+		{
+			IconWidget = SNew(SImage)
+				.ColorAndOpacity(FSlateColor::UseSubduedForeground())
+				.Image(IconBrush);
+		}
+	}
+
+	const float MenuIconSize = StyleSet->GetFloat(StyleName, ".MenuIconSize", 16.f);
+
+	// clang-format off
 	ChildSlot
 	.Padding( Padding )	// Large left margin mimics the indent of normal menu items when bNoIndent is false
 	[
@@ -174,15 +193,35 @@ void SWidgetBlock::BuildMultiBlockWidget(const ISlateStyle* StyleSet, const FNam
 		.ToolTip(FMultiBoxSettings::ToolTipConstructor.Execute(WidgetBlock->ToolTipText, nullptr, nullptr, /*ShowActionShortcut=*/ false ))
 		+ SHorizontalBox::Slot()
 		.AutoWidth()
+		.Padding(FMargin(2, 0, 6, 0))
+		[
+			SNew(SBox)
+			.Visibility(IconWidget != SNullWidget::NullWidget ? EVisibility::Visible : EVisibility::Collapsed)
+			.WidthOverride(MenuIconSize + 2)
+			.HeightOverride(MenuIconSize)
+			.HAlign(HAlign_Center)
+			.VAlign(VAlign_Center)
+			[
+				SNew(SBox)
+				.WidthOverride(MenuIconSize)
+				.HeightOverride(MenuIconSize)
+				[
+					IconWidget
+				]
+			]
+		]
+		+ SHorizontalBox::Slot()
+		.Padding(FMargin(2, 0, 6, 0))
+		.AutoWidth()
 		[
 			SNew(SHorizontalBox)
-			.Visibility( bHasLabel ? EVisibility::Visible : EVisibility::Collapsed )
+			.Visibility(bHasLabel ? EVisibility::Visible : EVisibility::Collapsed)
 			+ SHorizontalBox::Slot()
 			.AutoWidth()
 			.Padding(bHasLabel ? FMargin(0.f, 0.f, IconRightPadding, 0.f) : FMargin(0))
 			.VAlign(VAlign_Center)
 			[
-				SNew( STextBlock )
+				SNew(STextBlock)
 				.TextStyle(LabelStyle)
 				.Text(WidgetBlock->Label)
 				.HighlightText(SearchHighlightText)
@@ -196,6 +235,7 @@ void SWidgetBlock::BuildMultiBlockWidget(const ISlateStyle* StyleSet, const FNam
 			WidgetBlock->ContentWidget
 		]
 	];
+	// clang-format on
 
 	SetVisibility(MultiBlock->GetVisibilityOverride());
 }
diff --git a/Engine/Source/Runtime/Slate/Private/Framework/MultiBox/SWidgetBlock.h b/Engine/Source/Runtime/Slate/Private/Framework/MultiBox/SWidgetBlock.h
--- a/Engine/Source/Runtime/Slate/Private/Framework/MultiBox/SWidgetBlock.h
+++ b/Engine/Source/Runtime/Slate/Private/Framework/MultiBox/SWidgetBlock.h
@@ -23,8 +23,9 @@ public:
 	 * @param	InLabel						Optional label text to be added to the left of the content
 	 * @param	InToolTipText				Optional tooltip text to be added to the widget and label
 	 * @param	InStyleParams				Optional additional style parameters. @see FMenuEntryStyleParams
+	 * @param	InIcon						Optional icon to be shown to the left of the label/content
 	 */
-	FWidgetBlock(TSharedRef<SWidget> InContent, const FText& InLabel, const TAttribute<FText>& InToolTipText = FText(), const FMenuEntryStyleParams& InStyleParams = {});
+	FWidgetBlock(TSharedRef<SWidget> InContent, const FText& InLabel, const TAttribute<FText>& InToolTipText = FText(), const FMenuEntryStyleParams& InStyleParams = {}, const TAttribute<FSlateIcon>& InIcon = TAttribute<FSlateIcon>());
 
 	/** FMultiBlock interface */
 	virtual void CreateMenuEntry(class FMenuBuilder& MenuBuilder) const override;
@@ -52,6 +53,9 @@ private:
 	/** Optional ToolTip text */
 	TAttribute<FText> ToolTipText;
 
+	/** Optional Icon */
+	TAttribute<FSlateIcon> Icon;
+
 	/** Style parameters */
 	FMenuEntryStyleParams StyleParams;
 
diff --git a/Engine/Source/Runtime/Slate/Public/Framework/MultiBox/MultiBoxBuilder.h b/Engine/Source/Runtime/Slate/Public/Framework/MultiBox/MultiBoxBuilder.h
--- a/Engine/Source/Runtime/Slate/Public/Framework/MultiBox/MultiBoxBuilder.h
+++ b/Engine/Source/Runtime/Slate/Public/Framework/MultiBox/MultiBoxBuilder.h
@@ -398,8 +398,9 @@ public:
 	 * @param	InStyleParams		Optional style parameters and overrides, @see FMenuEntryStyleParams
 	 * @param	bInSearchable		If true, widget will be searchable (default == true)
 	 * @param	InToolTipText		Optional tooltip text to be added to the widget and label
+	 * @param	InIcon				Optional icon to be shown to the left of the label/content
 	 */
-	SLATE_API void AddWidget(const TSharedRef<SWidget>& InWidget, const FText& InLabel, const FMenuEntryStyleParams& InStyleParams, bool bInSearchable = true, const TAttribute<FText>& InToolTipText = FText());
+	SLATE_API void AddWidget(const TSharedRef<SWidget>& InWidget, const FText& InLabel, const FMenuEntryStyleParams& InStyleParams, bool bInSearchable = true, const TAttribute<FText>& InToolTipText = FText(), const TAttribute<FSlateIcon>& InIcon = TAttribute<FSlateIcon>());
 
 	/**
 	* Adds the widget the multibox will use for searching
