-Add length validation on some editable text widget: SEditableTextBox, SMultilineEditableTextBox, SInlineEditableTextBox.
-Updated various editors to take advantage of the feature, fixing crash or removing duplicated code.
-Small include clean up as I got tired to recompile too many files when playing in some header files.

#tests Copy pasted a long name, got an error message and name was rejected in curve editor. Tested various other tools, like SPropertyEditorText particularly to validate that it was still working as expected.
#rnx
[RN] Added a MaximumLength attribute in SEditableTextBox, SMultilineEditableTextBox, SInlineEditableTextBox. It prevents a too long text to be accepted.

diff --git a/Engine/Plugins/Animation/ControlRig/Source/ControlRigEditor/Private/Editor/SModularRigTreeView.cpp b/Engine/Plugins/Animation/ControlRig/Source/ControlRigEditor/Private/Editor/SModularRigTreeView.cpp
--- a/Engine/Plugins/Animation/ControlRig/Source/ControlRigEditor/Private/Editor/SModularRigTreeView.cpp
+++ b/Engine/Plugins/Animation/ControlRig/Source/ControlRigEditor/Private/Editor/SModularRigTreeView.cpp
@@ -341,6 +341,7 @@ TSharedRef<SWidget> SModularRigModelItem::GenerateWidgetForColumn(const FName& C
 		[
 			SAssignNew(InlineWidget, SInlineEditableTextBlock)
 			.Text(this, &SModularRigModelItem::GetName, true)
+			.MaximumLength(NAME_SIZE-1)
 			.OnVerifyTextChanged(this, &SModularRigModelItem::OnVerifyNameChanged)
 			.OnTextCommitted(this, &SModularRigModelItem::OnNameCommitted)
 			.ToolTipText(this, &SModularRigModelItem::GetItemTooltip)
diff --git a/Engine/Plugins/Compositing/Composure/Source/ComposureLayersEditor/Private/Widgets/SCompElementViewRow.cpp b/Engine/Plugins/Compositing/Composure/Source/ComposureLayersEditor/Private/Widgets/SCompElementViewRow.cpp
--- a/Engine/Plugins/Compositing/Composure/Source/ComposureLayersEditor/Private/Widgets/SCompElementViewRow.cpp
+++ b/Engine/Plugins/Compositing/Composure/Source/ComposureLayersEditor/Private/Widgets/SCompElementViewRow.cpp
@@ -104,6 +104,7 @@ TSharedRef<SWidget> SCompElementViewRow::GenerateWidgetForColumn(const FName& Co
 					.Text(ViewModel.Get(), &FCompElementViewModel::GetNameAsText)
 					.ColorAndOpacity(this, &SCompElementViewRow::GetColorAndOpacity)
 					.HighlightText(HighlightText)
+					.MaximumLength(NAME_SIZE-1)
 					.OnVerifyTextChanged(this, &SCompElementViewRow::OnRenameElementTextChanged)
 					.OnTextCommitted(this, &SCompElementViewRow::OnRenameElementTextCommitted)
 					.IsSelected(this, &SCompElementViewRow::IsSelectedExclusively)
diff --git a/Engine/Plugins/Enterprise/GLTFExporter/Source/GLTFExporter/Private/UI/GLTFMaterialBakeSizeCustomization.cpp b/Engine/Plugins/Enterprise/GLTFExporter/Source/GLTFExporter/Private/UI/GLTFMaterialBakeSizeCustomization.cpp
--- a/Engine/Plugins/Enterprise/GLTFExporter/Source/GLTFExporter/Private/UI/GLTFMaterialBakeSizeCustomization.cpp
+++ b/Engine/Plugins/Enterprise/GLTFExporter/Source/GLTFExporter/Private/UI/GLTFMaterialBakeSizeCustomization.cpp
@@ -5,6 +5,7 @@
 #include "UI/GLTFMaterialBakeSizeCustomization.h"
 #include "UserData/GLTFMaterialUserData.h"
 #include "Widgets/Input/SCheckBox.h"
+#include "Widgets/Input/SEditableText.h"
 #include "SSearchableComboBox.h"
 #include "DetailWidgetRow.h"
 #include "IPropertyUtilities.h"
diff --git a/Engine/Plugins/Experimental/Animation/SkeletalMeshModelingTools/Source/SkeletalMeshModelingTools/Private/SReferenceSkeletonTree.cpp b/Engine/Plugins/Experimental/Animation/SkeletalMeshModelingTools/Source/SkeletalMeshModelingTools/Private/SReferenceSkeletonTree.cpp
--- a/Engine/Plugins/Experimental/Animation/SkeletalMeshModelingTools/Source/SkeletalMeshModelingTools/Private/SReferenceSkeletonTree.cpp
+++ b/Engine/Plugins/Experimental/Animation/SkeletalMeshModelingTools/Source/SkeletalMeshModelingTools/Private/SReferenceSkeletonTree.cpp
@@ -99,6 +99,7 @@ void SBoneItem::Construct(const FArguments& InArgs)
 		.Text(this, &SBoneItem::GetName)
 		// .Font(TextFont)
 		.ColorAndOpacity_Static(&SBoneItem::GetTextColor, InArgs._IsSelected)
+		.MaximumLength(NAME_SIZE-1)
 		.OnVerifyTextChanged(this, &SBoneItem::OnVerifyNameChanged)
 		.OnTextCommitted(InArgs._Delegates.OnBoneNameCommitted)
 		.MultiLine(false)
diff --git a/Engine/Plugins/FX/Niagara/Source/NiagaraEditor/Private/Customizations/NiagaraParameterCollectionCustomNodeBuilder.h b/Engine/Plugins/FX/Niagara/Source/NiagaraEditor/Private/Customizations/NiagaraParameterCollectionCustomNodeBuilder.h
--- a/Engine/Plugins/FX/Niagara/Source/NiagaraEditor/Private/Customizations/NiagaraParameterCollectionCustomNodeBuilder.h
+++ b/Engine/Plugins/FX/Niagara/Source/NiagaraEditor/Private/Customizations/NiagaraParameterCollectionCustomNodeBuilder.h
@@ -67,6 +67,7 @@ public:
 					SNew(SInlineEditableTextBlock)
 					.Style(FNiagaraEditorStyle::Get(), "NiagaraEditor.ParameterInlineEditableText")
 					.Text(Parameter, &INiagaraParameterViewModel::GetNameText)
+					.MaximumLength(NAME_SIZE-1)
 					.OnVerifyTextChanged(Parameter, &INiagaraParameterViewModel::VerifyNodeNameTextChanged)
 					.OnTextCommitted(Parameter, &INiagaraParameterViewModel::NameTextCommitted);
 			}
diff --git a/Engine/Plugins/FX/Niagara/Source/NiagaraEditor/Private/Customizations/NiagaraTypeCustomizations.cpp b/Engine/Plugins/FX/Niagara/Source/NiagaraEditor/Private/Customizations/NiagaraTypeCustomizations.cpp
--- a/Engine/Plugins/FX/Niagara/Source/NiagaraEditor/Private/Customizations/NiagaraTypeCustomizations.cpp
+++ b/Engine/Plugins/FX/Niagara/Source/NiagaraEditor/Private/Customizations/NiagaraTypeCustomizations.cpp
@@ -2783,6 +2783,7 @@ void FNiagaraVariableDetailsCustomization::CustomizeChildren(TSharedRef<IPropert
 						.Font(IDetailLayoutBuilder::GetDetailFont())
 						.SelectAllTextWhenFocused(true)
 						.ClearKeyboardFocusOnCommit(false)
+						.MaximumLength(NAME_SIZE-1)
 						.OnTextCommitted_Lambda([NameHandle](const FText& NewName, ETextCommit::Type)
 						{
 							FScopedTransaction Transaction(LOCTEXT("ChangeVarName", "Change variable name"));
@@ -2957,6 +2958,7 @@ void FNiagaraDataChannelVariableDetailsCustomization::CustomizeChildren(TSharedR
 						.Font(IDetailLayoutBuilder::GetDetailFont())
 						.SelectAllTextWhenFocused(true)
 						.ClearKeyboardFocusOnCommit(false)
+						.MaximumLength(NAME_SIZE-1)
 						.OnTextCommitted_Lambda([NameHandle](const FText& NewName, ETextCommit::Type)
 						{
 							FScopedTransaction Transaction(LOCTEXT("ChangeVarName", "Change variable name"));
diff --git a/Engine/Plugins/VirtualProduction/RemoteControlWebInterface/Source/RemoteControlWebInterface/Private/RCWebInterfaceCustomizations.cpp b/Engine/Plugins/VirtualProduction/RemoteControlWebInterface/Source/RemoteControlWebInterface/Private/RCWebInterfaceCustomizations.cpp
--- a/Engine/Plugins/VirtualProduction/RemoteControlWebInterface/Source/RemoteControlWebInterface/Private/RCWebInterfaceCustomizations.cpp
+++ b/Engine/Plugins/VirtualProduction/RemoteControlWebInterface/Source/RemoteControlWebInterface/Private/RCWebInterfaceCustomizations.cpp
@@ -24,6 +24,7 @@
 #include "SSearchableComboBox.h"
 #include "UObject/TextProperty.h"
 #include "Widgets/Input/SButton.h"
+#include "Widgets/Input/SEditableTextBox.h"
 #include "Widgets/Text/STextBlock.h"
 #include "Widgets/DeclarativeSyntaxSupport.h"
 #include "Widgets/Layout/SWidgetSwitcher.h"
diff --git a/Engine/Source/Developer/ToolWidgets/Private/SSearchableComboBox.cpp b/Engine/Source/Developer/ToolWidgets/Private/SSearchableComboBox.cpp
--- a/Engine/Source/Developer/ToolWidgets/Private/SSearchableComboBox.cpp
+++ b/Engine/Source/Developer/ToolWidgets/Private/SSearchableComboBox.cpp
@@ -2,6 +2,7 @@
 
 #include "SSearchableComboBox.h"
 #include "Framework/Application/SlateUser.h"
+#include "Widgets/Input/SEditableTextBox.h"
 
 #define LOCTEXT_NAMESPACE "SearchableComboBox"
 
diff --git a/Engine/Source/Developer/ToolWidgets/Public/SSearchableComboBox.h b/Engine/Source/Developer/ToolWidgets/Public/SSearchableComboBox.h
--- a/Engine/Source/Developer/ToolWidgets/Public/SSearchableComboBox.h
+++ b/Engine/Source/Developer/ToolWidgets/Public/SSearchableComboBox.h
@@ -6,10 +6,13 @@
 #include "Input/Reply.h"
 #include "Widgets/SWidget.h"
 #include "Widgets/Input/SComboButton.h"
+#if UE_ENABLE_INCLUDE_ORDER_DEPRECATED_IN_5_6
 #include "Widgets/Input/SEditableTextBox.h"
+#endif
 #include "Widgets/Views/SListView.h"
 #include "Widgets/Input/SComboBox.h"
 
+class SEditableTextBox;
 /**
  * A searchable text combo box
  */
diff --git a/Engine/Source/Editor/CurveTableEditor/Private/CurveTableEditor.cpp b/Engine/Source/Editor/CurveTableEditor/Private/CurveTableEditor.cpp
--- a/Engine/Source/Editor/CurveTableEditor/Private/CurveTableEditor.cpp
+++ b/Engine/Source/Editor/CurveTableEditor/Private/CurveTableEditor.cpp
@@ -171,6 +171,7 @@ class FCurveTableEditorItem : public ICurveEditorTreeItem,  public TSharedFromTh
 					SAssignNew(InlineRenameWidget, SInlineEditableTextBlock)
 					.Text(DisplayName)
 					.ColorAndOpacity(FSlateColor::UseForeground())
+					.MaximumLength(NAME_SIZE-1)
 					.OnTextCommitted(this, &FCurveTableEditorItem::HandleNameCommitted)
 					.OnVerifyTextChanged(this, &FCurveTableEditorItem::VerifyNameChanged)
 				];
diff --git a/Engine/Source/Editor/DetailCustomizations/Private/AnimSequenceDetails.cpp b/Engine/Source/Editor/DetailCustomizations/Private/AnimSequenceDetails.cpp
--- a/Engine/Source/Editor/DetailCustomizations/Private/AnimSequenceDetails.cpp
+++ b/Engine/Source/Editor/DetailCustomizations/Private/AnimSequenceDetails.cpp
@@ -52,6 +52,7 @@
 #include "UObject/UnrealType.h"
 #include "Viewports.h"
 #include "Widgets/Input/SButton.h"
+#include "Widgets/Input/SEditableTextBox.h"
 #include "Widgets/Layout/SBorder.h"
 #include "Widgets/SBoxPanel.h"
 #include "Widgets/SViewport.h"
diff --git a/Engine/Source/Editor/GraphEditor/Private/SGraphEditorActionMenu.cpp b/Engine/Source/Editor/GraphEditor/Private/SGraphEditorActionMenu.cpp
--- a/Engine/Source/Editor/GraphEditor/Private/SGraphEditorActionMenu.cpp
+++ b/Engine/Source/Editor/GraphEditor/Private/SGraphEditorActionMenu.cpp
@@ -13,8 +13,6 @@
 #include "Types/SlateStructs.h"
 #include "Widgets/Layout/SBox.h"
 
-class SEditableTextBox;
-
 SGraphEditorActionMenu::~SGraphEditorActionMenu()
 {
 	OnClosedCallback.ExecuteIfBound();
diff --git a/Engine/Source/Editor/GraphEditor/Public/SGraphEditorActionMenu.h b/Engine/Source/Editor/GraphEditor/Public/SGraphEditorActionMenu.h
--- a/Engine/Source/Editor/GraphEditor/Public/SGraphEditorActionMenu.h
+++ b/Engine/Source/Editor/GraphEditor/Public/SGraphEditorActionMenu.h
@@ -11,7 +11,9 @@
 #include "Templates/SharedPointer.h"
 #include "Types/SlateEnums.h"
 #include "Widgets/DeclarativeSyntaxSupport.h"
+#if UE_ENABLE_INCLUDE_ORDER_DEPRECATED_IN_5_6
 #include "Widgets/Input/SEditableTextBox.h"
+#endif
 #include "Widgets/Layout/SBorder.h"
 
 class SEditableTextBox;
diff --git a/Engine/Source/Editor/InputBindingEditor/Private/Widgets/SChordEditor.cpp b/Engine/Source/Editor/InputBindingEditor/Private/Widgets/SChordEditor.cpp
--- a/Engine/Source/Editor/InputBindingEditor/Private/Widgets/SChordEditor.cpp
+++ b/Engine/Source/Editor/InputBindingEditor/Private/Widgets/SChordEditor.cpp
@@ -14,7 +14,6 @@
 #include "Templates/UniquePtr.h"
 #include "Templates/UnrealTemplate.h"
 #include "UObject/NameTypes.h"
-#include "Widgets/Text/ISlateEditableTextWidget.h"
 #include "Widgets/Text/SlateEditableTextLayout.h"
 
 struct FGeometry;
diff --git a/Engine/Source/Editor/Persona/Private/SPoseEditor.cpp b/Engine/Source/Editor/Persona/Private/SPoseEditor.cpp
--- a/Engine/Source/Editor/Persona/Private/SPoseEditor.cpp
+++ b/Engine/Source/Editor/Persona/Private/SPoseEditor.cpp
@@ -136,6 +136,7 @@ TSharedRef< SWidget > SPoseListRow::GenerateWidgetForColumn(const FName& ColumnN
 				.Text(this, &SPoseListRow::GetName)
 				.HighlightText(FilterText)
 				.ToolTipText(LOCTEXT("PoseName_ToolTip", "Modify Pose Name - Make sure this name is unique among all curves per skeleton."))
+				.MaximumLength(NAME_SIZE-1)
 				.OnVerifyTextChanged(this, &SPoseListRow::OnVerifyNameChanged)
 				.OnTextCommitted(this, &SPoseListRow::OnNameCommitted)
 			];
diff --git a/Engine/Source/Editor/Persona/Private/SRenameMorphTargetDialog.cpp b/Engine/Source/Editor/Persona/Private/SRenameMorphTargetDialog.cpp
--- a/Engine/Source/Editor/Persona/Private/SRenameMorphTargetDialog.cpp
+++ b/Engine/Source/Editor/Persona/Private/SRenameMorphTargetDialog.cpp
@@ -72,6 +72,7 @@ void SRenameMorphTargetDialog::Construct(const FArguments& InArgs)
 				SAssignNew(NewMorphTargetNameTextBox, SEditableTextBox)
 				.Font(FAppStyle::GetFontStyle( TEXT("PropertyWindow.NormalFont")))
 				.Text(FText::FromName(MorphTarget->GetFName()))
+				.MaximumLength(NAME_SIZE-1)
 				.OnVerifyTextChanged(this, &SRenameMorphTargetDialog::OnVerifyNameChanged)
 				.OnTextCommitted(this, &SRenameMorphTargetDialog::OnRenameTextCommitted)
 			]
diff --git a/Engine/Source/Editor/Persona/Private/SSkinWeightProfileImportOptions.cpp b/Engine/Source/Editor/Persona/Private/SSkinWeightProfileImportOptions.cpp
--- a/Engine/Source/Editor/Persona/Private/SSkinWeightProfileImportOptions.cpp
+++ b/Engine/Source/Editor/Persona/Private/SSkinWeightProfileImportOptions.cpp
@@ -7,6 +7,7 @@
 #include "Modules/ModuleManager.h"
 #include "Widgets/Views/SListView.h"
 #include "Widgets/Layout/SUniformGridPanel.h"
+#include "Widgets/Input/SEditableTextBox.h"
 #include "Widgets/Input/SButton.h"
 #include "PropertyHandle.h"
 #include "DetailLayoutBuilder.h"
diff --git a/Engine/Source/Editor/Persona/Public/SSkinWeightProfileImportOptions.h b/Engine/Source/Editor/Persona/Public/SSkinWeightProfileImportOptions.h
--- a/Engine/Source/Editor/Persona/Public/SSkinWeightProfileImportOptions.h
+++ b/Engine/Source/Editor/Persona/Public/SSkinWeightProfileImportOptions.h
@@ -10,7 +10,9 @@
 #include "Widgets/Input/SCheckBox.h"
 #include "Widgets/Text/STextBlock.h"
 #include "IDetailCustomization.h"
+#if UE_ENABLE_INCLUDE_ORDER_DEPRECATED_IN_5_6
 #include "Widgets/Input/SEditableTextBox.h"
+#endif
 #include "PropertyHandle.h"
 #include "Engine/SkeletalMesh.h"
 
@@ -18,6 +20,7 @@
 
 // Forward declares
 class IDetailsView;
+class SEditableTextBox;
 
 UCLASS(Config = EditorPerProjectUserSettings)
 class USkinWeightImportOptions : public UObject
diff --git a/Engine/Source/Editor/PropertyEditor/Private/UserInterface/PropertyEditor/SPropertyEditorText.cpp b/Engine/Source/Editor/PropertyEditor/Private/UserInterface/PropertyEditor/SPropertyEditorText.cpp
--- a/Engine/Source/Editor/PropertyEditor/Private/UserInterface/PropertyEditor/SPropertyEditorText.cpp
+++ b/Engine/Source/Editor/PropertyEditor/Private/UserInterface/PropertyEditor/SPropertyEditorText.cpp
@@ -101,6 +101,7 @@ void SPropertyEditorText::Construct( const FArguments& InArgs, const TSharedRef<
 				.Font(InArgs._Font)
 				.SelectAllTextWhenFocused(false)
 				.ClearKeyboardFocusOnCommit(false)
+				.MaximumLength(MaxLength)
 				.OnTextCommitted(this, &SPropertyEditorText::OnTextCommitted)
 				.OnVerifyTextChanged(this, &SPropertyEditorText::OnVerifyTextChanged)
 				.SelectAllTextOnCommit(false)
@@ -126,6 +127,7 @@ void SPropertyEditorText::Construct( const FArguments& InArgs, const TSharedRef<
 				.Font( InArgs._Font )
 				.SelectAllTextWhenFocused( true )
 				.ClearKeyboardFocusOnCommit(false)
+				.MaximumLength(MaxLength)
 				.OnTextCommitted( this, &SPropertyEditorText::OnTextCommitted )
 				.OnVerifyTextChanged( this, &SPropertyEditorText::OnVerifyTextChanged )
 				.SelectAllTextOnCommit( true )
@@ -228,12 +230,6 @@ bool SPropertyEditorText::OnVerifyTextChanged(const FText& Text, FText& OutError
 {
 	const FString& TextString = Text.ToString();
 
-	if (MaxLength > 0 && TextString.Len() > MaxLength)
-	{
-		OutError = FText::Format(LOCTEXT("PropertyTextTooLongError", "This value is too long ({0}/{1} characters)"), TextString.Len(), MaxLength);
-		return false;
-	}
-
 	if (!AllowedCharacters.IsEmpty())
 	{
 		if (!TextString.IsEmpty() && !AllowedCharacters.AreAllCharsIncluded(TextString))
diff --git a/Engine/Source/Editor/UnrealEd/Private/UnrealEdPrivatePCH.h b/Engine/Source/Editor/UnrealEd/Private/UnrealEdPrivatePCH.h
--- a/Engine/Source/Editor/UnrealEd/Private/UnrealEdPrivatePCH.h
+++ b/Engine/Source/Editor/UnrealEd/Private/UnrealEdPrivatePCH.h
@@ -416,7 +416,6 @@
 #include "Framework/Views/TableViewTypeTraits.h"
 #include "Widgets/Views/SListView.h"
 #include "Framework/Docking/LayoutService.h"
-#include "Widgets/Text/ISlateEditableTextWidget.h"
 #include "Widgets/Input/IVirtualKeyboardEntry.h"
 #include "Widgets/Input/SMenuAnchor.h"
 #include "Widgets/Input/SEditableText.h"
diff --git a/Engine/Source/Runtime/Slate/Private/Framework/Text/TextEditHelper.cpp b/Engine/Source/Runtime/Slate/Private/Framework/Text/TextEditHelper.cpp
--- a/Engine/Source/Runtime/Slate/Private/Framework/Text/TextEditHelper.cpp
+++ b/Engine/Source/Runtime/Slate/Private/Framework/Text/TextEditHelper.cpp
@@ -4,6 +4,8 @@
 #include "Fonts/FontMeasure.h"
 #include "Framework/Application/SlateApplication.h"
 
+#define LOCTEXT_NAMESPACE "FTextEditHelper"
+
 float FTextEditHelper::GetFontHeight(const FSlateFontInfo& FontInfo)
 {
 	const TSharedRef< FSlateFontMeasure > FontMeasure = FSlateApplication::Get().GetRenderer()->GetFontMeasureService();
@@ -16,3 +18,17 @@ float FTextEditHelper::CalculateCaretWidth(const float FontMaxCharHeight)
 	// We clamp the lower limit to 1 to avoid it being invisible, and the upper limit to 2 to avoid tall fonts having very wide carets.
 	return FMath::Clamp(EditableTextDefs::CaretWidthPercent * FontMaxCharHeight, 1.0f, 2.0f);
 }
+
+bool FTextEditHelper::VerifyTextLength(const FText& InText, FText& OutErrorMessage, int32 InMaximumLength)
+{
+	const int TextLength = InText.ToString().Len();
+	if (InMaximumLength > 0 && TextLength > InMaximumLength)
+	{
+		OutErrorMessage = FText::Format(LOCTEXT("TextTooLong", "This text is too long. It uses {0} characters of {1} allowed."), TextLength, InMaximumLength);
+		return false;
+	}
+
+	return true;
+}
+
+#undef LOCTEXT_NAMESPACE
diff --git a/Engine/Source/Runtime/Slate/Private/Framework/Text/TextEditHelper.h b/Engine/Source/Runtime/Slate/Private/Framework/Text/TextEditHelper.h
--- a/Engine/Source/Runtime/Slate/Private/Framework/Text/TextEditHelper.h
+++ b/Engine/Source/Runtime/Slate/Private/Framework/Text/TextEditHelper.h
@@ -43,4 +43,6 @@ public:
 	 * @return The width of the caret (might be clamped for very small fonts)
 	 */
 	static float CalculateCaretWidth(const float FontMaxCharHeight);
+
+	static bool VerifyTextLength(const FText& InText, FText& OutErrorMessage, int32 LengthToCheck);
 };
diff --git a/Engine/Source/Runtime/Slate/Private/Widgets/Input/SEditableTextBox.cpp b/Engine/Source/Runtime/Slate/Private/Widgets/Input/SEditableTextBox.cpp
--- a/Engine/Source/Runtime/Slate/Private/Widgets/Input/SEditableTextBox.cpp
+++ b/Engine/Source/Runtime/Slate/Private/Widgets/Input/SEditableTextBox.cpp
@@ -1,6 +1,7 @@
 // Copyright Epic Games, Inc. All Rights Reserved.
 
 #include "Widgets/Input/SEditableTextBox.h"
+#include "Framework/Text/TextEditHelper.h"
 #include "Widgets/SBoxPanel.h"
 #include "Widgets/Layout/SBox.h"
 #include "Widgets/Notifications/SPopUpErrorText.h"
@@ -35,6 +36,7 @@ void SEditableTextBox::Construct( const FArguments& InArgs )
 	BackgroundColorOverride = InArgs._BackgroundColor;
 	ReadOnlyForegroundColorOverride = InArgs._ReadOnlyForegroundColor;
 	FocusedForegroundColorOverride = InArgs._FocusedForegroundColor;
+	MaximumLength = InArgs._MaximumLength;
 	OnTextChanged = InArgs._OnTextChanged;
 	OnVerifyTextChanged = InArgs._OnVerifyTextChanged;
 	OnTextCommitted = InArgs._OnTextCommitted;
@@ -415,6 +417,11 @@ void SEditableTextBox::SetFocusedForegroundColor(const TAttribute<FSlateColor>&
 	FocusedForegroundColorOverride = InFocusedForegroundColor;
 }
 
+void SEditableTextBox::SetMaximumLength(const TAttribute<int32>& InMaximumLength)
+{
+	MaximumLength = InMaximumLength;
+}
+
 void SEditableTextBox::SetMinimumDesiredWidth(const TAttribute<float>& InMinimumDesiredWidth)
 {
 	EditableText->SetMinDesiredWidth(InMinimumDesiredWidth);
@@ -497,10 +504,12 @@ void SEditableTextBox::OnEditableTextChanged(const FText& InText)
 {
 	OnTextChanged.ExecuteIfBound(InText);
 
-	if (OnVerifyTextChanged.IsBound())
+	const int32 MaximumLengthValue = MaximumLength.Get();
+	if (OnVerifyTextChanged.IsBound() || MaximumLengthValue >= 0)
 	{
 		FText OutErrorMessage;
-		if (!OnVerifyTextChanged.Execute(InText, OutErrorMessage))
+		if (!FTextEditHelper::VerifyTextLength(InText, OutErrorMessage, MaximumLengthValue) ||
+			(OnVerifyTextChanged.IsBound() && !OnVerifyTextChanged.Execute(InText, OutErrorMessage)))
 		{
 			// Display as an error.
 			SetError(OutErrorMessage);
@@ -514,18 +523,16 @@ void SEditableTextBox::OnEditableTextChanged(const FText& InText)
 
 void SEditableTextBox::OnEditableTextCommitted(const FText& InText, ETextCommit::Type InCommitType)
 {
-	if (OnVerifyTextChanged.IsBound())
+	FText OutErrorMessage;
+	if (!FTextEditHelper::VerifyTextLength(InText, OutErrorMessage, MaximumLength.Get()) ||
+		(OnVerifyTextChanged.IsBound() && !OnVerifyTextChanged.Execute(InText, OutErrorMessage)))
 	{
-		FText OutErrorMessage;
-		if (!OnVerifyTextChanged.Execute(InText, OutErrorMessage))
+		// Display as an error.
+		if (InCommitType == ETextCommit::OnEnter)
 		{
-           	// Display as an error.
-			if (InCommitType == ETextCommit::OnEnter)
-			{
-				SetError(OutErrorMessage);
-			}
-			return;
+			SetError(OutErrorMessage);
 		}
+		return;
 	}
 
 	// Text commited without errors, so clear error text
diff --git a/Engine/Source/Runtime/Slate/Private/Widgets/Input/SMultiLineEditableTextBox.cpp b/Engine/Source/Runtime/Slate/Private/Widgets/Input/SMultiLineEditableTextBox.cpp
--- a/Engine/Source/Runtime/Slate/Private/Widgets/Input/SMultiLineEditableTextBox.cpp
+++ b/Engine/Source/Runtime/Slate/Private/Widgets/Input/SMultiLineEditableTextBox.cpp
@@ -1,6 +1,7 @@
 // Copyright Epic Games, Inc. All Rights Reserved.
 
 #include "Widgets/Input/SMultiLineEditableTextBox.h"
+#include "Framework/Text/TextEditHelper.h"
 #include "Widgets/SBoxPanel.h"
 #include "Widgets/Layout/SBox.h"
 #include "Widgets/Notifications/SPopUpErrorText.h"
@@ -41,6 +42,7 @@ void SMultiLineEditableTextBox::Construct( const FArguments& InArgs )
 	FontOverride = InArgs._Font;
 	ForegroundColorOverride = InArgs._ForegroundColor;
 	BackgroundColorOverride = InArgs._BackgroundColor;
+	MaximumLength = InArgs._MaximumLength;
 	ReadOnlyForegroundColorOverride = InArgs._ReadOnlyForegroundColor;
 	FocusedForegroundColorOverride = InArgs._FocusedForegroundColor;
 	bSelectWordOnMouseDoubleClick = InArgs._SelectWordOnMouseDoubleClick;
@@ -347,6 +349,11 @@ void SMultiLineEditableTextBox::SetReadOnlyForegroundColor(const TAttribute<FSla
 	ReadOnlyForegroundColorOverride = InReadOnlyForegroundColor;
 }
 
+void SMultiLineEditableTextBox::SetMaximumLength(const TAttribute<int32>& InMaximumLength)
+{
+	MaximumLength = InMaximumLength;
+}
+
 void SMultiLineEditableTextBox::SetSelectWordOnMouseDoubleClick(const TAttribute<bool>& InSelectWordOnMouseDoubleClick)
 {
 	EditableText->SetSelectWordOnMouseDoubleClick(InSelectWordOnMouseDoubleClick);
@@ -612,10 +619,12 @@ void SMultiLineEditableTextBox::OnEditableTextChanged(const FText& InText)
 {
 	OnTextChanged.ExecuteIfBound(InText);
 
-	if (OnVerifyTextChanged.IsBound())
+	const int32 MaximumLengthValue = MaximumLength.Get();
+	if (OnVerifyTextChanged.IsBound() || MaximumLengthValue >= 0)
 	{
 		FText OutErrorMessage;
-		if (!OnVerifyTextChanged.Execute(InText, OutErrorMessage))
+		if (!FTextEditHelper::VerifyTextLength(InText, OutErrorMessage, MaximumLengthValue) ||
+			(OnVerifyTextChanged.IsBound() && !OnVerifyTextChanged.Execute(InText, OutErrorMessage)))
 		{
 			// Display as an error.
 			SetError(OutErrorMessage);
@@ -629,28 +638,21 @@ void SMultiLineEditableTextBox::OnEditableTextChanged(const FText& InText)
 
 void SMultiLineEditableTextBox::OnEditableTextCommitted(const FText& InText, ETextCommit::Type InCommitType)
 {
-	if (OnVerifyTextChanged.IsBound())
+	FText OutErrorMessage;
+	if (!FTextEditHelper::VerifyTextLength(InText, OutErrorMessage, MaximumLength.Get()) ||
+		(OnVerifyTextChanged.IsBound() && !OnVerifyTextChanged.Execute(InText, OutErrorMessage)))
 	{
-		FText OutErrorMessage;
-		if (!OnVerifyTextChanged.Execute(InText, OutErrorMessage))
+		// Display as an error.
+		if (InCommitType == ETextCommit::OnEnter)
 		{
-			// Display as an error.
-			if (InCommitType == ETextCommit::OnEnter)
-			{
-				SetError(OutErrorMessage);
-			}
-			return;
-		}
-		else
-		{
-			if (InCommitType == ETextCommit::OnEnter)
-			{
-				SetError(FText::GetEmpty());
-			}
-
+			SetError(OutErrorMessage);
 		}
+		return;
 	}
 
+	// Text commited without errors, so clear error text
+	SetError(FText::GetEmpty());
+
 	OnTextCommitted.ExecuteIfBound(InText, InCommitType);
 }
 
diff --git a/Engine/Source/Runtime/Slate/Private/Widgets/Text/SInlineEditableTextBlock.cpp b/Engine/Source/Runtime/Slate/Private/Widgets/Text/SInlineEditableTextBlock.cpp
--- a/Engine/Source/Runtime/Slate/Private/Widgets/Text/SInlineEditableTextBlock.cpp
+++ b/Engine/Source/Runtime/Slate/Private/Widgets/Text/SInlineEditableTextBlock.cpp
@@ -1,8 +1,9 @@
 // Copyright Epic Games, Inc. All Rights Reserved.
 
 #include "Widgets/Text/SInlineEditableTextBlock.h"
-#include "Widgets/SBoxPanel.h"
 #include "Framework/Application/SlateApplication.h"
+#include "Framework/Text/TextEditHelper.h"
+#include "Widgets/SBoxPanel.h"
 #include "Widgets/Text/STextBlock.h"
 #include "Widgets/Input/SMultiLineEditableTextBox.h"
 #include "Widgets/Input/SEditableTextBox.h"
@@ -15,9 +16,11 @@ void SInlineEditableTextBlock::Construct( const FArguments& InArgs )
 	OnBeginTextEditDelegate = InArgs._OnBeginTextEdit;
 	OnTextCommittedDelegate = InArgs._OnTextCommitted;
 	IsSelected = InArgs._IsSelected;
-	OnVerifyTextChanged= InArgs._OnVerifyTextChanged;
+	OnVerifyTextChanged = InArgs._OnVerifyTextChanged;
 	Text = InArgs._Text;
 	bIsReadOnly = InArgs._IsReadOnly;
+	MaximumLength = InArgs._MaximumLength;
+
 	bIsMultiLine = InArgs._MultiLine;
 	DoubleSelectDelay = 0.0f;
 
@@ -173,6 +176,11 @@ void SInlineEditableTextBlock::SetReadOnly(const TAttribute<bool>& bInIsReadOnly
 	bIsReadOnly = bInIsReadOnly;
 }
 
+void SInlineEditableTextBlock::SetMaximumLength(const TAttribute<int32>& InMaximumLength)
+{
+	MaximumLength = InMaximumLength;
+}
+
 void SInlineEditableTextBlock::SetText( const TAttribute< FText >& InText )
 {
 	Text = InText;
@@ -291,8 +299,8 @@ void SInlineEditableTextBlock::OnTextChanged(const FText& InText)
 	if(IsInEditMode())
 	{
 		FText OutErrorMessage;
-
-		if(OnVerifyTextChanged.IsBound() && !OnVerifyTextChanged.Execute(InText, OutErrorMessage))
+		if(!FTextEditHelper::VerifyTextLength(InText, OutErrorMessage, MaximumLength.Get()) ||
+			(OnVerifyTextChanged.IsBound() && !OnVerifyTextChanged.Execute(InText, OutErrorMessage)))
 		{
 			SetTextBoxError( OutErrorMessage );
 		}
@@ -314,40 +322,39 @@ void SInlineEditableTextBlock::OnTextBoxCommitted(const FText& InText, ETextComm
 	}
 	else if(IsInEditMode())
 	{
-		if(OnVerifyTextChanged.IsBound())
+		if(InCommitType == ETextCommit::OnEnter)
 		{
-			if(InCommitType == ETextCommit::OnEnter)
-			{
-				FText OutErrorMessage;
-				if(!OnVerifyTextChanged.Execute(InText, OutErrorMessage))
-				{
-					// Display as an error.
-					SetTextBoxError( OutErrorMessage );
-					return;
-				}
-			}
-			else if(InCommitType == ETextCommit::OnUserMovedFocus || InCommitType == ETextCommit::Default)
+			FText OutErrorMessage;
+			if (!FTextEditHelper::VerifyTextLength(InText, OutErrorMessage, MaximumLength.Get()) ||
+				(OnVerifyTextChanged.IsBound() && !OnVerifyTextChanged.Execute(InText, OutErrorMessage)))
 			{
-				FText OutErrorMessage;
-				if(!OnVerifyTextChanged.Execute(InText, OutErrorMessage))
-				{
-					CancelEditMode();
-
-					// Commit the name, certain actions might need to be taken by the bound function
-					OnTextCommittedDelegate.ExecuteIfBound(Text.Get(), InCommitType);
-				
-					return;
-				}
+				// Display as an error.
+				SetTextBoxError( OutErrorMessage );
+				return;
 			}
-			else // When the user removes all focus from the window, revert the name
+		}
+		else if(InCommitType == ETextCommit::OnUserMovedFocus || InCommitType == ETextCommit::Default)
+		{
+			FText OutErrorMessage;
+			if (!FTextEditHelper::VerifyTextLength(InText, OutErrorMessage, MaximumLength.Get()) ||
+				(OnVerifyTextChanged.IsBound() && !OnVerifyTextChanged.Execute(InText, OutErrorMessage)))
 			{
 				CancelEditMode();
 
 				// Commit the name, certain actions might need to be taken by the bound function
 				OnTextCommittedDelegate.ExecuteIfBound(Text.Get(), InCommitType);
+				
 				return;
 			}
 		}
+		else // When the user removes all focus from the window, revert the name
+		{
+			CancelEditMode();
+
+			// Commit the name, certain actions might need to be taken by the bound function
+			OnTextCommittedDelegate.ExecuteIfBound(Text.Get(), InCommitType);
+			return;
+		}
 		
 		ExitEditingMode();
 
diff --git a/Engine/Source/Runtime/Slate/Public/SlateBasics.h b/Engine/Source/Runtime/Slate/Public/SlateBasics.h
--- a/Engine/Source/Runtime/Slate/Public/SlateBasics.h
+++ b/Engine/Source/Runtime/Slate/Public/SlateBasics.h
@@ -86,7 +86,9 @@ MONOLITHIC_HEADER_BOILERPLATE()
 #include "Framework/MultiBox/MultiBoxBuilder.h"
 #include "Framework/MultiBox/MultiBoxExtender.h"
 #include "Widgets/Input/IVirtualKeyboardEntry.h"
+#if UE_ENABLE_INCLUDE_ORDER_DEPRECATED_IN_5_6
 #include "Widgets/Text/ISlateEditableTextWidget.h"
+#endif
 #include "Widgets/Text/SMultiLineEditableText.h"
 #include "Widgets/Input/SMultiLineEditableTextBox.h"
 #include "Widgets/Input/SEditableText.h"
diff --git a/Engine/Source/Runtime/Slate/Public/SlateExtras.h b/Engine/Source/Runtime/Slate/Public/SlateExtras.h
--- a/Engine/Source/Runtime/Slate/Public/SlateExtras.h
+++ b/Engine/Source/Runtime/Slate/Public/SlateExtras.h
@@ -73,7 +73,9 @@ MONOLITHIC_HEADER_BOILERPLATE()
 #include "Framework/MultiBox/MultiBoxExtender.h"
 #include "Framework/MultiBox/MultiBoxBuilder.h"
 #include "Widgets/Input/IVirtualKeyboardEntry.h"
+#if UE_ENABLE_INCLUDE_ORDER_DEPRECATED_IN_5_6
 #include "Widgets/Text/ISlateEditableTextWidget.h"
+#endif
 #include "Widgets/Layout/SScrollBar.h"
 #include "Widgets/Text/SMultiLineEditableText.h"
 #include "Widgets/Input/SMultiLineEditableTextBox.h"
diff --git a/Engine/Source/Runtime/Slate/Public/SlateSharedPCH.h b/Engine/Source/Runtime/Slate/Public/SlateSharedPCH.h
--- a/Engine/Source/Runtime/Slate/Public/SlateSharedPCH.h
+++ b/Engine/Source/Runtime/Slate/Public/SlateSharedPCH.h
@@ -139,7 +139,9 @@
 #include "Widgets/Layout/SSplitter.h"
 #include "Widgets/Notifications/SErrorText.h"
 #include "Widgets/SToolTip.h"
+#if UE_ENABLE_INCLUDE_ORDER_DEPRECATED_IN_5_6
 #include "Widgets/Text/ISlateEditableTextWidget.h"
+#endif
 #include "Widgets/Text/STextBlock.h"
 #include "Widgets/Views/SExpanderArrow.h"
 #include "Widgets/Views/SHeaderRow.h"
diff --git a/Engine/Source/Runtime/Slate/Public/Widgets/Input/SEditableTextBox.h b/Engine/Source/Runtime/Slate/Public/Widgets/Input/SEditableTextBox.h
--- a/Engine/Source/Runtime/Slate/Public/Widgets/Input/SEditableTextBox.h
+++ b/Engine/Source/Runtime/Slate/Public/Widgets/Input/SEditableTextBox.h
@@ -46,6 +46,7 @@ public:
 		, _RevertTextOnEscape( false )
 		, _ClearKeyboardFocusOnCommit( true )
 		, _Justification(ETextJustify::Left)
+		, _MaximumLength(-1)
 		, _AllowContextMenu(true)
 		, _MinDesiredWidth( 0.0f )
 		, _SelectAllTextOnCommit( false )
@@ -105,6 +106,9 @@ public:
 		/** How should the value be justified in the editable text field. */
 		SLATE_ATTRIBUTE(ETextJustify::Type, Justification)
 
+		/** Maximum length of the text. Input will refuse text longer than that. */
+		SLATE_ATTRIBUTE(int32, MaximumLength)
+
 		/** Whether the context menu can be opened */
 		SLATE_ATTRIBUTE(bool, AllowContextMenu)
 
@@ -276,6 +280,13 @@ public:
 	 */
 	SLATE_API void SetFocusedForegroundColor(const TAttribute<FSlateColor>& InFocusedForegroundColor);
 
+	/**
+	 * Sets the maximum text length that will be accepted by the widget
+	 *
+	 * @param  InMaximumLength		Length to check for. Only values superior to 0 will be taken into account, otherwise no length check will occur.
+	 */
+	SLATE_API void SetMaximumLength(const TAttribute<int32>& InMaximumLength);
+
 	/**
 	 * Sets the minimum width that a text box should be.
 	 *
@@ -442,6 +453,9 @@ protected:
 	/** Focused foreground color (overrides style) */
 	TAttribute<FSlateColor> FocusedForegroundColorOverride;
 
+	/** Maximum text length that will be accepted by the widget */
+	TAttribute<int32> MaximumLength;
+
 	/** Allows for inserting additional widgets that extend the functionality of the text box */
 	TSharedPtr<SHorizontalBox> Box;
 
diff --git a/Engine/Source/Runtime/Slate/Public/Widgets/Input/SMultiLineEditableTextBox.h b/Engine/Source/Runtime/Slate/Public/Widgets/Input/SMultiLineEditableTextBox.h
--- a/Engine/Source/Runtime/Slate/Public/Widgets/Input/SMultiLineEditableTextBox.h
+++ b/Engine/Source/Runtime/Slate/Public/Widgets/Input/SMultiLineEditableTextBox.h
@@ -45,6 +45,7 @@ public:
 		, _ReadOnlyForegroundColor()
 		, _FocusedForegroundColor()
 		, _Justification(ETextJustify::Left)
+		, _MaximumLength(-1)
 		, _LineHeightPercentage(1.0f)
 		, _IsReadOnly( false )
 		, _AllowMultiLine( true )
@@ -108,6 +109,9 @@ public:
 		/** How the text should be aligned with the margin. */
 		SLATE_ATTRIBUTE(ETextJustify::Type, Justification)
 
+		/** Maximum length of the text. Input will refuse text longer than that. */
+		SLATE_ATTRIBUTE(int32, MaximumLength)
+
 		/** The amount to scale each lines height by. */
 		SLATE_ATTRIBUTE(float, LineHeightPercentage)
 
@@ -334,6 +338,13 @@ public:
 	 */
 	SLATE_API void SetReadOnlyForegroundColor(const TAttribute<FSlateColor>& InReadOnlyForegroundColor);
 
+	/**
+	 * Sets the maximum text length that will be accepted by the widget
+	 *
+	 * @param  InMaximumLength		Length to check for. Only values superior to 0 will be taken into account, otherwise no length check will occur.
+	 */
+	SLATE_API void SetMaximumLength(const TAttribute<int32>& InMaximumLength);
+
 	/**
 	 * Sets whether to select word on the mouse double click
 	 *
@@ -514,6 +525,9 @@ protected:
 	/** Focused foreground color (overrides style) */
 	TAttribute<FSlateColor> FocusedForegroundColorOverride;
 
+	/** Maximum text length that will be accepted by the widget */
+	TAttribute<int32> MaximumLength;
+
 	/** Whether to disable the context menu */
 	TAttribute< bool > AllowContextMenu;
 
diff --git a/Engine/Source/Runtime/Slate/Public/Widgets/Text/SInlineEditableTextBlock.h b/Engine/Source/Runtime/Slate/Public/Widgets/Text/SInlineEditableTextBlock.h
--- a/Engine/Source/Runtime/Slate/Public/Widgets/Text/SInlineEditableTextBlock.h
+++ b/Engine/Source/Runtime/Slate/Public/Widgets/Text/SInlineEditableTextBlock.h
@@ -42,6 +42,7 @@ class SInlineEditableTextBlock: public SCompoundWidget
 		, _AutoWrapNonEditText(false)
 		, _AutoWrapMultilineEditText(false)
 		, _Justification(ETextJustify::Left)
+		, _MaximumLength(-1)
 		, _LineBreakPolicy()
 		, _IsReadOnly(false)
 		, _MultiLine(false)
@@ -81,7 +82,10 @@ class SInlineEditableTextBlock: public SCompoundWidget
 		SLATE_ATTRIBUTE( bool, AutoWrapMultilineEditText )
 
 		/** How the text should be aligned with the margin. */
-		SLATE_ATTRIBUTE( ETextJustify::Type, Justification )
+		SLATE_ATTRIBUTE(ETextJustify::Type, Justification)
+
+		/** Maximum length of the text. Input will refuse text longer than that. */
+		SLATE_ATTRIBUTE(int32, MaximumLength)
 
 		/** The iterator to use to detect appropriate soft-wrapping points for lines (or null to use the default) */
 		SLATE_ARGUMENT( TSharedPtr<IBreakIterator>, LineBreakPolicy )
@@ -155,6 +159,13 @@ class SInlineEditableTextBlock: public SCompoundWidget
 
 	SLATE_API void SetReadOnly(const TAttribute<bool>& bInIsReadOnly);
 
+	/**
+	 * Sets the maximum text length that will be accepted by the widget
+	 *
+	 * @param  InMaximumLength		Length to check for. Only values superior to 0 will be taken into account, otherwise no length check will occur.
+	 */
+	SLATE_API void SetMaximumLength(const TAttribute<int32>& InMaximumLength);
+
 	SLATE_API void SetText( const TAttribute< FText >& InText );
 	SLATE_API void SetText( const FString& InText );
 
@@ -217,6 +228,9 @@ protected:
 	/** Attribute to look up if the widget is read-only */
 	TAttribute< bool > bIsReadOnly;
 
+	/** Maximum text length that will be accepted by the widget */
+	TAttribute<int32> MaximumLength;
+
 	/** Widget to focus when we finish editing */
 	TWeakPtr<SWidget> WidgetToFocus;
 
diff --git a/Engine/Source/Runtime/Slate/Public/Widgets/Text/SlateEditableTextLayout.h b/Engine/Source/Runtime/Slate/Public/Widgets/Text/SlateEditableTextLayout.h
--- a/Engine/Source/Runtime/Slate/Public/Widgets/Text/SlateEditableTextLayout.h
+++ b/Engine/Source/Runtime/Slate/Public/Widgets/Text/SlateEditableTextLayout.h
@@ -27,6 +27,7 @@ class FSlateWindowElementList;
 class FSlateTextBlockLayout;
 class FUICommandList;
 class IBreakIterator;
+class ISlateEditableTextWidget;
 class SWindow;
 enum class ETextShapingMethod : uint8;
 
