Skip motion event in slate app if platform user index is invalid (pad likely disconnected)

#rnx
#tests repro the issue without the fix and confirm the fix doesnt make the game to crash anymore

diff --git a/Engine/Source/Runtime/Slate/Private/Framework/Application/SlateApplication.cpp b/Engine/Source/Runtime/Slate/Private/Framework/Application/SlateApplication.cpp
--- a/Engine/Source/Runtime/Slate/Private/Framework/Application/SlateApplication.cpp
+++ b/Engine/Source/Runtime/Slate/Private/Framework/Application/SlateApplication.cpp
@@ -6618,6 +6618,13 @@ void FSlateApplication::ProcessMotionDetectedEvent( const FMotionEvent& MotionEv
 	FSlateDebugging::FScopeProcessInputEvent Scope(ESlateDebuggingInputEvent::MotionDetected, MotionEvent);
 #endif
 
+	const int32 UserIndex = MotionEvent.GetUserIndex();
+	const FPlatformUserId PlatUserId = FGenericPlatformMisc::GetPlatformUserForUserIndex(MotionEvent.GetUserIndex());
+	if (!ensure(PlatUserId.IsValid()))
+	{
+		return;
+	}
+
 	if (GSlateInputMotionFiresUserInteractionEvents)
 	{
 		SetLastUserInteractionTime(this->GetCurrentTime());
