[Slate] Add space (new Spacer MenuBlock) after menu searchbox if it's not followed by a section or separator

diff --git a/Engine/Source/Runtime/Slate/Private/Framework/MultiBox/MultiBox.cpp b/Engine/Source/Runtime/Slate/Private/Framework/MultiBox/MultiBox.cpp
--- a/Engine/Source/Runtime/Slate/Private/Framework/MultiBox/MultiBox.cpp
+++ b/Engine/Source/Runtime/Slate/Private/Framework/MultiBox/MultiBox.cpp
@@ -22,6 +22,7 @@
 #include "Framework/MultiBox/SToolBarButtonBlock.h"
 #include "Framework/MultiBox/SMenuEntryBlock.h"
 #include "Framework/MultiBox/SMenuSeparatorBlock.h"
+#include "Framework/MultiBox/SMenuSpacerBlock.h"
 #include "Framework/MultiBox/MultiBoxCustomization.h"
 #include "Framework/MultiBox/SClippingHorizontalBox.h"
 #include "Framework/MultiBox/SWidgetBlock.h"
@@ -976,13 +977,7 @@ void SMultiBoxWidget::CreateSearchTextWidget()
 	TSharedRef<SBox> SearchBox =
 		SNew(SBox)
 			.Padding(FMargin(8, 0, 8, 0))
-			.Visibility_Lambda(
-				[this]() -> EVisibility
-				{
-					const bool bShow = !SearchText.IsEmpty() || ShouldShowMenuSearchField();
-					return bShow ? EVisibility::Visible : EVisibility::Collapsed;
-				}
-			)
+			.Visibility(this, &SMultiBoxWidget::GetSearchBoxVisibility)
 			[
 				SearchTextWidget.ToSharedRef()
 			];
@@ -996,6 +991,12 @@ void SMultiBoxWidget::CreateSearchTextWidget()
 	MultiBox->AddMultiBlockToFront(NewWidgetBlock);
 }
 
+EVisibility SMultiBoxWidget::GetSearchBoxVisibility() const
+{
+	const bool bShow = !SearchText.IsEmpty() || ShouldShowMenuSearchField();
+	return bShow ? EVisibility::Visible : EVisibility::Collapsed;
+}
+
 /** Called when the SearchText changes */
 void SMultiBoxWidget::OnFilterTextChanged(const FText& InFilterText)
 {
@@ -1340,6 +1341,24 @@ void SMultiBoxWidget::BuildMultiBoxWidget()
 		{
 			AddBlockWidget(Block, HorizontalBox, VerticalBox, Location, bSectionContainsIcons, OptionsBlock);
 		}
+
+		const bool bIsSearchWidgetBlock = Index == 0 && MultiBox->bHasSearchWidget;
+		if (bIsSearchWidgetBlock && NextBlock.IsValid())
+		{
+			const bool bNextBlockIsSeparatorLike = NextBlock->IsSeparator() || NextBlock->IsGroupStartBlock() || NextBlock->IsPartOfHeading();
+
+			// If the search widget is the first block and the next block isn't a separator or a group/section start, we should add padding below the search block
+			if (!bNextBlockIsSeparatorLike)
+			{
+				constexpr bool bIsPartOfHeading = true;
+				TSharedPtr<FMenuSpacerBlock> SpacerBlock = MakeShared<FMenuSpacerBlock>(Block.GetExtensionHook(), bIsPartOfHeading);
+
+				// Note that the search widget might be collapsed, so the visibility of this spacer needs to be dynamic (like the search box itself)
+				SpacerBlock->SetVisibilityOverride(TAttribute<EVisibility>(this, &SMultiBoxWidget::GetSearchBoxVisibility));
+
+				AddBlockWidget(*SpacerBlock, HorizontalBox, VerticalBox, EMultiBlockLocation::None, bSectionContainsIcons, OptionsBlock);
+			}
+		}
 	}
 
 	// The first loop above added the multibox blocks and discovered/collected all sub-menus blocks. This second loops adds all sub-menu blocks discovered in the first pass.
@@ -1909,7 +1928,7 @@ void SMultiBoxWidget::FlattenSubMenusRecursive(uint32 MaxRecursionLevels)
 	}
 }
 
-bool SMultiBoxWidget::ShouldShowMenuSearchField()
+bool SMultiBoxWidget::ShouldShowMenuSearchField() const
 {
 	const int SearchFieldTreshold = FMultiBoxSettings::MenuSearchFieldVisibilityThreshold.Get();
 
diff --git a/Engine/Source/Runtime/Slate/Private/Framework/MultiBox/SMenuSpacerBlock.cpp b/Engine/Source/Runtime/Slate/Private/Framework/MultiBox/SMenuSpacerBlock.cpp
new file mode 100644
--- /dev/null
+++ b/Engine/Source/Runtime/Slate/Private/Framework/MultiBox/SMenuSpacerBlock.cpp
@@ -0,0 +1,69 @@
+// Copyright Epic Games, Inc. All Rights Reserved.
+
+#include "Framework/MultiBox/SMenuSpacerBlock.h"
+#include "Widgets/SBoxPanel.h"
+#include "Widgets/Layout/SSeparator.h"
+
+/**
+ * Constructor
+ */
+FMenuSpacerBlock::FMenuSpacerBlock(
+	const FName& InExtensionHook,
+	bool bInIsPartOfHeading)
+	: FMultiBlock(
+		nullptr
+		, nullptr
+		, InExtensionHook
+		, EMultiBlockType::Separator
+		, bInIsPartOfHeading)
+{
+	SetSearchable(false);
+}
+
+/**
+ * Allocates a widget for this type of MultiBlock.  Override this in derived classes.
+ *
+ * @return  MultiBlock widget object
+ */
+TSharedRef<class IMultiBlockBaseWidget> FMenuSpacerBlock::ConstructWidget() const
+{
+	return SNew(SMenuSpacerBlock);
+}
+
+/**
+ * Construct this widget
+ *
+ * @param	InArgs	The declaration data for this widget
+ */
+void SMenuSpacerBlock::Construct(const FArguments& InArgs)
+{
+}
+
+/**
+ * Builds this MultiBlock widget up from the MultiBlock associated with it
+ */
+void SMenuSpacerBlock::BuildMultiBlockWidget(const ISlateStyle* StyleSet, const FName& StyleName)
+{
+	// Unlike the separator, a spacer has no visual representation, so we halve the vertical padding to effectively maintain spacing between visual elements
+	static const FMargin BlockPadding = StyleSet->GetMargin(StyleName, ".Separator.Padding") * FMargin(1.0f, 0.5f);
+
+	ChildSlot
+	[
+		SNew(SVerticalBox)
+
+		+ SVerticalBox::Slot()
+		.AutoHeight()
+		.Padding(BlockPadding)
+		[
+			SNew(SSeparator)
+			.SeparatorImage(StyleSet->GetBrush(StyleName, ".Separator"))
+			.Visibility(EVisibility::Hidden)
+			.Thickness(1.0f)
+		]
+	];
+
+	// Add this widget to the search list of the multibox and hide it
+	OwnerMultiBoxWidget.Pin()->AddElement(this->AsWidget(), FText::GetEmpty(), MultiBlock->GetSearchable());
+
+	SetVisibility(MultiBlock->GetVisibilityOverride());
+}
diff --git a/Engine/Source/Runtime/Slate/Private/Framework/MultiBox/SMenuSpacerBlock.h b/Engine/Source/Runtime/Slate/Private/Framework/MultiBox/SMenuSpacerBlock.h
new file mode 100644
--- /dev/null
+++ b/Engine/Source/Runtime/Slate/Private/Framework/MultiBox/SMenuSpacerBlock.h
@@ -0,0 +1,53 @@
+// Copyright Epic Games, Inc. All Rights Reserved.
+
+#pragma once
+
+#include "CoreTypes.h"
+#include "Widgets/DeclarativeSyntaxSupport.h"
+#include "Framework/MultiBox/MultiBox.h"
+
+/**
+ * Menu spacer MultiBlock, a variant of the separator block that has no drawn representation, but consumes space in the menu
+ */
+class FMenuSpacerBlock
+	: public FMultiBlock
+{
+public:
+	/**
+	 * Constructor
+	 */
+	FMenuSpacerBlock(const FName& InExtensionHook, bool bInIsPartOfHeading);
+
+private:
+	/** FMultiBlock private interface */
+	virtual TSharedRef<class IMultiBlockBaseWidget> ConstructWidget() const override;
+
+private:
+	// Friend our corresponding widget class
+	friend class SMenuSpacerBlock;
+};
+
+
+/**
+ * Menu spacer MultiBlock widget
+ */
+class SMenuSpacerBlock
+	: public SMultiBlockBaseWidget
+{
+public:
+	SLATE_BEGIN_ARGS(SMenuSpacerBlock)
+	{}
+	SLATE_END_ARGS()
+
+	/**
+	 * Builds this MultiBlock widget up from the MultiBlock associated with it
+	 */
+	SLATE_API virtual void BuildMultiBlockWidget(const ISlateStyle* StyleSet, const FName& StyleName) override;
+
+	/**
+	 * Construct this widget
+	 *
+	 * @param	InArgs	The declaration data for this widget
+	 */
+	SLATE_API void Construct(const FArguments& InArgs);
+};
diff --git a/Engine/Source/Runtime/Slate/Public/Framework/MultiBox/MultiBox.h b/Engine/Source/Runtime/Slate/Public/Framework/MultiBox/MultiBox.h
--- a/Engine/Source/Runtime/Slate/Public/Framework/MultiBox/MultiBox.h
+++ b/Engine/Source/Runtime/Slate/Public/Framework/MultiBox/MultiBox.h
@@ -860,7 +860,7 @@ public:
 	SLATE_API double GetSummonedMenuTime() const;
 
 	UE_DEPRECATED_FORGAME(5.5, "This function will be made private in the future.")
-	bool ShouldShowMenuSearchField();
+	bool ShouldShowMenuSearchField() const;
 
 private:
 	/** Adds a block Widget to this widget */
@@ -874,6 +874,8 @@ private:
 	/** Creates the SearchTextWidget if the MultiBox has requested one */
 	SLATE_API void CreateSearchTextWidget();
 
+	SLATE_API EVisibility GetSearchBoxVisibility() const;
+
 	/** Called when the search field text is committed (such as pressing ESC). */
 	SLATE_API void OnFilterTextCommitted(const FText& InFilterText, ETextCommit::Type CommitType);
 	
