[STableViewBase] Prevent recalculating the TargetScrollOffset when the list is at the end, to keep the last row correctly visible

#rnx

diff --git a/Engine/Source/Runtime/Slate/Private/Widgets/Views/STableViewBase.cpp b/Engine/Source/Runtime/Slate/Private/Widgets/Views/STableViewBase.cpp
--- a/Engine/Source/Runtime/Slate/Private/Widgets/Views/STableViewBase.cpp
+++ b/Engine/Source/Runtime/Slate/Private/Widgets/Views/STableViewBase.cpp
@@ -823,8 +823,10 @@ STableViewBase::~STableViewBase() = default;
 
 double STableViewBase::GetTargetScrollOffset() const
 {
-	if (FixedLineScrollOffset.IsSet() && !IsRightClickScrolling() && InertialScrollManager.GetScrollVelocity() == 0.f)
+	if (!bWasAtEndOfList && FixedLineScrollOffset.IsSet() && !IsRightClickScrolling() && InertialScrollManager.GetScrollVelocity() == 0.f)
 	{
+		// The scroll has ended, we calculate the appropriate the scroll offset to snap on a line.
+		// We don't do that if we are at the end of the list because in this case showing the last line correctly is more important.
 		const int32 NumItemsPerLine = GetNumItemsPerLine();
 		const double DesiredLineOffset = FMath::FloorToDouble(DesiredScrollOffset / NumItemsPerLine) - FixedLineScrollOffset.GetValue();
 		
