STextScroller:  create an option to scroll text vertically
- helps handle some text overflow issues

diff --git a/Engine/Plugins/Runtime/CommonUI/Source/CommonUI/Private/CommonRichTextBlock.cpp b/Engine/Plugins/Runtime/CommonUI/Source/CommonUI/Private/CommonRichTextBlock.cpp
--- a/Engine/Plugins/Runtime/CommonUI/Source/CommonUI/Private/CommonRichTextBlock.cpp
+++ b/Engine/Plugins/Runtime/CommonUI/Source/CommonUI/Private/CommonRichTextBlock.cpp
@@ -492,6 +492,7 @@ TSharedRef<SWidget> UCommonRichTextBlock::RebuildWidget()
 	MyTextScroller =
 		SNew(STextScroller)
 		.ScrollOptions(TextScrollStyle->ToScrollOptions())
+		.ScrollOrientation(ScrollOrientation)
 		[
 			Super::RebuildWidget()
 		];
diff --git a/Engine/Plugins/Runtime/CommonUI/Source/CommonUI/Private/CommonTextBlock.cpp b/Engine/Plugins/Runtime/CommonUI/Source/CommonUI/Private/CommonTextBlock.cpp
--- a/Engine/Plugins/Runtime/CommonUI/Source/CommonUI/Private/CommonTextBlock.cpp
+++ b/Engine/Plugins/Runtime/CommonUI/Source/CommonUI/Private/CommonTextBlock.cpp
@@ -597,6 +597,7 @@ TSharedRef<SWidget> UCommonTextBlock::RebuildWidget()
 	TextScroller = 
 		SNew(STextScroller)
 		.ScrollOptions(TextScrollStyle->ToScrollOptions())
+		.ScrollOrientation(ScrollOrientation)
 		[ 
 			Super::RebuildWidget() 
 		];
diff --git a/Engine/Plugins/Runtime/CommonUI/Source/CommonUI/Public/CommonRichTextBlock.h b/Engine/Plugins/Runtime/CommonUI/Source/CommonUI/Public/CommonRichTextBlock.h
--- a/Engine/Plugins/Runtime/CommonUI/Source/CommonUI/Public/CommonRichTextBlock.h
+++ b/Engine/Plugins/Runtime/CommonUI/Source/CommonUI/Public/CommonRichTextBlock.h
@@ -83,6 +83,10 @@ private:
 	UPROPERTY(EditAnywhere, BlueprintReadOnly, Category = Appearance, meta = (ExposeOnSpawn = true, AllowPrivateAccess = true))
 	TSubclassOf<UCommonTextScrollStyle> ScrollStyle;
 
+	/** The orientation the text will scroll if out of bounds. */
+	UPROPERTY(EditAnywhere, BlueprintReadOnly, Category = Appearance, meta = (ExposeOnSpawn = true, AllowPrivateAccess = true))
+	TEnumAsByte<EOrientation> ScrollOrientation = Orient_Horizontal;
+
 	/** If scrolling is enabled/disabled initially, this can be updated in blueprint */
 	UPROPERTY(EditAnywhere, BlueprintReadOnly, Category = Appearance, meta = (ExposeOnSpawn = true, AllowPrivateAccess = true))
 	bool bIsScrollingEnabled = true;
diff --git a/Engine/Plugins/Runtime/CommonUI/Source/CommonUI/Public/CommonTextBlock.h b/Engine/Plugins/Runtime/CommonUI/Source/CommonUI/Public/CommonTextBlock.h
--- a/Engine/Plugins/Runtime/CommonUI/Source/CommonUI/Public/CommonTextBlock.h
+++ b/Engine/Plugins/Runtime/CommonUI/Source/CommonUI/Public/CommonTextBlock.h
@@ -204,6 +204,10 @@ private:
 	/** References the scroll style asset to use, no reference disables scrolling*/
 	UPROPERTY(EditAnywhere, BlueprintReadOnly, Category = CommonText, meta = (ExposeOnSpawn = true, AllowPrivateAccess = true))
 	TSubclassOf<UCommonTextScrollStyle> ScrollStyle;
+	
+	/** The orientation the text will scroll if out of bounds. */
+	UPROPERTY(EditAnywhere, BlueprintReadOnly, Category = CommonText, meta = (ExposeOnSpawn = true, AllowPrivateAccess = true))
+	TEnumAsByte<EOrientation> ScrollOrientation = Orient_Horizontal;
 
 #if WITH_EDITORONLY_DATA
 	/** Used to track widgets that were created before changing the default style pointer to null */
diff --git a/Engine/Source/Runtime/Slate/Private/Widgets/Text/STextScroller.cpp b/Engine/Source/Runtime/Slate/Private/Widgets/Text/STextScroller.cpp
--- a/Engine/Source/Runtime/Slate/Private/Widgets/Text/STextScroller.cpp
+++ b/Engine/Source/Runtime/Slate/Private/Widgets/Text/STextScroller.cpp
@@ -9,6 +9,7 @@
 void STextScroller::Construct(const FArguments& InArgs)
 {
 	ScrollOptions = InArgs._ScrollOptions;
+	ScrollOrientation = InArgs._ScrollOrientation;
 
 	// We don't tick, we use an active ticker.
 	SetCanTick(false);
@@ -26,8 +27,8 @@ void STextScroller::OnArrangeChildren(const FGeometry& AllottedGeometry, FArrang
 	{
 		const FVector2D ThisContentScale = GetContentScale();
 		const FMargin SlotPadding(LayoutPaddingWithFlow(GSlateFlowDirection, ChildSlot.GetPadding()));
-		const AlignmentArrangeResult XResult = AlignChild<Orient_Horizontal>(GSlateFlowDirection, AllottedGeometry.GetLocalSize().X, ChildSlot, SlotPadding, ThisContentScale.X, false);
-		const AlignmentArrangeResult YResult = AlignChild<Orient_Vertical>(AllottedGeometry.GetLocalSize().Y, ChildSlot, SlotPadding, ThisContentScale.Y);
+		const AlignmentArrangeResult XResult = AlignChild<Orient_Horizontal>(GSlateFlowDirection, AllottedGeometry.GetLocalSize().X, ChildSlot, SlotPadding, ThisContentScale.X, ScrollOrientation != Orient_Horizontal);
+		const AlignmentArrangeResult YResult = AlignChild<Orient_Vertical>(AllottedGeometry.GetLocalSize().Y, ChildSlot, SlotPadding, ThisContentScale.Y,  ScrollOrientation != Orient_Vertical);
 		const FVector2D MinimumSize = ChildSlot.GetWidget()->GetDesiredSize();
 
 		ArrangedChildren.AddWidget(ChildVisibility, AllottedGeometry.MakeChild(
@@ -40,13 +41,14 @@ void STextScroller::OnArrangeChildren(const FGeometry& AllottedGeometry, FArrang
 
 int32 STextScroller::OnPaint(const FPaintArgs& Args, const FGeometry& AllottedGeometry, const FSlateRect& MyCullingRect, FSlateWindowElementList& OutDrawElements, int32 LayerId, const FWidgetStyle& InWidgetStyle, bool bParentEnabled) const
 {
+	const int AxisIndex = ScrollOrientation == Orient_Vertical ? 1 : 0;
 	{
 		STextScroller* MutableThis = const_cast<STextScroller*>(this);
 
-		const float VisibleWidth = AllottedGeometry.GetLocalSize().X;
-		const float DesiredWidth = VisibleWidth == 0.0f ? VisibleWidth : ChildSlot.GetWidget()->GetDesiredSize().X;
+		const float VisibleGeometry = AllottedGeometry.GetLocalSize()[AxisIndex];
+		const float DesiredGeometry = VisibleGeometry == 0.0f ? VisibleGeometry : ChildSlot.GetWidget()->GetDesiredSize()[AxisIndex];
 
-		if (DesiredWidth > (VisibleWidth + 2.0f) && IsScrollingEnabled())
+		if (DesiredGeometry > (VisibleGeometry + 2.0f) && IsScrollingEnabled())
 		{
 			if (!ActiveTimerHandle.IsValid())
 			{
@@ -62,8 +64,10 @@ int32 STextScroller::OnPaint(const FPaintArgs& Args, const FGeometry& AllottedGe
 
 	if (ScrollOffset != 0)
 	{
-		const float ScrollDirection = GSlateFlowDirection == EFlowDirection::LeftToRight ? -1 : 1;
-		const FGeometry ScrolledGeometry = AllottedGeometry.MakeChild(FSlateRenderTransform(FVector2D(ScrollOffset * ScrollDirection, 0)));
+		const float ScrollDirection = GSlateFlowDirection == EFlowDirection::LeftToRight || ScrollOrientation == Orient_Vertical ? -1 : 1;
+		FVector2D GeometryOffset = FVector2D::ZeroVector;
+		GeometryOffset[AxisIndex] = ScrollOffset * ScrollDirection;
+		const FGeometry ScrolledGeometry = AllottedGeometry.MakeChild(FSlateRenderTransform(GeometryOffset));
 		return SCompoundWidget::OnPaint(Args, ScrolledGeometry, MyCullingRect, OutDrawElements, LayerId, InWidgetStyle, bParentEnabled);
 	}
 	else
@@ -105,18 +109,32 @@ EActiveTimerReturnType STextScroller::OnScrollTextTick(double CurrentTime, float
 	switch (TickerState)
 	{
 	case ETickerState::StartTicking:
-		// If we need to scroll, then it's imperative that we arrange from the left, rather than fill, so that we flip to right aligning and scrolling
-		// to the right (potentially).
-		SetFlowDirectionPreference(EFlowDirectionPreference::Culture);
-		ChildSlot.HAlign(HAlign_Left);
+		if (ScrollOrientation == Orient_Vertical)
+		{
+			ChildSlot.VAlign(VAlign_Top);
+		}
+		else
+		{
+			// If we need to scroll, then it's imperative that we arrange from the left, rather than fill, so that we flip to right aligning and scrolling
+			// to the right (potentially).
+			SetFlowDirectionPreference(EFlowDirectionPreference::Culture);
+			ChildSlot.HAlign(HAlign_Left);
+		}
 
 		TickerState = ETickerState::Ticking;
 		return EActiveTimerReturnType::Continue; // Defer to the next tick so that the layout is correct
 
 	case ETickerState::StopTicking:
-		// If we no longer need to scroll, just inherit the flow direction.
-		SetFlowDirectionPreference(EFlowDirectionPreference::Inherit);
-		ChildSlot.HAlign(HAlign_Fill);
+		if (ScrollOrientation == Orient_Vertical)
+		{
+			ChildSlot.VAlign(VAlign_Fill);
+		}
+		else
+		{
+			// If we no longer need to scroll, just inherit the flow direction.
+			SetFlowDirectionPreference(EFlowDirectionPreference::Inherit);
+			ChildSlot.HAlign(HAlign_Fill);
+		}
 
 		TickerState = ETickerState::None;
 		ResetScrollState();
@@ -126,8 +144,9 @@ EActiveTimerReturnType STextScroller::OnScrollTextTick(double CurrentTime, float
 	default:
 		break;
 	}
-
-	const float ContentSize = ChildSlot.GetWidget()->GetDesiredSize().X;
+	
+	const int AxisIndex = ScrollOrientation == Orient_Vertical ? 1 : 0;
+	const float ContentSize = ChildSlot.GetWidget()->GetDesiredSize()[AxisIndex];
 	TimeElapsed += DeltaTime;
 
 	switch (ActiveState)
@@ -164,7 +183,7 @@ EActiveTimerReturnType STextScroller::OnScrollTextTick(double CurrentTime, float
 	case EActiveState::Scrolling:
 	{
 		ScrollOffset += ScrollOptions.Speed * DeltaTime;
-		if ((ScrollOffset + GetCachedGeometry().GetLocalSize().X) >= ContentSize)
+		if ((ScrollOffset + GetCachedGeometry().GetLocalSize()[AxisIndex]) >= ContentSize)
 		{
 			TimeElapsed = 0.0f;
 			ActiveState = EActiveState::Stop;
diff --git a/Engine/Source/Runtime/Slate/Public/Widgets/Text/STextScroller.h b/Engine/Source/Runtime/Slate/Public/Widgets/Text/STextScroller.h
--- a/Engine/Source/Runtime/Slate/Public/Widgets/Text/STextScroller.h
+++ b/Engine/Source/Runtime/Slate/Public/Widgets/Text/STextScroller.h
@@ -30,11 +30,14 @@ class STextScroller : public SCompoundWidget
 {
 public:
 	SLATE_BEGIN_ARGS(STextScroller)
+		: _ScrollOrientation(Orient_Horizontal)
 	{
 		_Clipping = EWidgetClipping::OnDemand;
 	}
 		/** Options to apply when scrolling the text */
 		SLATE_ARGUMENT(FTextScrollerOptions, ScrollOptions)
+		/** Orientation this text will scroll if needed */
+		SLATE_ARGUMENT(EOrientation, ScrollOrientation)
 		/** Slot for the text widget to be scrolled */
 		SLATE_DEFAULT_SLOT(FArguments, Content)
 	SLATE_END_ARGS()
@@ -84,4 +87,5 @@ private:
 	TSharedPtr<FActiveTimerHandle> ActiveTimerHandle;
 
 	FTextScrollerOptions ScrollOptions;
+	EOrientation ScrollOrientation = Orient_Horizontal;
 };
