Integrating PR #11692, along with some additional fixes to realign text highlights based on outline/shadow size

PR #11692: Fixing excess padding coming from additional slate runs
#tests Updated FontOutline functional tests in EngineTest, verified that there were no regressions in other font rendering tests
[RN] Fixed various issues with Slate text runs that could cause unwanted spacing/alignment when using rich text blocks with multiple styles

#virtualized

diff --git a/Engine/Source/Runtime/Slate/Private/Framework/Text/SlateHyperlinkRun.cpp b/Engine/Source/Runtime/Slate/Private/Framework/Text/SlateHyperlinkRun.cpp
--- a/Engine/Source/Runtime/Slate/Private/Framework/Text/SlateHyperlinkRun.cpp
+++ b/Engine/Source/Runtime/Slate/Private/Framework/Text/SlateHyperlinkRun.cpp
@@ -46,15 +46,16 @@ int16 FSlateHyperlinkRun::GetMaxHeight( float Scale ) const
 
 FVector2D FSlateHyperlinkRun::Measure( int32 StartIndex, int32 EndIndex, float Scale, const FRunTextContext& TextContext ) const 
 {
-	const FVector2D ShadowOffsetToApply((EndIndex == Range.EndIndex) ? FMath::Abs(Style.TextStyle.ShadowOffset.X * Scale) : 0.0f, FMath::Abs(Style.TextStyle.ShadowOffset.Y * Scale));
-
-	if ( EndIndex - StartIndex == 0 )
+	const FVector2d OutlineSize = GetOutlineSize(StartIndex, EndIndex, Scale);
+	const FVector2d ShadowSize = GetShadowSize(StartIndex, EndIndex, Scale);
+	
+	if (EndIndex - StartIndex == 0)
 	{
-		return FVector2D( ShadowOffsetToApply.X * Scale, GetMaxHeight( Scale ) );
+		return FVector2D(0, GetMaxHeight(Scale)) + OutlineSize + ShadowSize;
 	}
-
+	
 	// Use the full text range (rather than the run range) so that text that spans runs will still be shaped correctly
-	return ShapedTextCacheUtil::MeasureShapedText(TextContext.ShapedTextCache, FCachedShapedTextKey(FTextRange(0, Text->Len()), Scale, TextContext, Style.TextStyle.Font), FTextRange(StartIndex, EndIndex), **Text) + ShadowOffsetToApply;
+	return ShapedTextCacheUtil::MeasureShapedText(TextContext.ShapedTextCache, FCachedShapedTextKey(FTextRange(0, Text->Len()), Scale, TextContext, Style.TextStyle.Font), FTextRange(StartIndex, EndIndex), **Text) + OutlineSize + ShadowSize;
 }
 
 int8 FSlateHyperlinkRun::GetKerning( int32 CurrentIndex, float Scale, const FRunTextContext& TextContext ) const 
@@ -62,6 +63,20 @@ int8 FSlateHyperlinkRun::GetKerning( int32 CurrentIndex, float Scale, const FRun
 	return 0;
 }
 
+FVector2d FSlateHyperlinkRun::GetOutlineSize(int32 StartIndex, int32 EndIndex, float Scale) const
+{
+	// Offset the measured shaped text by the outline since the outline was not factored into the size of the text
+	// Need to add the outline offsetting to the beginning and the end because it surrounds both sides.
+	const float ScaledOutlineSize = Style.TextStyle.Font.OutlineSettings.OutlineSize * Scale;
+
+	return FVector2D((StartIndex == Range.BeginIndex ? ScaledOutlineSize : 0) + (EndIndex == Range.EndIndex ? ScaledOutlineSize : 0), ScaledOutlineSize);
+}
+
+FVector2d FSlateHyperlinkRun::GetShadowSize(int32 StartIndex, int32 EndIndex, float Scale) const
+{
+	return FVector2d((EndIndex == Range.EndIndex) ? FMath::Abs(Style.TextStyle.ShadowOffset.X * Scale) : 0.0f, FMath::Abs(Style.TextStyle.ShadowOffset.Y * Scale));
+}
+
 TSharedRef< ILayoutBlock > FSlateHyperlinkRun::CreateBlock( int32 StartIndex, int32 EndIndex, FVector2D Size, const FLayoutBlockTextContext& TextContext, const TSharedPtr< IRunRenderer >& Renderer )
 {
 	FText ToolTipText;
diff --git a/Engine/Source/Runtime/Slate/Private/Framework/Text/SlateTextRun.cpp b/Engine/Source/Runtime/Slate/Private/Framework/Text/SlateTextRun.cpp
--- a/Engine/Source/Runtime/Slate/Private/Framework/Text/SlateTextRun.cpp
+++ b/Engine/Source/Runtime/Slate/Private/Framework/Text/SlateTextRun.cpp
@@ -44,20 +44,16 @@ int16 FSlateTextRun::GetMaxHeight( float Scale ) const
 
 FVector2D FSlateTextRun::Measure( int32 BeginIndex, int32 EndIndex, float Scale, const FRunTextContext& TextContext ) const 
 {
-	const FVector2D ShadowOffsetToApply((EndIndex == Range.EndIndex) ? FMath::Abs(Style.ShadowOffset.X * Scale) : 0.0f, FMath::Abs(Style.ShadowOffset.Y * Scale));
-
-	// Offset the measured shaped text by the outline since the outline was not factored into the size of the text
-	// Need to add the outline offsetting to the beginning and the end because it surrounds both sides.
-	const float ScaledOutlineSize = Style.Font.OutlineSettings.OutlineSize * Scale;
-	const FVector2D OutlineSizeToApply((BeginIndex == Range.BeginIndex ? ScaledOutlineSize : 0) + (EndIndex == Range.EndIndex ? ScaledOutlineSize : 0), ScaledOutlineSize);
-
+	const FVector2d OutlineSize = GetOutlineSize(BeginIndex, EndIndex, Scale);
+	const FVector2d ShadowSize = GetShadowSize(BeginIndex, EndIndex, Scale);
+	
 	if (EndIndex - BeginIndex == 0)
 	{
-		return FVector2D(0, GetMaxHeight(Scale)) + ShadowOffsetToApply + OutlineSizeToApply;
+		return FVector2D(0, GetMaxHeight(Scale) + OutlineSize.Y + ShadowSize.Y) + OutlineSize + ShadowSize;
 	}
-
+	
 	// Use the full text range (rather than the run range) so that text that spans runs will still be shaped correctly
-	return ShapedTextCacheUtil::MeasureShapedText(TextContext.ShapedTextCache, FCachedShapedTextKey(FTextRange(0, Text->Len()), Scale, TextContext, Style.Font), FTextRange(BeginIndex, EndIndex), **Text) + ShadowOffsetToApply + OutlineSizeToApply;
+	return ShapedTextCacheUtil::MeasureShapedText(TextContext.ShapedTextCache, FCachedShapedTextKey(FTextRange(0, Text->Len()), Scale, TextContext, Style.Font), FTextRange(BeginIndex, EndIndex), **Text) + OutlineSize + ShadowSize;
 }
 
 int8 FSlateTextRun::GetKerning(int32 CurrentIndex, float Scale, const FRunTextContext& TextContext) const
@@ -72,6 +68,27 @@ int8 FSlateTextRun::GetKerning(int32 CurrentIndex, float Scale, const FRunTextCo
 	return ShapedTextCacheUtil::GetShapedGlyphKerning(TextContext.ShapedTextCache, FCachedShapedTextKey(FTextRange(0, Text->Len()), Scale, TextContext, Style.Font), PreviousIndex, **Text);
 }
 
+FVector2d FSlateTextRun::GetOutlineSize(int32 StartIndex, int32 EndIndex, float Scale) const
+{
+	// Offset the measured shaped text by the outline since the outline was not factored into the size of the text
+	// Need to add the outline offsetting to the beginning and the end because it surrounds both sides.
+	const float ScaledOutlineSize = Style.Font.OutlineSettings.OutlineSize * Scale;
+
+	return FVector2D((StartIndex == Range.BeginIndex ? ScaledOutlineSize : 0) + (EndIndex == Range.EndIndex ? ScaledOutlineSize : 0), ScaledOutlineSize);
+}
+
+FVector2d FSlateTextRun::GetShadowSize(int32 StartIndex, int32 EndIndex, float Scale) const
+{
+	// For positive values of ShadowOffset, add the shadow's width if we're measuring the end of the text run
+	// For negative offsets, add if we're measuring the beginning
+	if ((Style.ShadowOffset.X > 0 && (EndIndex == Range.EndIndex)) || (Style.ShadowOffset.X < 0 && (StartIndex == Range.BeginIndex)))
+	{
+		return FVector2d(FMath::Abs(Style.ShadowOffset.X * Scale), FMath::Abs(Style.ShadowOffset.Y * Scale));
+	}
+
+	return FVector2d(0.0f, FMath::Abs(Style.ShadowOffset.Y * Scale));
+}
+
 TSharedRef< ILayoutBlock > FSlateTextRun::CreateBlock( int32 BeginIndex, int32 EndIndex, FVector2D Size, const FLayoutBlockTextContext& TextContext, const TSharedPtr< IRunRenderer >& Renderer )
 {
 	return FDefaultLayoutBlock::Create( SharedThis( this ), FTextRange( BeginIndex, EndIndex ), Size, TextContext, Renderer );
diff --git a/Engine/Source/Runtime/Slate/Private/Framework/Text/TextLayout.cpp b/Engine/Source/Runtime/Slate/Private/Framework/Text/TextLayout.cpp
--- a/Engine/Source/Runtime/Slate/Private/Framework/Text/TextLayout.cpp
+++ b/Engine/Source/Runtime/Slate/Private/Framework/Text/TextLayout.cpp
@@ -411,24 +411,42 @@ void FTextLayout::CreateLineViewBlocks( int32 LineModelIndex, const int32 StopIn
 			OutSoftLine = MoveTemp(VisualSoftLine);
 		}
 
+		float MaxLineWidth = 0.f;
+		float PrevBlockOutlineWidth = 0.f;
 		float CurrentHorizontalPos = 0.0f;
 		for (int32 Index = 0; Index < OutSoftLine.Num(); Index++)
 		{
 			const TSharedRef< ILayoutBlock > Block = OutSoftLine[ Index ];
 			const TSharedRef< IRun > Run = Block->GetRun();
 
+			const FTextRange BlockTextRange = Block->GetTextRange();
+
 			const int16 BlockBaseline = Run->GetBaseLine(Scale);
 			const int16 VerticalOffset = MaxAboveBaseline - Block->GetSize().Y - BlockBaseline;
-			const int8 BlockKerning = Run->GetKerning(Block->GetTextRange().BeginIndex, Scale, RunTextContext);
+			const int8 BlockKerning = Run->GetKerning(BlockTextRange.BeginIndex, Scale, RunTextContext);
 
+			const FVector2d BlockShadowSize = Run->GetShadowSize(BlockTextRange.BeginIndex, BlockTextRange.EndIndex, Scale);
+			const FVector2d BlockOutlineSize = Run->GetOutlineSize(BlockTextRange.BeginIndex, BlockTextRange.EndIndex, Scale);
+		
+			if (Index > 0)
+			{
+				// Scooch adjacent blocks with different outline sizes as needed to keep the actual glyph spacing correct 
+				const float PrevBlockOutlineAdjustment = (PrevBlockOutlineWidth - BlockOutlineSize.X) * 0.5f;
+				CurrentHorizontalPos += PrevBlockOutlineAdjustment;
+			}
 			Block->SetLocationOffset(FVector2D(CurrentOffset.X + CurrentHorizontalPos + BlockKerning, CurrentOffset.Y + VerticalOffset));
 
-			CurrentHorizontalPos += Block->GetSize().X;
+			// Track what the final width would be as each block is added, on the off-chance that one has
+			// such a dramatic outline or shadow that it extends beyond the extent of the final block
+			MaxLineWidth = FMath::Max(MaxLineWidth, CurrentHorizontalPos + Block->GetSize().X);
+			CurrentHorizontalPos += Block->GetSize().X - BlockShadowSize.X - BlockOutlineSize.X;
+
+			PrevBlockOutlineWidth = BlockOutlineSize.X;
 		}
 
 		UnscaleLineHeight = MaxAboveBaseline + MaxBelowBaseline;
 
-		LineSize.X = CurrentHorizontalPos;
+		LineSize.X = MaxLineWidth;
 		LineSize.Y = UnscaleLineHeight * LineHeightPercentage;
 
 		FTextLayout::FLineView LineView;
@@ -838,14 +856,27 @@ void FTextLayout::FlowHighlights()
 				LineViewHighlight.Width = 0.0f;
 			};
 
+			float PrevBlockOutlineWidth = 0.f;
 			for (int32 CurrentBlockIndex = 0; CurrentBlockIndex < LineView.Blocks.Num(); ++CurrentBlockIndex)
 			{
 				const TSharedRef<ILayoutBlock>& Block = LineView.Blocks[CurrentBlockIndex];
 				const FTextRange& BlockTextRange = Block->GetTextRange();
 				const TSharedRef<IRun> Run = Block->GetRun();
 
+				const FVector2d BlockShadowSize = Run->GetShadowSize(BlockTextRange.BeginIndex, BlockTextRange.EndIndex, Scale);
+				const FVector2d BlockOutlineSize = Run->GetOutlineSize(BlockTextRange.BeginIndex, BlockTextRange.EndIndex, Scale);
+				
+				if (CurrentBlockIndex > 0)
+				{
+					// Scooch adjacent blocks with different outline sizes as needed to keep the actual glyph spacing correct 
+					const float PrevBlockOutlineAdjustment = (PrevBlockOutlineWidth - BlockOutlineSize.X) * 0.5f;
+					RunningBlockOffset += PrevBlockOutlineAdjustment;
+				}
+				PrevBlockOutlineWidth = BlockOutlineSize.X;
+
 				const float CurrentBlockRunningOffset = RunningBlockOffset;
-				RunningBlockOffset += Block->GetSize().X;
+
+				RunningBlockOffset += Block->GetSize().X - BlockShadowSize.X - BlockOutlineSize.X;
 
 				const FTextRange IntersectedRange = BlockTextRange.Intersect(LineViewHighlightRange);
 				const bool bBlockIsHighlighted = !IntersectedRange.IsEmpty() || IntersectedRange == LineViewHighlightRange;
@@ -871,11 +902,17 @@ void FTextLayout::FlowHighlights()
 							// In RTL text, the space after the end of the text (which is visually on the left) is added as an offset
 							if (Block->GetTextContext().TextDirection == TextBiDi::ETextDirection::LeftToRight)
 							{
-								LineViewHighlight.Offset.X += Run->Measure(BlockTextRange.BeginIndex, IntersectedRange.BeginIndex, Scale, RunTextContext).X;
+								const FVector2d ShadowSize = Run->GetShadowSize(BlockTextRange.BeginIndex, IntersectedRange.BeginIndex, Scale);
+								const FVector2d OutlineSize = Run->GetOutlineSize(BlockTextRange.BeginIndex, IntersectedRange.BeginIndex, Scale);
+
+								LineViewHighlight.Offset.X += Run->Measure(BlockTextRange.BeginIndex, IntersectedRange.BeginIndex, Scale, RunTextContext).X - ShadowSize.X - OutlineSize.X;
 							}
 							else
 							{
-								LineViewHighlight.Offset.X += Run->Measure(IntersectedRange.EndIndex, BlockTextRange.EndIndex, Scale, RunTextContext).X;
+								const FVector2d ShadowSize = Run->GetShadowSize(IntersectedRange.EndIndex, BlockTextRange.EndIndex, Scale);
+								const FVector2d OutlineSize = Run->GetOutlineSize(IntersectedRange.EndIndex, BlockTextRange.EndIndex, Scale);
+
+								LineViewHighlight.Offset.X += Run->Measure(IntersectedRange.EndIndex, BlockTextRange.EndIndex, Scale, RunTextContext).X - ShadowSize.X - OutlineSize.X;
 							}
 						}
 					}
diff --git a/Engine/Source/Runtime/Slate/Public/Framework/Text/IRun.h b/Engine/Source/Runtime/Slate/Public/Framework/Text/IRun.h
--- a/Engine/Source/Runtime/Slate/Public/Framework/Text/IRun.h
+++ b/Engine/Source/Runtime/Slate/Public/Framework/Text/IRun.h
@@ -138,6 +138,8 @@ public:
 
 	virtual FVector2D Measure( int32 StartIndex, int32 EndIndex, float Scale, const FRunTextContext& TextContext) const = 0;
 	virtual int8 GetKerning(int32 CurrentIndex, float Scale, const FRunTextContext& TextContext) const = 0;
+	virtual FVector2D GetOutlineSize(int32 StartIndex, int32 EndIndex, float Scale) const { return FVector2d::ZeroVector; }
+	virtual FVector2D GetShadowSize(int32 StartIndex, int32 EndIndex, float Scale) const { return FVector2d::ZeroVector; }
 
 	virtual TSharedRef< class ILayoutBlock > CreateBlock( int32 StartIndex, int32 EndIndex, FVector2D Size, const FLayoutBlockTextContext& TextContext, const TSharedPtr< class IRunRenderer >& Renderer ) = 0;
 
diff --git a/Engine/Source/Runtime/Slate/Public/Framework/Text/SlateHyperlinkRun.h b/Engine/Source/Runtime/Slate/Public/Framework/Text/SlateHyperlinkRun.h
--- a/Engine/Source/Runtime/Slate/Public/Framework/Text/SlateHyperlinkRun.h
+++ b/Engine/Source/Runtime/Slate/Public/Framework/Text/SlateHyperlinkRun.h
@@ -68,6 +68,10 @@ public:
 
 	SLATE_API virtual int8 GetKerning(int32 CurrentIndex, float Scale, const FRunTextContext& TextContext) const override;
 
+	SLATE_API virtual FVector2d GetOutlineSize(int32 StartIndex, int32 EndIndex, float Scale) const override;
+	
+	SLATE_API virtual FVector2d GetShadowSize(int32 StartIndex, int32 EndIndex, float Scale) const override;
+
 	SLATE_API virtual TSharedRef< ILayoutBlock > CreateBlock( int32 StartIndex, int32 EndIndex, FVector2D Size, const FLayoutBlockTextContext& TextContext, const TSharedPtr< IRunRenderer >& Renderer ) override;
 
 	SLATE_API virtual int32 OnPaint(const FPaintArgs& PaintArgs, const FTextArgs& TextArgs, const FGeometry& AllottedGeometry, const FSlateRect& MyCullingRect, FSlateWindowElementList& OutDrawElements, int32 LayerId, const FWidgetStyle& InWidgetStyle, bool bParentEnabled) const override;
diff --git a/Engine/Source/Runtime/Slate/Public/Framework/Text/SlateTextRun.h b/Engine/Source/Runtime/Slate/Public/Framework/Text/SlateTextRun.h
--- a/Engine/Source/Runtime/Slate/Public/Framework/Text/SlateTextRun.h
+++ b/Engine/Source/Runtime/Slate/Public/Framework/Text/SlateTextRun.h
@@ -33,6 +33,8 @@ public:
 	SLATE_API virtual int16 GetMaxHeight( float Scale ) const override;
 	SLATE_API virtual FVector2D Measure( int32 StartIndex, int32 EndIndex, float Scale, const FRunTextContext& TextContext ) const override;
 	SLATE_API virtual int8 GetKerning(int32 CurrentIndex, float Scale, const FRunTextContext& TextContext) const override;
+	SLATE_API virtual FVector2d GetOutlineSize(int32 StartIndex, int32 EndIndex, float Scale) const override;
+	SLATE_API virtual FVector2d GetShadowSize(int32 StartIndex, int32 EndIndex, float Scale) const override;
 
 	SLATE_API virtual TSharedRef< ILayoutBlock > CreateBlock( int32 StartIndex, int32 EndIndex, FVector2D Size, const FLayoutBlockTextContext& TextContext, const TSharedPtr< IRunRenderer >& Renderer ) override;
 
