Slate: Copy comment of private SInlineEditableTextBlock::IsSelected to the SLATE_EVENT counterpart since it better informs on when it should be used

#rnx

diff --git a/Engine/Source/Runtime/Slate/Public/Widgets/Text/SInlineEditableTextBlock.h b/Engine/Source/Runtime/Slate/Public/Widgets/Text/SInlineEditableTextBlock.h
--- a/Engine/Source/Runtime/Slate/Public/Widgets/Text/SInlineEditableTextBlock.h
+++ b/Engine/Source/Runtime/Slate/Public/Widgets/Text/SInlineEditableTextBlock.h
@@ -111,7 +111,9 @@ class SInlineEditableTextBlock: public SCompoundWidget
 		/** Callback when the text editing ends. */
 		SLATE_EVENT(FSimpleDelegate, OnExitEditingMode)
 
-		/** Callback to check if the widget is selected, should only be hooked up if parent widget is handling selection or focus. */
+		/** Delegate to execute to check the status of if the widget is selected or not
+			Only needs to be hooked up if an external widget is managing selection, such
+			as a SListView or STreeView. */
 		SLATE_EVENT( FIsSelected, IsSelected )
 
 		/** Called whenever the text is changed programmatically or interactively by the user */
