[STileView] Fixed some calculations in ScrollIntoView:
- when FixedLineOffset is ON and LineOfItem == MaxDisplayedLine we now scroll to the next line instead of just trying to scroll the tile into view
- when FixedLineOffset is ON and LineOfItem == MaxDisplayedLine + 1 we now scroll to align a line with the top of the view
- when LineOfItem == MinDisplayedLine - 1 we now scroll to put the tile on the first line instead of trying to center it
- tile centering fixed (it was centering the top of the tile)

[RN] Fixed STileView::ScrollIntoView to correctly center an element that was out of the view, to scroll less if the element was partially into view and to take FixedLineOffset into account in all situations

diff --git a/Engine/Source/Runtime/Slate/Public/Widgets/Views/STileView.h b/Engine/Source/Runtime/Slate/Public/Widgets/Views/STileView.h
--- a/Engine/Source/Runtime/Slate/Public/Widgets/Views/STileView.h
+++ b/Engine/Source/Runtime/Slate/Public/Widgets/Views/STileView.h
@@ -468,10 +468,41 @@ protected:
 
 				if (LineOfItem < MinDisplayedLine || LineOfItem > MaxDisplayedLine)
 				{
-					// Set the line with the item at the beginning of the view area
-					float NewLineOffset = (float)LineOfItem;
-					// Center the line in the view area
-					NewLineOffset -= NumLinesInView * 0.5f;
+					float NewLineOffset = 0.0f;
+					if (LineOfItem < MinDisplayedLine - 1 || LineOfItem > MaxDisplayedLine + 1)
+					{
+						// The item was completely hidden so we scroll to center it.
+						
+						// Set the line with the item at the beginning of the view area
+						NewLineOffset = (float)LineOfItem;
+						// Center the line in the view area
+						NewLineOffset -= NumLinesInView * 0.5f;
+						// Center the tile itself
+						NewLineOffset += 0.5f;
+					}
+					else if (LineOfItem == MinDisplayedLine - 1)
+					{
+						// The item was just on the line above, set the line with the item at the beginning of the view area
+						// This prevent issues when navigating the list element by element triggers this code since we don't want to suddenly center an element
+						NewLineOffset = (float)LineOfItem;
+					}
+					else // LineOfItem == MaxDisplayedLine + 1
+					{
+						// The item was just on the next line, scroll it into view
+						// This prevent issues when navigating the list element by element triggers this code since we don't want to suddenly center an element
+						if (this->FixedLineScrollOffset.IsSet())
+						{
+							// We should scroll to align a line with the top of the view while putting the item at the bottom
+							const int32 FullLinesInView = FMath::FloorToInt32(NumLinesInView);
+							NewLineOffset = (float)LineOfItem - ((float)FullLinesInView - 1.f);
+						}
+						else
+						{
+							// We should scroll just enough to show the item on the last line
+							NewLineOffset = (float)LineOfItem - NumLinesInView + 1.f + this->NavigationScrollOffset;					
+						}
+					}
+					
 					// Convert the line offset into an item offset
 					double NewScrollOffset = NewLineOffset * (double)NumItemsPerLine;
 					// And clamp the scroll offset within the allowed limits
@@ -489,10 +520,20 @@ protected:
 					}
 					else if (LineOfItem == MaxDisplayedLine)
 					{
-						// This line is clipped at the end, so we need to advance just enough to bring it fully into view
-						// Since all tiles are required to be of the same size, this is straightforward
-						const float NewLineOffset = (float)LineOfItem - NumLinesInView + 1.f + (this->FixedLineScrollOffset.IsSet() ? 0.f : this->NavigationScrollOffset);
-						this->SetScrollOffset(NewLineOffset * (float)NumItemsPerLine);
+						if (this->FixedLineScrollOffset.IsSet())
+						{
+							// This line is clipped at the end but since FixedLineScrollOffset is set we must align a line to the top of the list,
+							// therefore we just scroll to the next line
+							const float HiddenLinesCount = (float)LineOfItem - FMath::FloorToFloat(NumLinesInView); 
+							this->SetScrollOffset((1.0f + HiddenLinesCount) * (float)NumItemsPerLine);							
+						}
+						else
+						{
+							// This line is clipped at the end, so we need to advance just enough to bring it fully into view
+							// Since all tiles are required to be of the same size, this is straightforward
+							const float NewLineOffset = (float)LineOfItem - NumLinesInView + 1.f + this->NavigationScrollOffset;
+							this->SetScrollOffset(NewLineOffset * (float)NumItemsPerLine);
+						}
 					}
 				}
 
