Modify maximize/restore button so that it drops you out of windowed-fullscreen.
Add enter & exit fullscreen button icons.
Update close icon to better match overall style.

diff --git a/Engine/Source/Runtime/Slate/Private/Framework/Application/SWindowTitleBar.cpp b/Engine/Source/Runtime/Slate/Private/Framework/Application/SWindowTitleBar.cpp
--- a/Engine/Source/Runtime/Slate/Private/Framework/Application/SWindowTitleBar.cpp
+++ b/Engine/Source/Runtime/Slate/Private/Framework/Application/SWindowTitleBar.cpp
@@ -429,7 +429,11 @@ FReply SWindowTitleBar::MaximizeRestoreButton_OnClicked( )
 
 		if (NativeWindow.IsValid())
 		{
-			if (NativeWindow->IsMaximized())
+			if (NativeWindow->GetWindowMode() == EWindowMode::WindowedFullscreen)
+			{
+				OwnerWindow->SetWindowMode(EWindowMode::Windowed);
+			}
+			else if (NativeWindow->IsMaximized())
 			{
 				NativeWindow->Restore();
 			}
@@ -454,44 +458,49 @@ const FSlateBrush* SWindowTitleBar::GetMaximizeRestoreImage( ) const
 
 	TSharedPtr<FGenericWindow> NativeWindow = OwnerWindow->GetNativeWindow();
 
-	if (NativeWindow.IsValid() && NativeWindow->IsMaximized())
+	if (!OwnerWindow->HasMaximizeBox())
 	{
-		if (!OwnerWindow->HasMaximizeBox())
-		{
-			return &Style->MaximizeButtonStyle.Disabled;
-		}
-		else if (MaximizeRestoreButton->IsPressed())
-		{
-			return &Style->RestoreButtonStyle.Pressed;
-		}
-		else if (MaximizeRestoreButton->IsHovered())
+		return &Style->MaximizeButtonStyle.Disabled;
+	}
+	
+	if (NativeWindow.IsValid())
+	{
+		if (NativeWindow->GetWindowMode() == EWindowMode::WindowedFullscreen)
 		{
-			return &Style->RestoreButtonStyle.Hovered;
+			if (MaximizeRestoreButton->IsPressed())
+			{
+				return &Style->ExitFullscreenButtonStyle.Pressed;
+			}
+			if (MaximizeRestoreButton->IsHovered())
+			{
+				return &Style->ExitFullscreenButtonStyle.Hovered;
+			}
+			return &Style->ExitFullscreenButtonStyle.Normal;
 		}
-		else
+		
+		if (NativeWindow->IsMaximized())
 		{
+			if (MaximizeRestoreButton->IsPressed())
+			{
+				return &Style->RestoreButtonStyle.Pressed;
+			}
+			if (MaximizeRestoreButton->IsHovered())
+			{
+				return &Style->RestoreButtonStyle.Hovered;
+			}
 			return &Style->RestoreButtonStyle.Normal;
 		}
 	}
-	else
+
+	if (MaximizeRestoreButton->IsPressed())
 	{
-		if (!OwnerWindow->HasMaximizeBox())
-		{
-			return &Style->MaximizeButtonStyle.Disabled;
-		}
-		else if (MaximizeRestoreButton->IsPressed())
-		{
-			return &Style->MaximizeButtonStyle.Pressed;
-		}
-		else if (MaximizeRestoreButton->IsHovered())
-		{
-			return &Style->MaximizeButtonStyle.Hovered;
-		}
-		else
-		{
-			return &Style->MaximizeButtonStyle.Normal;
-		}
+		return &Style->MaximizeButtonStyle.Pressed;
+	}
+	if (MaximizeRestoreButton->IsHovered())
+	{
+		return &Style->MaximizeButtonStyle.Hovered;
 	}
+	return &Style->MaximizeButtonStyle.Normal;
 }
 
 FReply SWindowTitleBar::MinimizeButton_OnClicked( )
@@ -588,9 +597,20 @@ FText SWindowTitleBar::GetMaximizeButtonToolTip() const
 {
 	TSharedPtr<SWindow> OwnerWindow = OwnerWindowPtr.Pin();
 
-	if (OwnerWindow.IsValid() && OwnerWindow->IsWindowMaximized())
+	if (OwnerWindow.IsValid())
 	{
-		return NSLOCTEXT("WindowTitleBar", "MaximizeRestoreDown", "Restore Down");
+		TSharedPtr<FGenericWindow> NativeWindow = OwnerWindow->GetNativeWindow();
+		if (NativeWindow.IsValid())
+		{
+			if (NativeWindow->IsMaximized())
+			{
+				return NSLOCTEXT("WindowTitleBar", "MaximizeRestoreDown", "Restore Down");
+			}
+			if (NativeWindow->GetWindowMode() == EWindowMode::WindowedFullscreen)
+			{
+				return NSLOCTEXT("WindowTitleBar", "MaximizeExitWindowedFullscreen", "Exit Full Screen");
+			}
+		}
 	}
 
 	return NSLOCTEXT("WindowTitleBar", "Maximize", "Maximize");
diff --git a/Engine/Source/Runtime/SlateCore/Private/Styling/StarshipCoreStyle.cpp b/Engine/Source/Runtime/SlateCore/Private/Styling/StarshipCoreStyle.cpp
--- a/Engine/Source/Runtime/SlateCore/Private/Styling/StarshipCoreStyle.cpp
+++ b/Engine/Source/Runtime/SlateCore/Private/Styling/StarshipCoreStyle.cpp
@@ -762,6 +762,16 @@ TSharedRef<ISlateStyle> FStarshipCoreStyle::Create()
 			.SetHovered(IMAGE_BRUSH_SVG("Starship/CoreWidgets/Window/restore", FVector2f(42.0f, 23.0f), FStyleColors::ForegroundHover))
 			.SetPressed(IMAGE_BRUSH_SVG("Starship/CoreWidgets/Window/restore", FVector2f(42.0f, 23.0f), FStyleColors::Foreground));
 
+		const FButtonStyle EnterFullscreenButtonStyle = FButtonStyle(NoBorder)
+			.SetNormal( IMAGE_BRUSH_SVG("Starship/CoreWidgets/Window/enter_fullscreen", FVector2f(42.0f, 23.0f), FStyleColors::Foreground))
+			.SetHovered(IMAGE_BRUSH_SVG("Starship/CoreWidgets/Window/enter_fullscreen", FVector2f(42.0f, 23.0f), FStyleColors::ForegroundHover))
+			.SetPressed(IMAGE_BRUSH_SVG("Starship/CoreWidgets/Window/enter_fullscreen", FVector2f(42.0f, 23.0f), FStyleColors::Foreground));
+
+		const FButtonStyle ExitFullscreenButtonStyle = FButtonStyle(NoBorder)
+			.SetNormal( IMAGE_BRUSH_SVG("Starship/CoreWidgets/Window/exit_fullscreen", FVector2f(42.0f, 23.0f), FStyleColors::Foreground))
+			.SetHovered(IMAGE_BRUSH_SVG("Starship/CoreWidgets/Window/exit_fullscreen", FVector2f(42.0f, 23.0f), FStyleColors::ForegroundHover))
+			.SetPressed(IMAGE_BRUSH_SVG("Starship/CoreWidgets/Window/exit_fullscreen", FVector2f(42.0f, 23.0f), FStyleColors::Foreground));
+
 		const FButtonStyle CloseButtonStyle = FButtonStyle(NoBorder)
 			.SetNormal( IMAGE_BRUSH_SVG("Starship/CoreWidgets/Window/close", FVector2f(42.0f, 23.0f), FStyleColors::Foreground))
 			.SetHovered(IMAGE_BRUSH_SVG("Starship/CoreWidgets/Window/close", FVector2f(42.0f, 23.0f), FStyleColors::ForegroundHover))
@@ -773,6 +783,8 @@ TSharedRef<ISlateStyle> FStarshipCoreStyle::Create()
 			.SetMinimizeButtonStyle(MinimizeButtonStyle)
 			.SetMaximizeButtonStyle(MaximizeButtonStyle)
 			.SetRestoreButtonStyle(RestoreButtonStyle)
+			.SetEnterFullscreenButtonStyle(EnterFullscreenButtonStyle)
+			.SetExitFullscreenButtonStyle(ExitFullscreenButtonStyle)
 			.SetCloseButtonStyle(CloseButtonStyle)
 #endif
 			.SetTitleTextStyle(NormalText)
diff --git a/Engine/Source/Runtime/SlateCore/Public/Styling/SlateTypes.h b/Engine/Source/Runtime/SlateCore/Public/Styling/SlateTypes.h
--- a/Engine/Source/Runtime/SlateCore/Public/Styling/SlateTypes.h
+++ b/Engine/Source/Runtime/SlateCore/Public/Styling/SlateTypes.h
@@ -2156,6 +2156,16 @@ struct FWindowStyle : public FSlateWidgetStyle
 	FButtonStyle RestoreButtonStyle;
 	FWindowStyle& SetRestoreButtonStyle( const FButtonStyle& InRestoreButtonStyle ){ RestoreButtonStyle = InRestoreButtonStyle; return *this; }
 
+	/** Style used to draw the enter fullscreen button */
+    UPROPERTY(EditAnywhere, BlueprintReadWrite, Category=Appearance)
+    FButtonStyle EnterFullscreenButtonStyle;
+    FWindowStyle& SetEnterFullscreenButtonStyle( const FButtonStyle& InEnterFullscreenButtonStyle ){ EnterFullscreenButtonStyle = InEnterFullscreenButtonStyle; return *this; }
+
+	/** Style used to draw the exit fullscreen button */
+	UPROPERTY(EditAnywhere, BlueprintReadWrite, Category=Appearance)
+	FButtonStyle ExitFullscreenButtonStyle;
+	FWindowStyle& SetExitFullscreenButtonStyle( const FButtonStyle& InExitFullscreenButtonStyle ){ ExitFullscreenButtonStyle = InExitFullscreenButtonStyle; return *this; }
+	
 	/** Style used to draw the window close button */
 	UPROPERTY(EditAnywhere, BlueprintReadWrite, Category=Appearance)
 	FButtonStyle CloseButtonStyle;
