Add a warning when a value is displayed as zero, but not actually zero. May occur with externally authored content.

#tests verified only shows on denorms / near integers

diff --git a/Engine/Source/Runtime/Core/Private/Internationalization/FastDecimalFormat.cpp b/Engine/Source/Runtime/Core/Private/Internationalization/FastDecimalFormat.cpp
--- a/Engine/Source/Runtime/Core/Private/Internationalization/FastDecimalFormat.cpp
+++ b/Engine/Source/Runtime/Core/Private/Internationalization/FastDecimalFormat.cpp
@@ -27,6 +27,12 @@ static FAutoConsoleVariableRef CVarFastDecimalFormatLargeFloatSupport(
 	bFastDecimalFormatLargeFloatSupport,
 	TEXT("True implies we perform additional processing for floating point types over 9223372036854775807 to prevent clipping to this value."));
 
+static int32 bFastDecimalFormatCanIndicateIfNearlyInteger = 1;
+static FAutoConsoleVariableRef CVarIndicateIfNearlyInteger(
+	TEXT("Core.bFastDecimalFormatCanIndicateIfNearlyInteger"),
+	bFastDecimalFormatCanIndicateIfNearlyInteger,
+	TEXT("True implies we will indicate nearly integer values with '...' - such as 1e-16 which would normally be displayed as 0."));
+
 static const uint64 Pow10Table[] = {
 	1,						// 10^0
 	10,						// 10^1
@@ -548,6 +554,12 @@ void FractionalToString(const double InVal, const FDecimalNumberFormattingRules&
 	double FractionalPart = 0.0;
 	FractionalToString_SplitAndRoundNumber(bIsNegative, InVal, InFormattingOptions.MaximumFractionalDigits, InFormattingOptions.RoundingMode, IntegralPart, FractionalPart);
 
+	// Capture if nearly integer so we can warn users about denorms / near zero values being displayed as zero (Ex: 1e-16)
+	const bool bNearlyInteger = bFastDecimalFormatCanIndicateIfNearlyInteger 
+		&& InFormattingOptions.IndicateNearlyInteger 
+		&& FMath::IsNearlyEqual(IntegralPart, InVal) 
+		&& IntegralPart != InVal;
+
 	if (bIsNegative)
 	{
 		IntegralPart = -IntegralPart;
@@ -615,6 +627,15 @@ void FractionalToString(const double InVal, const FDecimalNumberFormattingRules&
 		FractionalPartBuffer[FractionalPartLen] = TCHAR('\0');
 	}
 
+	// Append an '...' if the value is nearly an integer
+	if (bNearlyInteger && FractionalPartLen + 3 < UE_ARRAY_COUNT(FractionalPartBuffer))
+	{
+		FractionalPartBuffer[FractionalPartLen++] = '.';
+		FractionalPartBuffer[FractionalPartLen++] = '.';
+		FractionalPartBuffer[FractionalPartLen++] = '.';
+		FractionalPartBuffer[FractionalPartLen] = TCHAR('\0');
+	}
+
 	BuildFinalString(bIsNegative, InFormattingOptions.AlwaysSign, InFormattingRules, IntegralPartBuffer, IntegralPartLen, FractionalPartBuffer, FractionalPartLen, OutString);
 }
 
diff --git a/Engine/Source/Runtime/Core/Private/Internationalization/Text.cpp b/Engine/Source/Runtime/Core/Private/Internationalization/Text.cpp
--- a/Engine/Source/Runtime/Core/Private/Internationalization/Text.cpp
+++ b/Engine/Source/Runtime/Core/Private/Internationalization/Text.cpp
@@ -124,6 +124,7 @@ const void* FTextInspector::GetSharedDataId(const FText& Text)
 FNumberFormattingOptions::FNumberFormattingOptions()
 	: AlwaysSign(false)
 	, UseGrouping(true)
+	, IndicateNearlyInteger(false)
 	, RoundingMode(ERoundingMode::HalfToEven)
 	, MinimumIntegralDigits(1)
 	, MaximumIntegralDigits(DBL_MAX_10_EXP + DBL_DIG + 1)
diff --git a/Engine/Source/Runtime/Core/Public/Internationalization/Text.h b/Engine/Source/Runtime/Core/Public/Internationalization/Text.h
--- a/Engine/Source/Runtime/Core/Public/Internationalization/Text.h
+++ b/Engine/Source/Runtime/Core/Public/Internationalization/Text.h
@@ -205,6 +205,9 @@ struct FNumberFormattingOptions
 	bool UseGrouping;
 	FNumberFormattingOptions& SetUseGrouping( bool InValue ){ UseGrouping = InValue; return *this; }
 
+	bool IndicateNearlyInteger;
+	FNumberFormattingOptions& SetIndicateNearlyInteger( bool InValue ) { IndicateNearlyInteger = InValue; return *this; }
+
 	ERoundingMode RoundingMode;
 	FNumberFormattingOptions& SetRoundingMode( ERoundingMode InValue ){ RoundingMode = InValue; return *this; }
 
diff --git a/Engine/Source/Runtime/Slate/Public/Widgets/Input/NumericTypeInterface.h b/Engine/Source/Runtime/Slate/Public/Widgets/Input/NumericTypeInterface.h
--- a/Engine/Source/Runtime/Slate/Public/Widgets/Input/NumericTypeInterface.h
+++ b/Engine/Source/Runtime/Slate/Public/Widgets/Input/NumericTypeInterface.h
@@ -24,11 +24,15 @@ struct INumericTypeInterface
 	/** Gets the minimum and maximum fractional digits. */
 	virtual int32 GetMinFractionalDigits() const = 0;
 	virtual int32 GetMaxFractionalDigits() const = 0;
+	virtual bool GetIndicateNearlyInteger() const { return false; };
 
 	/** Sets the minimum and maximum fractional digits - A minimum greater than 0 will always have that many trailing zeros */
 	virtual void SetMinFractionalDigits(const TAttribute<TOptional<int32>>& NewValue) = 0;
 	virtual void SetMaxFractionalDigits(const TAttribute<TOptional<int32>>& NewValue) = 0;
 
+	/** Sets if we should indicate that a value is being rounded to an integer via '...' (Ex: 0.0 shown, real value 1e-18). */
+	virtual void SetIndicateNearlyInteger(const TAttribute<TOptional<bool>>& NewValue) {};
+
 	/** Convert the type to/from a string */
 	virtual FString ToString(const NumericType& Value) const = 0;
 	virtual TOptional<NumericType> FromString(const FString& InString, const NumericType& ExistingValue) = 0;
@@ -37,7 +41,7 @@ struct INumericTypeInterface
 	virtual bool IsCharacterValid(TCHAR InChar) const = 0;
 
 	/** Optional callback to broadcast when a setting in the type interface changes */
-	virtual FOnSettingChanged* GetOnSettingChanged() { return nullptr;}
+	virtual FOnSettingChanged* GetOnSettingChanged() { return nullptr; }
 };
 
 /** Default numeric type interface */
@@ -46,47 +50,66 @@ struct TDefaultNumericTypeInterface : INumericTypeInterface<NumericType>
 {
 
 	/** The default minimum fractional digits */
-	static const int32 DefaultMinFractionalDigits = 1;
+	static const int16 DefaultMinFractionalDigits = 1;
 
 	/** The default maximum fractional digits */
-	static const int32 DefaultMaxFractionalDigits = 6;
+	static const int16 DefaultMaxFractionalDigits = 6;
+
+	/** The default indicate nearly integer */
+	static const bool DefaultIndicateNearlyInteger = false;
 
 	/** The current minimum fractional digits */
-	int32 MinFractionalDigits = DefaultMinFractionalDigits;
+	int16 MinFractionalDigits = DefaultMinFractionalDigits;
 
 	/** The current maximum fractional digits */
-	int32 MaxFractionalDigits = DefaultMaxFractionalDigits;
+	int16 MaxFractionalDigits = DefaultMaxFractionalDigits;
+
+	/** True implies indicate when a value is displayed as rounded via '...' */
+	bool bIndicateNearlyInteger = false;
 
 	/** Gets the minimum and maximum fractional digits. */
 	virtual int32 GetMinFractionalDigits() const override
 	{
-		return MinFractionalDigits;
+		return IntCastChecked<int32>(MinFractionalDigits);
 	}
 	virtual int32 GetMaxFractionalDigits() const override
 	{
-		return MaxFractionalDigits;
+		return IntCastChecked<int32>(MaxFractionalDigits);
 	}
 
+	/** Get if we should indicate nearly integer values */
+	virtual bool GetIndicateNearlyInteger() const override
+	{ 
+		return bIndicateNearlyInteger;
+	};
+
 	/** Sets the minimum and maximum fractional digits - A minimum greater than 0 will always have that many trailing zeros */
 	virtual void SetMinFractionalDigits(const TAttribute<TOptional<int32>>& NewValue) override
 	{
-		MinFractionalDigits = (NewValue.Get().IsSet()) ? FMath::Max(0, NewValue.Get().GetValue()) :
+		MinFractionalDigits = NewValue.Get().IsSet() ? IntCastChecked<int16>(FMath::Max(0, NewValue.Get().GetValue())) :
 			DefaultMinFractionalDigits;
 	}
 
 	virtual void SetMaxFractionalDigits(const TAttribute<TOptional<int32>>& NewValue) override
 	{
-		MaxFractionalDigits = (NewValue.Get().IsSet()) ? FMath::Max(0, NewValue.Get().GetValue()) :
+		MaxFractionalDigits = NewValue.Get().IsSet() ? IntCastChecked<int16>(FMath::Max(0, NewValue.Get().GetValue())) :
 			DefaultMaxFractionalDigits;
 	}
 
+	virtual void SetIndicateNearlyInteger(const TAttribute<TOptional<bool>>& NewValue) override
+	{
+		bIndicateNearlyInteger = NewValue.Get().IsSet() ? NewValue.Get().GetValue() :
+			DefaultIndicateNearlyInteger;
+	}
+
 	/** Convert the type to/from a string */
 	virtual FString ToString(const NumericType& Value) const override
 	{
 		const FNumberFormattingOptions NumberFormattingOptions = FNumberFormattingOptions()
 			.SetUseGrouping(false)
 			.SetMinimumFractionalDigits(TIsIntegral<NumericType>::Value ? 0 : MinFractionalDigits)
-			.SetMaximumFractionalDigits(TIsIntegral<NumericType>::Value ? 0 : FMath::Max(MaxFractionalDigits, MinFractionalDigits));
+			.SetMaximumFractionalDigits(TIsIntegral<NumericType>::Value ? 0 : FMath::Max(MaxFractionalDigits, MinFractionalDigits))
+			.SetIndicateNearlyInteger(TIsIntegral<NumericType>::Value ? false : bIndicateNearlyInteger);
 		return FastDecimalFormat::NumberToString(Value, ExpressionParser::GetLocalizedNumberFormattingRules(), NumberFormattingOptions);
 	}
 
diff --git a/Engine/Source/Runtime/Slate/Public/Widgets/Input/SNumericEntryBox.h b/Engine/Source/Runtime/Slate/Public/Widgets/Input/SNumericEntryBox.h
--- a/Engine/Source/Runtime/Slate/Public/Widgets/Input/SNumericEntryBox.h
+++ b/Engine/Source/Runtime/Slate/Public/Widgets/Input/SNumericEntryBox.h
@@ -572,6 +572,26 @@ private:
 	FText GetValueAsText() const
 	{
 		const TOptional<NumericType>& Value = ValueAttribute.Get();
+
+#if WITH_EDITOR
+		if constexpr (std::is_floating_point_v<NumericType>)
+		{
+			if (Value.IsSet() && CachedValueString.EndsWith(TEXT("...")))
+			{
+				NumericType NumericValue = Value.GetValue();
+				FString ValueScientificStr = FString::Printf(TEXT("%e"), NumericValue);
+
+				// Dev Note: If you are seeing this message when not expected it can mean the following:
+				// 1. Simple content / value change is needed assuming the ~1e-16 off value is not desired.
+				// 2. If the above doesn't fix the issue, then there probably is an editor customization / some data writer that needs fixing
+				// 3. If the value is a true denorm, smaller than 1.175e-38 / 2.225e-307 for floats / doubles respectively. There is a larger issue.
+				//    In this case, it is possible that you may be reading from garbage / random memory not intended for floating point types.
+
+				return FText::Format(NSLOCTEXT("SNumericEntryBox", "NearZeroWarning_ToolTip", "Exact value: {0}. Set value or reimport after updating source content to change."), FText::FromString(ValueScientificStr));
+			}
+		}
+#endif // WITH_EDITOR
+
 		if (Value.IsSet() == true)
 		{
 			NumericType CurrentValue = Value.GetValue();
@@ -703,6 +723,22 @@ private:
 	virtual void Tick( const FGeometry& AllottedGeometry, const double InCurrentTime, const float InDeltaTime ) override
 	{
 		// Update the cached value, if needed.
+#if WITH_EDITOR
+		if constexpr (std::is_floating_point_v<NumericType>)
+		{
+			if (!GetPersistentState().bIsInGameLayer && !Interface->GetIndicateNearlyInteger())
+			{
+				Interface->SetIndicateNearlyInteger(true);
+				bCachedValueStringDirty = true;
+
+				if (SpinBox.IsValid())
+				{
+					SpinBox->ResetCachedValueString();
+				}
+			}
+		}
+#endif // WITH_EDITOR
+
 		const TOptional<NumericType>& Value = ValueAttribute.Get();
 		if (Value.IsSet() == true)
 		{
@@ -779,7 +815,7 @@ private:
 	/** Toggle checkbox */
 	TSharedPtr<SCheckBox> ToggleCheckBox;
 	/** Spinbox widget */
-	TSharedPtr<SWidget> SpinBox;
+	TSharedPtr<SSpinBox<NumericType>> SpinBox;
 	/** Editable widget */
 	TSharedPtr<SEditableText> EditableText;
 	/** Delegate to call when the value changes */
diff --git a/Engine/Source/Runtime/Slate/Public/Widgets/Input/SSpinBox.h b/Engine/Source/Runtime/Slate/Public/Widgets/Input/SSpinBox.h
--- a/Engine/Source/Runtime/Slate/Public/Widgets/Input/SSpinBox.h
+++ b/Engine/Source/Runtime/Slate/Public/Widgets/Input/SSpinBox.h
@@ -318,6 +318,9 @@ public:
 	SLATE_API void SetTextClearKeyboardFocusOnCommit(bool bNewValue);
 	SLATE_API void SetTextSelectAllTextOnCommit(bool bNewValue);
 
+	/** Reset the cached string. Typically used when the value is the same but the display format changed (through the callback). */
+	SLATE_API void ResetCachedValueString();
+
 protected:
 	/** Make the spinbox switch to keyboard-based input mode. */
 	SLATE_API void EnterTextMode();
@@ -477,9 +480,6 @@ private:
 	/** Used to prevent per-frame re-conversion of the cached numeric value to a string. */
 	FString CachedValueString;
 
-	/** Reset the cached string when the value is the same but the display format changed (through the callback). */
-	void ResetCachedValueString();
-
 	/** Whetever the interfaced setting changed and the CachedValueString needs to be recomputed. */
 	mutable bool bCachedValueStringDirty;
 
