Replace EndsWidth check in tooltip with Contains to account for units.

#tests editor

diff --git a/Engine/Source/Runtime/Slate/Public/Widgets/Input/SNumericEntryBox.h b/Engine/Source/Runtime/Slate/Public/Widgets/Input/SNumericEntryBox.h
--- a/Engine/Source/Runtime/Slate/Public/Widgets/Input/SNumericEntryBox.h
+++ b/Engine/Source/Runtime/Slate/Public/Widgets/Input/SNumericEntryBox.h
@@ -576,7 +576,7 @@ private:
 #if WITH_EDITOR
 		if constexpr (std::is_floating_point_v<NumericType>)
 		{
-			if (Value.IsSet() && CachedValueString.EndsWith(TEXT("...")))
+			if (Value.IsSet() && CachedValueString.Contains(TEXT("...")))
 			{
 				NumericType NumericValue = Value.GetValue();
 				FString ValueScientificStr = FString::Printf(TEXT("%e"), NumericValue);
