SRotatorInputBox : Call OnBeginSliderMovement / OnEndSliderMovement for Roll and Yaw to align with the existing behavior of Pitch.

Pitch was hooked up to both OnBeginSliderMovement and OnPitchBeginSliderMovement but Roll and Yaw were only hooked up to the more specific delegates.

diff --git a/Engine/Source/Runtime/Slate/Public/Widgets/Input/SRotatorInputBox.h b/Engine/Source/Runtime/Slate/Public/Widgets/Input/SRotatorInputBox.h
--- a/Engine/Source/Runtime/Slate/Public/Widgets/Input/SRotatorInputBox.h
+++ b/Engine/Source/Runtime/Slate/Public/Widgets/Input/SRotatorInputBox.h
@@ -174,8 +174,8 @@ public:
 				.Value( InArgs._Roll )
 				.OnValueChanged( InArgs._OnRollChanged )
 				.OnValueCommitted( InArgs._OnRollCommitted )
-				.OnBeginSliderMovement( InArgs._OnRollBeginSliderMovement )
-				.OnEndSliderMovement( InArgs._OnRollEndSliderMovement )
+				.OnBeginSliderMovement( CreatePerComponentSliderMovementEvent(InArgs._OnBeginSliderMovement, InArgs._OnRollBeginSliderMovement) )
+				.OnEndSliderMovement( CreatePerComponentSliderMovementEvent< FOnNumericValueChanged, NumericType >(InArgs._OnEndSliderMovement, InArgs._OnRollEndSliderMovement) )
 				.UndeterminedString( NSLOCTEXT("SRotatorInputBox", "MultipleValues", "Multiple Values") )
 				.ToolTipText_Lambda([RollAttr = InArgs._Roll]
 				{
@@ -246,8 +246,8 @@ public:
 				.Value( InArgs._Yaw )
 				.OnValueChanged( InArgs._OnYawChanged )
 				.OnValueCommitted( InArgs._OnYawCommitted )
-				.OnBeginSliderMovement( InArgs._OnYawBeginSliderMovement )
-				.OnEndSliderMovement( InArgs._OnYawEndSliderMovement )
+				.OnBeginSliderMovement( CreatePerComponentSliderMovementEvent(InArgs._OnBeginSliderMovement, InArgs._OnYawBeginSliderMovement) )
+				.OnEndSliderMovement( CreatePerComponentSliderMovementEvent< FOnNumericValueChanged, NumericType >(InArgs._OnEndSliderMovement, InArgs._OnYawEndSliderMovement) )
 				.UndeterminedString( NSLOCTEXT("SRotatorInputBox", "MultipleValues", "Multiple Values") )
 				.ToolTipText_Lambda([YawAttr = InArgs._Yaw]
 				{
