Flip default SSpinBox PreventThrottling back to true
* Users found it more intuitive that the spinbox update as it is dragged instead of on release

diff --git a/Engine/Source/Runtime/Slate/Public/Widgets/Input/SSpinBox.h b/Engine/Source/Runtime/Slate/Public/Widgets/Input/SSpinBox.h
--- a/Engine/Source/Runtime/Slate/Public/Widgets/Input/SSpinBox.h
+++ b/Engine/Source/Runtime/Slate/Public/Widgets/Input/SSpinBox.h
@@ -97,7 +97,7 @@ public:
 		, _MinDesiredWidth(0.0f)
 		, _Justification(ETextJustify::Left)
 		, _KeyboardType(Keyboard_Default)
-		, _PreventThrottling(false)
+		, _PreventThrottling(true)
 	{}
 
 	/** The style used to draw this spinbox */
