SSpinBox : Call EditableText->SetText when entering text mode and from CommitValue when in text mode.

Also, call EditableText->SetText from TextField_OnTextCommitted if the NewValue is not numeric and we don't CommitValue.

Previously, any calls to EditableText->SetText ( e.g. from OnMouseWheel ) would break the attribute binding on EditableText.Text leading to a disconnect between the value displayed in text mode and the ValueAttribute.

diff --git a/Engine/Source/Runtime/Slate/Private/Widgets/Input/SSpinBox.cpp b/Engine/Source/Runtime/Slate/Private/Widgets/Input/SSpinBox.cpp
--- a/Engine/Source/Runtime/Slate/Private/Widgets/Input/SSpinBox.cpp
+++ b/Engine/Source/Runtime/Slate/Private/Widgets/Input/SSpinBox.cpp
@@ -168,7 +168,6 @@ void SSpinBox<NumericType>::Construct(const FArguments& InArgs)
 				.Visibility(EVisibility::Collapsed)
 				.Font(InArgs._Font)
 				.SelectAllTextWhenFocused(true)
-				.Text(this, &SSpinBox<NumericType>::GetValueAsText)
 				.OnIsTypedCharValid(this, &SSpinBox<NumericType>::IsCharacterValid)
 				.OnTextChanged(this, &SSpinBox<NumericType>::TextField_OnTextChanged)
 				.OnTextCommitted(this, &SSpinBox<NumericType>::TextField_OnTextCommitted)
@@ -602,14 +601,7 @@ FReply SSpinBox<NumericType>::OnMouseWheel(const FGeometry& MyGeometry, const FP
 		const double NewValue = InternalValue + (Sign * Step);
 		const NumericType RoundedNewValue = RoundIfIntegerValue(NewValue);
 
-		TSharedPtr<INumericTypeInterface<NumericType>> Interface = InterfaceAttr.Get();
-
-		// First SetText is to update the value before calling CommitValue. Otherwise, when the text lose
-		// focus from the CommitValue, it will override the value we just committed.
-		// The second SetText is to update the text to the InternalValue since it could have been clamped.
-		EditableText->SetText(FText::FromString(Interface->ToString((NumericType)NewValue)));
 		CommitValue(RoundedNewValue, NewValue, CommittedViaSpin, ETextCommit::OnEnter);
-		EditableText->SetText(FText::FromString(Interface->ToString((NumericType)InternalValue)));
 
 		return FReply::Handled();
 	}
@@ -861,6 +853,7 @@ template <typename NumericType>
 void SSpinBox<NumericType>::EnterTextMode()
 {
 	TextBlock->SetVisibility(EVisibility::Collapsed);
+	EditableText->SetText(GetValueAsText());
 	EditableText->SetVisibility(EVisibility::Visible);
 }
 
@@ -944,6 +937,10 @@ void SSpinBox<NumericType>::TextField_OnTextCommitted(const FText& NewText, ETex
 	{
 		CommitValue(NewValue.GetValue(), (double)NewValue.GetValue(), CommittedViaTypeIn, CommitInfo);
 	}
+	else if (IsInTextMode())
+	{
+		EditableText->SetText(GetValueAsText());
+	}
 }
 
 template <typename NumericType>
@@ -1018,6 +1015,12 @@ void SSpinBox<NumericType>::CommitValue(NumericType NewValue, double NewSpinValu
 		ValueAttribute.Set(NewValue);
 	}
 
+	if (IsInTextMode())
+	{
+		TSharedPtr<INumericTypeInterface<NumericType>> Interface = InterfaceAttr.Get();
+		EditableText->SetText(FText::FromString(Interface->ToString((NumericType)NewValue)));
+	}
+
 	// Update the cache of the external value to what the user believes the value is now.
 	const NumericType CurrentValue = ValueAttribute.Get();
 	if (CachedExternalValue != CurrentValue || bCachedValueStringDirty)
