Null check for monitor change during startup

#rnx

diff --git a/Engine/Source/Runtime/Slate/Private/Framework/Application/SlateApplication.cpp b/Engine/Source/Runtime/Slate/Private/Framework/Application/SlateApplication.cpp
--- a/Engine/Source/Runtime/Slate/Private/Framework/Application/SlateApplication.cpp
+++ b/Engine/Source/Runtime/Slate/Private/Framework/Application/SlateApplication.cpp
@@ -7315,7 +7315,10 @@ void FSlateApplication::OnVirtualDesktopSizeChanged(const FDisplayMetrics& NewDi
 		VirtualDisplayRect.Right,
 		VirtualDisplayRect.Bottom);
 
-	Renderer->OnVirtualDesktopSizeChanged(NewDisplayMetric);
+	if (Renderer.IsValid())
+	{
+		Renderer->OnVirtualDesktopSizeChanged(NewDisplayMetric);
+	}
 }
 
 
