Slate: Expose SetStyle and IsOpen on STextComboBox

#rnx

diff --git a/Engine/Source/Runtime/Slate/Public/Widgets/Input/STextComboBox.h b/Engine/Source/Runtime/Slate/Public/Widgets/Input/STextComboBox.h
--- a/Engine/Source/Runtime/Slate/Public/Widgets/Input/STextComboBox.h
+++ b/Engine/Source/Runtime/Slate/Public/Widgets/Input/STextComboBox.h
@@ -94,6 +94,18 @@ public:
 		StringCombo->ClearItemsSource();
 	}
 
+	/** Update the style of the internal ComboBox */
+	void SetStyle(const FComboBoxStyle* InComboStyle)
+	{
+		StringCombo->SetStyle(InComboStyle);
+	}
+
+	/** @return true if the popup is open; false otherwise. */
+	bool IsOpen() const
+	{
+		return StringCombo->IsOpen();
+	}
+
 	/** Request to reload the text options in the combobox from the OptionsSource attribute */
 	SLATE_API void RefreshOptions();
 
