Changing SpinBox to use a default step on keyboard press if step is set to 0
Moving SpinBox step calculation functionality to a new function and defining variables used
#rnx

diff --git a/Engine/Source/Runtime/Slate/Private/Widgets/Input/SSpinBox.cpp b/Engine/Source/Runtime/Slate/Private/Widgets/Input/SSpinBox.cpp
--- a/Engine/Source/Runtime/Slate/Private/Widgets/Input/SSpinBox.cpp
+++ b/Engine/Source/Runtime/Slate/Private/Widgets/Input/SSpinBox.cpp
@@ -466,21 +466,8 @@ FReply SSpinBox<NumericType>::OnMouseMove(const FGeometry& MyGeometry, const FPo
 				}
 			}
 
-			ECommitMethod CommitMethod = CommittedViaSpin;
-				
-			const bool bIsSmallStep = (GetMaxSliderValue() - GetMinSliderValue()) <= 10.0;
-			double Step = bIsSmallStep ? 0.1 : 1.0;
-
-			if (MouseEvent.IsControlDown())
-			{
-				Step *= CtrlMultiplier.Get();
-				CommitMethod = CommittedViaSpinMultiplier;
-			}
-			else if (MouseEvent.IsShiftDown())
-			{
-				Step *= ShiftMultiplier.Get();
-				CommitMethod = CommittedViaSpinMultiplier;
-			}
+			ECommitMethod CommitMethod = MouseEvent.IsControlDown() || MouseEvent.IsShiftDown() ? CommittedViaSpinMultiplier : CommittedViaSpin;
+			double Step = GetDefaultStepSize(MouseEvent);
 
 			//if we have a range to draw in
 			if (!bUnlimitedSpinRange)
@@ -584,10 +571,10 @@ FReply SSpinBox<NumericType>::OnMouseWheel(const FGeometry& MyGeometry, const FP
 {
 	if (bEnableWheel && PointerDraggingSliderIndex == INDEX_NONE && HasKeyboardFocus())
 	{
-		// If there is no WheelStep defined, we use 1.0 (Or 0.1 if slider range is <= 10)
+		// If there is no WheelStep defined, we use StepSize (Or SmallStepSize if slider range is <= SmallStepSizeMax)
 		constexpr bool bIsIntegral = TIsIntegral<NumericType>::Value;
-		const bool bIsSmallStep = !bIsIntegral && (GetMaxSliderValue() - GetMinSliderValue()) <= 10.0;
-		double Step = WheelStep.IsSet() && WheelStep.Get().IsSet() ? WheelStep.Get().GetValue() : (bIsSmallStep ? 0.1 : 1.0);
+		const bool bIsSmallStep = !bIsIntegral && (GetMaxSliderValue() - GetMinSliderValue()) <= SmallStepSizeMax;
+		double Step = WheelStep.IsSet() && WheelStep.Get().IsSet() ? WheelStep.Get().GetValue() : (bIsSmallStep ? SmallStepSize : StepSize);
 
 		if (MouseEvent.IsControlDown())
 		{
@@ -667,7 +654,7 @@ FReply SSpinBox<NumericType>::OnKeyDown(const FGeometry& MyGeometry, const FKeyE
 	else if (Key == EKeys::Up || Key == EKeys::Right)
 	{
 		const NumericType LocalValueAttribute = ValueAttribute.Get();
-		const NumericType LocalDelta = Delta.Get();
+		const NumericType LocalDelta = Delta.Get() != 0 ? Delta.Get() : (NumericType)GetDefaultStepSize(InKeyEvent);
 		InternalValue = (double)LocalValueAttribute;
 		CommitValue(LocalValueAttribute + LocalDelta, InternalValue + (double)LocalDelta, CommittedViaArrowKey, ETextCommit::OnEnter);
 		ExitTextMode();
@@ -676,7 +663,7 @@ FReply SSpinBox<NumericType>::OnKeyDown(const FGeometry& MyGeometry, const FKeyE
 	else if (Key == EKeys::Down || Key == EKeys::Left)
 	{
 		const NumericType LocalValueAttribute = ValueAttribute.Get();
-		const NumericType LocalDelta = Delta.Get();
+		const NumericType LocalDelta = Delta.Get() != 0 ? Delta.Get() : (NumericType)GetDefaultStepSize(InKeyEvent);
 		InternalValue = (double)LocalValueAttribute;
 		CommitValue(LocalValueAttribute - LocalDelta, InternalValue + (double)LocalDelta, CommittedViaArrowKey, ETextCommit::OnEnter);
 		ExitTextMode();
@@ -1125,6 +1112,24 @@ void SSpinBox<NumericType>::CancelMouseCapture()
 	NotifyValueCommitted(PreDragValue);
 }
 
+template<typename NumericType>
+double SSpinBox<NumericType>::GetDefaultStepSize(FInputEvent InputEvent)
+{
+	const bool bIsSmallStep = (GetMaxSliderValue() - GetMinSliderValue()) <= SmallStepSizeMax;
+	double Step = bIsSmallStep ? SmallStepSize : StepSize;
+
+	if (InputEvent.IsControlDown())
+	{
+		Step *= CtrlMultiplier.Get();
+	}
+	else if (InputEvent.IsShiftDown())
+	{
+		Step *= ShiftMultiplier.Get();
+	}
+	
+	return Step;
+}
+
 template <typename NumericType>
 void SSpinBox<NumericType>::ResetCachedValueString()
 { 
diff --git a/Engine/Source/Runtime/Slate/Public/Widgets/Input/SSpinBox.h b/Engine/Source/Runtime/Slate/Public/Widgets/Input/SSpinBox.h
--- a/Engine/Source/Runtime/Slate/Public/Widgets/Input/SSpinBox.h
+++ b/Engine/Source/Runtime/Slate/Public/Widgets/Input/SSpinBox.h
@@ -511,4 +511,19 @@ private:
 	bool bBroadcastValueChangesPerKey = false;
 	
 	TOptional<FPendingCommitValue> PendingCommitValue;
+
+	/*
+	* Gets the default amount to change the slider when delta is not applicable. 
+	* Control takes priority over shift
+	*/
+	double GetDefaultStepSize(FInputEvent InputEvent);
+
+	/** Gets the default amount to change the slider when delta is not applicable. **/
+	const double StepSize = 1.f;
+
+	/** Step size to use when range is below SmallStepSizeMax. **/
+	const double SmallStepSize = .1f; 
+
+	/** Largest numerical value to use the SmallStepSize instead of StepSize. **/
+	const double SmallStepSizeMax = 10.f;
 };
\ No newline at end of file
