Fixing a bug in slate where we could attempt to load config data before initializing the UClass that stored that config data. The result was a silent failure to load the config.

[RN] minor slate

diff --git a/Engine/Source/Runtime/Slate/Private/Framework/Application/SlateApplication.cpp b/Engine/Source/Runtime/Slate/Private/Framework/Application/SlateApplication.cpp
--- a/Engine/Source/Runtime/Slate/Private/Framework/Application/SlateApplication.cpp
+++ b/Engine/Source/Runtime/Slate/Private/Framework/Application/SlateApplication.cpp
@@ -7350,13 +7350,49 @@ TSharedRef<FSlateApplication> FSlateApplication::InitializeAsStandaloneApplicati
 	return Slate;
 }
 
+#if ALLOW_THEMES
+namespace SlatePrivate
+{
+	FDelegateHandle UObjectsRegisteredDelegateHandle;
+
+	void FinishInitializingStarshipStyleAfterConfigsAreReadable(FName Name, ECompiledInUObjectsRegisteredStatus Status)
+	{
+		if (Status == ECompiledInUObjectsRegisteredStatus::PostCDO)
+		{
+			// Ensure that *our* CDO is ready since PostCDO is broadcast many times
+			if (USlateThemeManager::StaticClass()->GetDefaultObject(false) != nullptr)
+			{
+				USlateThemeManager::Get().LoadThemes();
+
+				FStarshipCoreStyle::ResetToDefault();
+				FAppStyle::SetAppStyleSet(FStarshipCoreStyle::GetCoreStyle());
+				FUMGCoreStyle::ResetToDefault();
+
+				FCoreUObjectDelegates::CompiledInUObjectsRegisteredDelegate.Remove(UObjectsRegisteredDelegateHandle);
+				UObjectsRegisteredDelegateHandle.Reset();
+			}
+		}
+	}
+}
+#endif // ALOW_THEMES
+
 void FSlateApplication::InitializeCoreStyle()
 {
 	if (FCoreStyle::IsStarshipStyle() && !FStarshipCoreStyle::IsInitialized())
 	{
 #if ALLOW_THEMES
-		USlateThemeManager::Get().LoadThemes();
+		// This code may run too early to safely load themes (because we may not have fully loaded this module yet). If the CDO doesn't exist yet
+		// don't create it and instead defer loading the themes until we finish processing module
+		if (USlateThemeManager::StaticClass()->GetDefaultObject(false) != nullptr)
+		{
+			USlateThemeManager::Get().LoadThemes();
+		}
+		else
+		{
+			SlatePrivate::UObjectsRegisteredDelegateHandle = FCoreUObjectDelegates::CompiledInUObjectsRegisteredDelegate.AddStatic(&SlatePrivate::FinishInitializingStarshipStyleAfterConfigsAreReadable);
+		}
 #endif
+		// If we deferred theme loading, we'll repeat these operations in FinishInitializingStarshipStyleAfterConfigsAreReadable.
 		FStarshipCoreStyle::ResetToDefault();
 		FAppStyle::SetAppStyleSet(FStarshipCoreStyle::GetCoreStyle());
 	}
diff --git a/Engine/Source/Runtime/SlateCore/Private/Styling/SlateColor.cpp b/Engine/Source/Runtime/SlateCore/Private/Styling/SlateColor.cpp
--- a/Engine/Source/Runtime/SlateCore/Private/Styling/SlateColor.cpp
+++ b/Engine/Source/Runtime/SlateCore/Private/Styling/SlateColor.cpp
@@ -31,6 +31,19 @@ bool FSlateColor::SerializeFromMismatchedTag(const FPropertyTag& Tag, FStructure
 
 const FLinearColor& FSlateColor::GetColorFromTable() const
 {
-	return USlateThemeManager::Get().GetColor(ColorTableId);
+	static bool bClassDefaultObjectInitialized = false;
+	if (!bClassDefaultObjectInitialized && USlateThemeManager::StaticClass()->GetDefaultObject(false))
+	{
+		bClassDefaultObjectInitialized = true;
+	}
+
+	if (bClassDefaultObjectInitialized)
+	{
+		return USlateThemeManager::Get().GetColor(ColorTableId);
+	}
+	else
+	{
+		return FLinearColor::Transparent;
+	}
 }
 
diff --git a/Engine/Source/Runtime/SlateCore/Public/Styling/StyleColors.h b/Engine/Source/Runtime/SlateCore/Public/Styling/StyleColors.h
--- a/Engine/Source/Runtime/SlateCore/Public/Styling/StyleColors.h
+++ b/Engine/Source/Runtime/SlateCore/Public/Styling/StyleColors.h
@@ -19,7 +19,7 @@
 struct FSlateColor;
 
 /** Themes are only allowed in the editor or standalone tools */
-#define ALLOW_THEMES WITH_EDITOR || IS_PROGRAM
+#define ALLOW_THEMES (WITH_EDITOR || IS_PROGRAM)
 
 // HEX Colors from sRGB Color space
 #define COLOR( HexValue ) FLinearColor::FromSRGBColor(FColor::FromHex(HexValue))
