Fix internal automation test

diff --git a/Engine/Source/Runtime/Slate/Private/Widgets/Input/SSearchBox.cpp b/Engine/Source/Runtime/Slate/Private/Widgets/Input/SSearchBox.cpp
--- a/Engine/Source/Runtime/Slate/Private/Widgets/Input/SSearchBox.cpp
+++ b/Engine/Source/Runtime/Slate/Private/Widgets/Input/SSearchBox.cpp
@@ -158,6 +158,7 @@ void SSearchBox::Construct( const FArguments& InArgs )
 		.HAlign(HAlign_Center)
 		.VAlign(VAlign_Center)
 		.OnClicked(this, &SSearchBox::OnClearSearch)
+		.AddMetaData<FTagMetaData>(TEXT("SearchClearButton"))
 		// Allow the button to steal focus so that the search text will be automatically committed. Afterwards focus will be returned to the text box.
 		// If the user is keyboard-centric, they'll "ctrl+a, delete" to clear the search
 		.IsFocusable(true)
