New viewport toolbar
- Added an hovered brush to visually group raised checkboxes and their comboboxes.

diff --git a/Engine/Source/Editor/EditorStyle/Private/StarshipStyle.cpp b/Engine/Source/Editor/EditorStyle/Private/StarshipStyle.cpp
--- a/Engine/Source/Editor/EditorStyle/Private/StarshipStyle.cpp
+++ b/Engine/Source/Editor/EditorStyle/Private/StarshipStyle.cpp
@@ -2559,10 +2559,23 @@ void FStarshipEditorStyle::FStyle::SetupViewportStyles()
 				BackgroundColorRaised.A = .0f;
 				const FSlateRoundedBoxBrush BackgroundBrushRaised(BackgroundColorRaised, 5.f, BackgroundColorRaised, 1.0);
 
-				ButtonStyleRaised.SetNormal(BackgroundBrushRaised).SetNormalForeground(FStyleColors::Foreground);
+				ButtonStyleRaised.SetNormal(BackgroundBrushRaised)
+					.SetNormalForeground(FStyleColors::Foreground)
+					.SetPressed(BackgroundBrush);
 			}
 
-			ViewportToolbarRaisedStyle.SetButtonStyle(ButtonStyleRaised).SetComboButtonPadding(0);
+			FCheckBoxStyle ToggleButtonStyleRaised = ToggleButtonStyle;
+			{
+				ToggleButtonStyleRaised.SetUncheckedHoveredImage(HoveredButtonBrush)
+					.SetCheckedPressedImage(BackgroundBrush)
+					.SetUncheckedPressedImage(BackgroundBrush);
+			}
+
+			ViewportToolbarRaisedStyle.SetButtonStyle(ButtonStyleRaised)
+				.SetToggleButtonStyle(ToggleButtonStyleRaised)
+				.SetComboButtonPadding(0)
+				.SetBlockHovered(BackgroundBrush);
+
 			Set("ViewportToolbar.Raised", ViewportToolbarRaisedStyle);
 		}
 
diff --git a/Engine/Source/Runtime/Slate/Private/Framework/MultiBox/SToolBarComboButtonBlock.cpp b/Engine/Source/Runtime/Slate/Private/Framework/MultiBox/SToolBarComboButtonBlock.cpp
--- a/Engine/Source/Runtime/Slate/Private/Framework/MultiBox/SToolBarComboButtonBlock.cpp
+++ b/Engine/Source/Runtime/Slate/Private/Framework/MultiBox/SToolBarComboButtonBlock.cpp
@@ -8,6 +8,7 @@
 #include "Framework/MultiBox/MultiBoxBuilder.h"
 #include "Widgets/Input/SCheckBox.h"
 #include "Widgets/Input/SComboButton.h"
+#include "Widgets/Layout/SBorder.h"
 #include "Widgets/Layout/SBox.h"
 #include "Styling/ToolBarStyle.h"
 #include "Widgets/Images/SLayeredImage.h"
@@ -227,6 +228,7 @@ void SToolBarComboButtonBlock::BuildMultiBlockWidget(const ISlateStyle* StyleSet
 	}
 
 	OpenForegroundColor = ButtonStyle->HoveredForeground;
+	BlockHovered = &ToolBarStyle.BlockHovered;
 
 	TAttribute<FText> ActualToolTip;
 	if (ToolBarComboButtonBlock->ToolTip.IsSet())
@@ -295,7 +297,7 @@ void SToolBarComboButtonBlock::BuildMultiBlockWidget(const ISlateStyle* StyleSet
 		}
 
 		// clang-format off
-		const TSharedRef<SCheckBox> CheckBox = SNew(SCheckBox)
+		SAssignNew(CheckBoxWidget, SCheckBox)
 			// Use the tool bar style for this check box
 			.Style(CheckStyle)
 			.CheckBoxContentUsesAutoWidth(false)
@@ -309,7 +311,7 @@ void SToolBarComboButtonBlock::BuildMultiBlockWidget(const ISlateStyle* StyleSet
 		// clang-format on
 
 		// clang-format off
-		const TSharedRef<SComboButton> ComboButton = SAssignNew(ComboButtonWidget, SComboButton)
+		SAssignNew(ComboButtonWidget, SComboButton)
 			.AddMetaData<FTagMetaData>(FTagMetaData(TagName))
 			.ContentPadding(0.f)
 			.ComboButtonStyle(ComboStyle)
@@ -332,21 +334,26 @@ void SToolBarComboButtonBlock::BuildMultiBlockWidget(const ISlateStyle* StyleSet
 		// clang-format off
 		ChildSlot
 		[
-			SNew(SHorizontalBox)
-			.ToolTip(FMultiBoxSettings::ToolTipConstructor.Execute(
-				ActualToolTip, nullptr, UICommand, /*ShowActionShortcut=*/true
-			))
-			+ SHorizontalBox::Slot()
-			.AutoWidth()
-			.VAlign(VAlign_Center)
-			[
-				CheckBox
-			]
-			+ SHorizontalBox::Slot()
-			.AutoWidth()
-			.VAlign(VAlign_Center)
+			SNew(SBorder)
+			.BorderImage( this, &SToolBarComboButtonBlock::GetBorderImage )
+			.Padding(0.f)
 			[
-				ComboButton
+				SNew(SHorizontalBox)
+				.ToolTip(FMultiBoxSettings::ToolTipConstructor.Execute(
+					ActualToolTip, nullptr, UICommand, /*ShowActionShortcut=*/true
+				))
+				+ SHorizontalBox::Slot()
+				.AutoWidth()
+				.VAlign(VAlign_Center)
+				[
+					CheckBoxWidget.ToSharedRef()
+				]
+				+ SHorizontalBox::Slot()
+				.AutoWidth()
+				.VAlign(VAlign_Center)
+				[
+					ComboButtonWidget.ToSharedRef()
+				]
 			]
 		];
 		// clang-format on
@@ -522,3 +529,15 @@ FSlateColor SToolBarComboButtonBlock::OnGetForegroundColor() const
 		return FSlateColor::UseStyle();
 	}
 }
+
+const FSlateBrush* SToolBarComboButtonBlock::GetBorderImage() const
+{
+	if ((ComboButtonWidget && ComboButtonWidget->IsHovered()) || (CheckBoxWidget && CheckBoxWidget->IsHovered()))
+	{
+		return BlockHovered;
+	}
+	else
+	{
+		return FAppStyle::GetBrush("NoBorder");
+	}
+}
diff --git a/Engine/Source/Runtime/Slate/Public/Framework/MultiBox/SToolBarComboButtonBlock.h b/Engine/Source/Runtime/Slate/Public/Framework/MultiBox/SToolBarComboButtonBlock.h
--- a/Engine/Source/Runtime/Slate/Public/Framework/MultiBox/SToolBarComboButtonBlock.h
+++ b/Engine/Source/Runtime/Slate/Public/Framework/MultiBox/SToolBarComboButtonBlock.h
@@ -10,6 +10,7 @@
 #include "Widgets/DeclarativeSyntaxSupport.h"
 #include "Framework/MultiBox/MultiBox.h"
 
+class SCheckBox;
 class SComboButton;
 
 /**
@@ -210,6 +211,8 @@ private:
 
 	FSlateColor OnGetForegroundColor() const;
 
+	const FSlateBrush* GetBorderImage() const;
+
 private:
 	/** Overrides the visibility of the of label. This is used to set up the LabelVisibility attribute */
 	TOptional<EVisibility> LabelVisibilityOverride;
@@ -220,11 +223,15 @@ private:
 	/** Optional overridden icon for this tool bar button.  IF not set, then the action's icon will be used instead. */
 	TAttribute<FSlateIcon> Icon;
 
+	TSharedPtr<SCheckBox> CheckBoxWidget;
 	TSharedPtr<SComboButton> ComboButtonWidget;
 
 	/** The foreground color for button when the combo button is open */
 	FSlateColor OpenForegroundColor;
 
+	/** The hovered style for the entire block used when the block has multiple widgets */
+	const FSlateBrush* BlockHovered = nullptr;
+
 	/** Whether this toolbar should always use small icons, regardless of the current settings */
 	bool bForceSmallIcons;
 };
diff --git a/Engine/Source/Runtime/SlateCore/Private/Styling/ToolBarStyle.cpp b/Engine/Source/Runtime/SlateCore/Private/Styling/ToolBarStyle.cpp
--- a/Engine/Source/Runtime/SlateCore/Private/Styling/ToolBarStyle.cpp
+++ b/Engine/Source/Runtime/SlateCore/Private/Styling/ToolBarStyle.cpp
@@ -28,6 +28,7 @@ FToolBarStyle::FToolBarStyle()
 	, CheckBoxPadding(0)
 	, BlockPadding(0)
 	, IndentedBlockPadding(0)
+	, BlockHovered(FSlateNoResource())
 	, BackgroundPadding(0)
 	, WrapButtonPadding(0)
 	, WrapButtonIndex(-1)
@@ -53,6 +54,7 @@ void FToolBarStyle::GetResources(TArray<const FSlateBrush*>& OutBrushes) const
 	OutBrushes.Add(&BackgroundBrush);
 	OutBrushes.Add(&ExpandBrush);
 	OutBrushes.Add(&SeparatorBrush);
+	OutBrushes.Add(&BlockHovered);
 
 	LabelStyle.GetResources(OutBrushes);
 	EditableTextStyle.GetResources(OutBrushes);
diff --git a/Engine/Source/Runtime/SlateCore/Public/Styling/ToolBarStyle.h b/Engine/Source/Runtime/SlateCore/Public/Styling/ToolBarStyle.h
--- a/Engine/Source/Runtime/SlateCore/Public/Styling/ToolBarStyle.h
+++ b/Engine/Source/Runtime/SlateCore/Public/Styling/ToolBarStyle.h
@@ -121,6 +121,11 @@ struct FToolBarStyle : public FSlateWidgetStyle
 	FMargin IndentedBlockPadding;
 	FToolBarStyle& SetIndentedBlockPadding(const FMargin& InIndentedBlockPadding) { IndentedBlockPadding = InIndentedBlockPadding; return *this; }
 
+	/** Hovered brush for an entire block */
+	UPROPERTY(EditAnywhere, BlueprintReadWrite, Category = Appearance)
+	FSlateBrush BlockHovered;
+	FToolBarStyle& SetBlockHovered(const FSlateBrush& InBlockHovered) { BlockHovered = InBlockHovered; return *this; }
+
 	UPROPERTY(EditAnywhere, BlueprintReadWrite, Category = Appearance)
 	FMargin BackgroundPadding;
 	FToolBarStyle& SetBackgroundPadding(const FMargin& InMargin) { BackgroundPadding = InMargin; return *this; }
