Added bConsumePointerInput to ScrollBox allowing to pass input events to the widgets underneath. Added GetOverscrollOffset() and GetOverscrollPercentage()  to ScrollBox, returning overscroll offset in slate units and in percent from its view size.

#rnx

diff --git a/Engine/Source/Runtime/Slate/Private/Widgets/Layout/SScrollBox.cpp b/Engine/Source/Runtime/Slate/Private/Widgets/Layout/SScrollBox.cpp
--- a/Engine/Source/Runtime/Slate/Private/Widgets/Layout/SScrollBox.cpp
+++ b/Engine/Source/Runtime/Slate/Private/Widgets/Layout/SScrollBox.cpp
@@ -117,6 +117,7 @@ void SScrollBox::Construct( const FArguments& InArgs )
 	bAnimateWheelScrolling = InArgs._AnimateWheelScrolling;
 	WheelScrollMultiplier = InArgs._WheelScrollMultiplier;
 	bEnableTouchScrolling = InArgs._EnableTouchScrolling;
+	bConsumePointerInput = InArgs._ConsumePointerInput;
 	NavigationScrollPadding = InArgs._NavigationScrollPadding;
 	NavigationDestination = InArgs._NavigationDestination;
 	ScrollWhenFocusChanges = InArgs._ScrollWhenFocusChanges;
@@ -337,6 +338,65 @@ float SScrollBox::GetScrollOffset() const
 	return DesiredScrollOffset;
 }
 
+float SScrollBox::GetOverscrollOffset() const
+{
+	if (AllowOverscroll != EAllowOverscroll::Yes)
+	{
+		return 0.0f;
+	}
+
+	const float ScrollOffset = ScrollPanel->PhysicalOffset;
+	if (ScrollOffset < 0.0f)
+	{
+		// Negative overscroll (before the start of the ScrollBox)
+		return ScrollOffset;
+	}
+
+	const float ContentSize = GetScrollComponentFromVector(ScrollPanel->GetDesiredSize());
+	const float ViewSize = GetScrollComponentFromVector(CachedGeometry.GetLocalSize());
+	if (ScrollOffset > ContentSize - ViewSize)
+	{
+		// Positive overscroll (past the end of the ScrollBox)
+		return ScrollOffset - (ContentSize - ViewSize);
+	}
+
+	// No overscroll
+	return 0.0f;
+}
+
+float SScrollBox::GetOverscrollPercentage() const
+{
+	if (AllowOverscroll != EAllowOverscroll::Yes)
+	{
+		return 0.0f;
+	}
+
+	const float ViewSize = GetScrollComponentFromVector(CachedGeometry.GetLocalSize());
+	if (ViewSize <= 0.0f)
+	{
+		// Avoid division by zero
+		return 0.0f;
+	}
+	
+	const float ScrollOffset = ScrollPanel->PhysicalOffset;
+	if (ScrollOffset < 0.0f)
+	{
+		// Negative overscroll (before the start of the ScrollBox)
+		return (ScrollOffset / ViewSize) * 100.0f;
+	}
+
+	const float ContentSize = GetScrollComponentFromVector(ScrollPanel->GetDesiredSize());
+	if (ScrollOffset > ContentSize - ViewSize)
+	{
+		// Positive overscroll (past the end of the ScrollBox)
+		const float OverscrollAmount = ScrollOffset - (ContentSize - ViewSize);
+		return (OverscrollAmount / ViewSize) * 100.0f;
+	}
+
+	// No overscroll
+	return 0.0f;
+}
+
 float SScrollBox::GetScrollOffsetOfEnd() const
 {
 	const FGeometry ScrollPanelGeometry = FindChildGeometry(CachedGeometry, ScrollPanel.ToSharedRef());
@@ -718,7 +778,7 @@ FReply SScrollBox::OnMouseButtonDown( const FGeometry& MyGeometry, const FPointe
 
 	if ( MouseEvent.IsTouchEvent() )
 	{
-		return FReply::Handled();
+		return bConsumePointerInput ? FReply::Handled() : SWidget::OnMouseButtonDown(MyGeometry, MouseEvent);
 	}
 	else
 	{
@@ -728,7 +788,7 @@ FReply SScrollBox::OnMouseButtonDown( const FGeometry& MyGeometry, const FPointe
 
 			Invalidate(EInvalidateWidget::Layout);
 
-			return FReply::Handled();
+			return bConsumePointerInput ? FReply::Handled() : SWidget::OnMouseButtonDown(MyGeometry, MouseEvent);
 		}
 	}
 
@@ -750,7 +810,7 @@ FReply SScrollBox::OnMouseButtonUp( const FGeometry& MyGeometry, const FPointerE
 
 		Invalidate(EInvalidateWidget::Layout);
 
-		FReply Reply = FReply::Handled().ReleaseMouseCapture();
+		FReply Reply = bConsumePointerInput ? FReply::Handled().ReleaseMouseCapture() : SWidget::OnMouseButtonUp(MyGeometry, MouseEvent).ReleaseMouseCapture();
 		bShowSoftwareCursor = false;
 
 		// If we have mouse capture, snap the mouse back to the closest location that is within the panel's bounds
@@ -766,7 +826,7 @@ FReply SScrollBox::OnMouseButtonUp( const FGeometry& MyGeometry, const FPointerE
 
 			Reply.SetMousePos(BestPositionInPanel);
 		}
-
+		
 		return Reply;
 	}
 
@@ -796,11 +856,11 @@ FReply SScrollBox::OnMouseMove( const FGeometry& MyGeometry, const FPointerEvent
 					// The user has moved the list some amount; they are probably
 					// trying to scroll. From now on, the list assumes the user is scrolling
 					// until they lift their finger.
-					Reply = FReply::Handled().CaptureMouse(AsShared());
+					Reply = bConsumePointerInput ? FReply::Handled().CaptureMouse(AsShared()) : SWidget::OnMouseMove(MyGeometry, MouseEvent).CaptureMouse(AsShared()) ;
 				}
 				else
 				{
-					Reply = FReply::Handled();
+					Reply = bConsumePointerInput ? FReply::Handled() : SWidget::OnMouseMove(MyGeometry, MouseEvent);
 				}
 			}
 		}
@@ -812,7 +872,7 @@ FReply SScrollBox::OnMouseMove( const FGeometry& MyGeometry, const FPointerEvent
 				InertialScrollManager.AddScrollSample(-ScrollByAmountScreen, FSlateApplication::Get().GetCurrentTime());
 				ScrollBy(MyGeometry, -ScrollByAmountLocal, EAllowOverscroll::Yes, false);
 
-				Reply = FReply::Handled();
+				Reply = bConsumePointerInput ? FReply::Handled() : SWidget::OnMouseMove(MyGeometry, MouseEvent);
 			}
 		}
 
@@ -833,7 +893,7 @@ FReply SScrollBox::OnMouseMove( const FGeometry& MyGeometry, const FPointerEvent
 				InertialScrollManager.AddScrollSample(-ScrollByAmountScreen, FPlatformTime::Seconds());
 				const bool bDidScroll = ScrollBy(MyGeometry, -ScrollByAmountLocal, AllowOverscroll, false);
 
-				FReply Reply = FReply::Handled();
+				FReply Reply = bConsumePointerInput ? FReply::Handled() : SWidget::OnMouseMove(MyGeometry, MouseEvent);
 
 				// Capture the mouse if we need to
 				if ( HasMouseCapture() == false )
@@ -970,7 +1030,7 @@ FReply SScrollBox::OnTouchEnded(const FGeometry& MyGeometry, const FPointerEvent
 		
 		BeginInertialScrolling();
 
-		return FReply::Handled().ReleaseMouseCapture();
+		return bConsumePointerInput ? FReply::Handled().ReleaseMouseCapture() : SWidget::OnTouchEnded(MyGeometry, InTouchEvent).ReleaseMouseCapture() ;
 	}
 
 	return FReply::Unhandled();
diff --git a/Engine/Source/Runtime/Slate/Public/Widgets/Layout/SScrollBox.h b/Engine/Source/Runtime/Slate/Public/Widgets/Layout/SScrollBox.h
--- a/Engine/Source/Runtime/Slate/Public/Widgets/Layout/SScrollBox.h
+++ b/Engine/Source/Runtime/Slate/Public/Widgets/Layout/SScrollBox.h
@@ -158,6 +158,7 @@ public:
 		, _ScrollAnimationInterpSpeed(15.f)
 		, _WheelScrollMultiplier(1.f)
 		, _EnableTouchScrolling(true)
+		, _ConsumePointerInput (true)
 		, _NavigationDestination(EDescendantScrollDestination::IntoView)
 		, _NavigationScrollPadding(0.0f)
 		, _ScrollWhenFocusChanges(EScrollWhenFocusChanges::NoScroll)
@@ -206,6 +207,8 @@ public:
 		SLATE_ARGUMENT(float, WheelScrollMultiplier);
 
 		SLATE_ARGUMENT(bool, EnableTouchScrolling);
+		
+		SLATE_ARGUMENT(bool, ConsumePointerInput );
 
 		SLATE_ARGUMENT(EDescendantScrollDestination, NavigationDestination);
 
@@ -272,10 +275,16 @@ public:
 	/** Enables/disables being able to scroll using touch input. */
 	SLATE_API void SetIsTouchScrollingEnabled(const bool bInEnableTouchScrolling);
 	
+	SLATE_API void SetConsumePointerInput(bool bInConsumePointerInput) { bConsumePointerInput = bInConsumePointerInput; }
+	
 	SLATE_API void SetScrollWhenFocusChanges(EScrollWhenFocusChanges NewScrollWhenFocusChanges);
 
 	SLATE_API float GetScrollOffset() const;
-
+	
+	SLATE_API float GetOverscrollOffset() const;
+	
+	SLATE_API float GetOverscrollPercentage() const;
+	
 	SLATE_API float GetViewFraction() const;
 
 	SLATE_API float GetViewOffsetFraction() const;
@@ -514,6 +523,9 @@ protected:
 	/** True to allow scrolling by using touch input. */
 	bool bEnableTouchScrolling = true;
 
+	/** If true, touch input events will pass through to widgets under the ScrollBox, while still being handled by the ScrollBox. */
+	bool bConsumePointerInput = true;
+
 	/** The speed of interpolation for the scrolling animation */
 	float ScrollingAnimationInterpolationSpeed = 15.f;
 
diff --git a/Engine/Source/Runtime/UMG/Private/Components/ScrollBox.cpp b/Engine/Source/Runtime/UMG/Private/Components/ScrollBox.cpp
--- a/Engine/Source/Runtime/UMG/Private/Components/ScrollBox.cpp
+++ b/Engine/Source/Runtime/UMG/Private/Components/ScrollBox.cpp
@@ -105,6 +105,7 @@ PRAGMA_DISABLE_DEPRECATION_WARNINGS
 		.ScrollAnimationInterpSpeed(ScrollAnimationInterpolationSpeed)
 		.WheelScrollMultiplier(WheelScrollMultiplier)
 		.EnableTouchScrolling(bEnableTouchScrolling)
+		.ConsumePointerInput(bConsumePointerInput)
 		.OnUserScrolled(BIND_UOBJECT_DELEGATE(FOnUserScrolled, SlateHandleUserScrolled))
 		.OnScrollBarVisibilityChanged(BIND_UOBJECT_DELEGATE(FOnScrollBarVisibilityChanged, SlateHandleScrollBarVisibilityChanged));
 PRAGMA_ENABLE_DEPRECATION_WARNINGS
@@ -144,6 +145,7 @@ PRAGMA_DISABLE_DEPRECATION_WARNINGS
 	MyScrollBox->SetScrollingAnimationInterpolationSpeed(ScrollAnimationInterpolationSpeed);
 	MyScrollBox->SetWheelScrollMultiplier(WheelScrollMultiplier);
 	MyScrollBox->SetIsTouchScrollingEnabled(bEnableTouchScrolling);
+	MyScrollBox->SetConsumePointerInput(bConsumePointerInput);
 	MyScrollBox->SetStyle(&WidgetStyle);
 	MyScrollBox->InvalidateStyle();
 	MyScrollBox->SetScrollBarStyle(&WidgetBarStyle);
@@ -161,6 +163,26 @@ float UScrollBox::GetScrollOffset() const
 	return 0;
 }
 
+float UScrollBox::GetOverscrollOffset() const
+{
+	if ( MyScrollBox.IsValid() )
+	{
+		return MyScrollBox->GetOverscrollOffset();
+	}
+
+	return 0;
+}
+
+float UScrollBox::GetOverscrollPercentage() const
+{
+	if ( MyScrollBox.IsValid() )
+	{
+		return MyScrollBox->GetOverscrollPercentage();
+	}
+
+	return 0;
+}
+
 float UScrollBox::GetScrollOffsetOfEnd() const
 {
 	if (MyScrollBox.IsValid())
@@ -475,6 +497,21 @@ bool UScrollBox::GetIsTouchScrollingEnabled() const
 	return bEnableTouchScrolling;
 }
 
+void UScrollBox::SetConsumePointerInput(bool bInConsumePointerInput)
+{
+	bConsumePointerInput = bInConsumePointerInput;
+
+	if (MyScrollBox.IsValid())
+	{
+		MyScrollBox->SetConsumePointerInput(bInConsumePointerInput);
+	}
+}
+
+bool UScrollBox::GetConsumePointerInput() const
+{
+	return bConsumePointerInput;
+}
+
 void UScrollBox::SetScrollWhenFocusChanges(EScrollWhenFocusChanges NewScrollWhenFocusChanges)
 {
 	ScrollWhenFocusChanges = NewScrollWhenFocusChanges;
diff --git a/Engine/Source/Runtime/UMG/Public/Components/ScrollBox.h b/Engine/Source/Runtime/UMG/Public/Components/ScrollBox.h
--- a/Engine/Source/Runtime/UMG/Public/Components/ScrollBox.h
+++ b/Engine/Source/Runtime/UMG/Public/Components/ScrollBox.h
@@ -30,6 +30,10 @@ private:
 	/** True to allow scrolling using touch input. */
 	UPROPERTY(EditAnywhere, BlueprintReadWrite, Getter="GetIsTouchScrollingEnabled", Setter="SetIsTouchScrollingEnabled", Category = "Scroll", meta = (AllowPrivateAccess = "true", DisplayAfter = "WheelScrollMultiplier"))
 	bool bEnableTouchScrolling = true;
+	
+	/** Option to pass touch input events to widgets behind the ScrollBox. */
+	UPROPERTY(EditAnywhere, BlueprintReadWrite, Getter = "GetConsumePointerInput", Setter = "SetConsumePointerInput", Category = "Scroll", meta = (AllowPrivateAccess = "true", DisplayAfter = "bEnableTouchScrolling"))
+	bool bConsumePointerInput = true;
 
 public:
 
@@ -188,6 +192,12 @@ public:
 	UMG_API void SetIsTouchScrollingEnabled(bool bInEnableTouchScrolling);
 
 	UMG_API bool GetIsTouchScrollingEnabled() const;
+	
+	UFUNCTION(BlueprintCallable, Category = "Scroll")
+	UMG_API bool GetConsumePointerInput() const;
+
+	UFUNCTION(BlueprintCallable, Category = "Scroll")
+	UMG_API void SetConsumePointerInput(bool bInConsumePointerInput);
 
 	UFUNCTION(BlueprintCallable, Category = "Scroll")
 	UMG_API void SetScrollWhenFocusChanges(EScrollWhenFocusChanges NewScrollWhenFocusChanges);
@@ -240,6 +250,17 @@ public:
 	UFUNCTION(BlueprintCallable, Category="Widget")
 	UMG_API float GetScrollOffset() const;
 
+	/**
+	 * Gets the overscroll offset of the scrollbox in Slate Units.
+	 */
+	UFUNCTION(BlueprintCallable, Category="Widget")
+	UMG_API float GetOverscrollOffset() const;
+	/**
+	 * Gets the overscroll offset of the scrollbox in percentage from its view size.
+	 */
+	UFUNCTION(BlueprintCallable, Category="Widget")
+	UMG_API float GetOverscrollPercentage() const;
+
 	/** Gets the scroll offset of the bottom of the ScrollBox in Slate Units. */
 	UFUNCTION(BlueprintCallable, Category = "Widget")
 	UMG_API float GetScrollOffsetOfEnd() const;
