Fixing a bug in slate where we could attempt to load config data before initializing the UClass that stored that config data. The result was a silent failure to load the config. That load operation is now deferred. As a result, it is necessary to defer setting a theme by way of a new callback FSlateApplication::

[RN] minor slate

diff --git a/Engine/Source/Editor/EditorStyle/Private/EditorStyleModule.cpp b/Engine/Source/Editor/EditorStyle/Private/EditorStyleModule.cpp
--- a/Engine/Source/Editor/EditorStyle/Private/EditorStyleModule.cpp
+++ b/Engine/Source/Editor/EditorStyle/Private/EditorStyleModule.cpp
@@ -1,6 +1,7 @@
 // Copyright Epic Games, Inc. All Rights Reserved.
 
 #include "CoreMinimal.h"
+#include "Framework/Application/SlateApplication.h"
 #include "Modules/ModuleManager.h"
 #include "Interfaces/IEditorStyleModule.h"
 #include "SlateEditorStyle.h"
@@ -21,30 +22,23 @@ public:
 
 	virtual void StartupModule( ) override
 	{
-		if (FCoreStyle::IsStarshipStyle())
-		{
 #if ALLOW_THEMES
-			USlateThemeManager::Get().ValidateActiveTheme();
-#endif
-
-			bUsingStarshipStyle = true;
-			FStarshipEditorStyle::Initialize();
-
-			// set the application style to be the editor style
-			FAppStyle::SetAppStyleSetName(FStarshipEditorStyle::GetStyleSetName());
-
-		}
-		else
+		if (!FSlateApplication::IsThemeSystemReady())
 		{
-			FSlateEditorStyle::Initialize();
-
-			// set the application style to be the editor style
-			FAppStyle::SetAppStyleSetName(FAppStyle::GetAppStyleSetName());
+			FSlateApplication::OnThemeSystemReady().AddRaw(this, &FEditorStyleModule::InitializeStyle);
+			return;
 		}
+#endif // ALLOW_THEMES
+
+		InitializeStyle();
 	}
 
 	virtual void ShutdownModule( ) override
 	{
+#if ALLOW_THEMES
+		FSlateApplication::OnThemeSystemReady().RemoveAll(this);
+#endif // ALLOW_THEMES
+
 		if (bUsingStarshipStyle)
 		{
 			FStarshipEditorStyle::Shutdown();
@@ -58,6 +52,29 @@ public:
 	// End IModuleInterface interface
 private:
 	bool bUsingStarshipStyle = false;
+
+	void InitializeStyle()
+	{
+		if (FCoreStyle::IsStarshipStyle())
+		{
+#if ALLOW_THEMES
+			USlateThemeManager::Get().ValidateActiveTheme();
+#endif
+
+			bUsingStarshipStyle = true;
+			FStarshipEditorStyle::Initialize();
+
+			// set the application style to be the editor style
+			FAppStyle::SetAppStyleSetName(FStarshipEditorStyle::GetStyleSetName());
+		}
+		else
+		{
+			FSlateEditorStyle::Initialize();
+
+			// set the application style to be the editor style
+			FAppStyle::SetAppStyleSetName(FAppStyle::GetAppStyleSetName());
+		}
+	}
 };
 
 
diff --git a/Engine/Source/Runtime/Slate/Private/Framework/Application/SlateApplication.cpp b/Engine/Source/Runtime/Slate/Private/Framework/Application/SlateApplication.cpp
--- a/Engine/Source/Runtime/Slate/Private/Framework/Application/SlateApplication.cpp
+++ b/Engine/Source/Runtime/Slate/Private/Framework/Application/SlateApplication.cpp
@@ -7350,13 +7350,78 @@ TSharedRef<FSlateApplication> FSlateApplication::InitializeAsStandaloneApplicati
 	return Slate;
 }
 
+namespace SlatePrivate
+{
+	FDelegateHandle UObjectsRegisteredDelegateHandle;
+	FSlateApplication::FOnThemeSystemReady ThemeSystemReadyDelegate;
+
+#if ALLOW_THEMES
+	void BroadcastAndClearThemeSystemReadyDelegate()
+	{
+		FSlateApplication::FOnThemeSystemReady LocalCopy = ThemeSystemReadyDelegate;
+		ThemeSystemReadyDelegate.Clear();
+		LocalCopy.Broadcast();
+	}
+
+	void FinishInitializingStarshipStyleAfterConfigsAreReadable(FName Name, ECompiledInUObjectsRegisteredStatus Status)
+	{
+		if (Status == ECompiledInUObjectsRegisteredStatus::PostCDO)
+		{
+			// Ensure that *our* CDO is ready since PostCDO is broadcast many times
+			if (USlateThemeManager::StaticClass()->GetDefaultObject(false) != nullptr)
+			{
+				USlateThemeManager::Get().LoadThemes();
+
+				FStarshipCoreStyle::ResetToDefault();
+				FAppStyle::SetAppStyleSet(FStarshipCoreStyle::GetCoreStyle());
+				FUMGCoreStyle::ResetToDefault();
+
+				FCoreUObjectDelegates::CompiledInUObjectsRegisteredDelegate.Remove(UObjectsRegisteredDelegateHandle);
+				UObjectsRegisteredDelegateHandle.Reset();
+
+				BroadcastAndClearThemeSystemReadyDelegate();
+			}
+		}
+	}
+#endif // ALOW_THEMES
+}
+
+FSlateApplication::FOnThemeSystemReady& FSlateApplication::OnThemeSystemReady()
+{
+	return SlatePrivate::ThemeSystemReadyDelegate;
+}
+
+bool FSlateApplication::IsThemeSystemReady()
+{
+#if ALLOW_THEMES
+	if (USlateThemeManager::StaticClass()->GetDefaultObject(false) != nullptr)
+	{
+		return USlateThemeManager::Get().AreThemesLoaded();
+	}
+#endif
+
+	return false;
+}
+
 void FSlateApplication::InitializeCoreStyle()
 {
+	bool bDeferredThemeLoading = false;
 	if (FCoreStyle::IsStarshipStyle() && !FStarshipCoreStyle::IsInitialized())
 	{
 #if ALLOW_THEMES
-		USlateThemeManager::Get().LoadThemes();
+		// This code may run too early to safely load themes (because we may not have fully loaded this module yet). If the CDO doesn't exist yet
+		// don't create it and instead defer loading the themes until we finish processing module
+		if (USlateThemeManager::StaticClass()->GetDefaultObject(false) != nullptr)
+		{
+			USlateThemeManager::Get().LoadThemes();
+		}
+		else
+		{
+			SlatePrivate::UObjectsRegisteredDelegateHandle = FCoreUObjectDelegates::CompiledInUObjectsRegisteredDelegate.AddStatic(&SlatePrivate::FinishInitializingStarshipStyleAfterConfigsAreReadable);
+			bDeferredThemeLoading = true;
+		}
 #endif
+		// If we deferred theme loading, we'll repeat these operations in FinishInitializingStarshipStyleAfterConfigsAreReadable.
 		FStarshipCoreStyle::ResetToDefault();
 		FAppStyle::SetAppStyleSet(FStarshipCoreStyle::GetCoreStyle());
 	}
@@ -7367,6 +7432,14 @@ void FSlateApplication::InitializeCoreStyle()
 	}
 
 	FUMGCoreStyle::ResetToDefault();
+
+#if ALLOW_THEMES
+	if (!bDeferredThemeLoading)
+	{
+		// Must call this *after* resetting the themes/styles
+		SlatePrivate::BroadcastAndClearThemeSystemReadyDelegate();
+	}
+#endif // ALLOW_THEMES
 }
 
 void FSlateApplication::SetWidgetReflector(const TSharedRef<IWidgetReflector>& WidgetReflector)
diff --git a/Engine/Source/Runtime/Slate/Public/Framework/Application/SlateApplication.h b/Engine/Source/Runtime/Slate/Public/Framework/Application/SlateApplication.h
--- a/Engine/Source/Runtime/Slate/Public/Framework/Application/SlateApplication.h
+++ b/Engine/Source/Runtime/Slate/Public/Framework/Application/SlateApplication.h
@@ -284,6 +284,18 @@ public:
 	
 	static SLATE_API TSharedRef<FSlateApplication> InitializeAsStandaloneApplication(const TSharedRef< class FSlateRenderer >& PlatformRenderer, const TSharedRef<class GenericApplication>& InPlatformApplication);
 
+	/**
+	 * Returns true if themes have been loaded and can be applied. If false, any theme-based styling should be deferred via OnThemeSystemReady() 
+	 */
+	static SLATE_API bool IsThemeSystemReady();
+
+	/**
+	 * Use this method to defer loading themes until the theme system is ready. Callbacks are cleared after being called. Any callbacks registered afterward
+	 * will be called if and when InitializeCoreStyle() is invoked again.
+	 */
+	DECLARE_MULTICAST_DELEGATE(FOnThemeSystemReady);
+	static SLATE_API FOnThemeSystemReady& OnThemeSystemReady();
+
 	static SLATE_API void InitializeCoreStyle();
 
 	/**
diff --git a/Engine/Source/Runtime/SlateCore/Private/Styling/SlateColor.cpp b/Engine/Source/Runtime/SlateCore/Private/Styling/SlateColor.cpp
--- a/Engine/Source/Runtime/SlateCore/Private/Styling/SlateColor.cpp
+++ b/Engine/Source/Runtime/SlateCore/Private/Styling/SlateColor.cpp
@@ -31,6 +31,19 @@ bool FSlateColor::SerializeFromMismatchedTag(const FPropertyTag& Tag, FStructure
 
 const FLinearColor& FSlateColor::GetColorFromTable() const
 {
-	return USlateThemeManager::Get().GetColor(ColorTableId);
+	static bool bClassDefaultObjectInitialized = false;
+	if (!bClassDefaultObjectInitialized && USlateThemeManager::StaticClass()->GetDefaultObject(false))
+	{
+		bClassDefaultObjectInitialized = true;
+	}
+
+	if (bClassDefaultObjectInitialized)
+	{
+		return USlateThemeManager::Get().GetColor(ColorTableId);
+	}
+	else
+	{
+		return FLinearColor::Transparent;
+	}
 }
 
diff --git a/Engine/Source/Runtime/SlateCore/Private/Styling/StyleColors.cpp b/Engine/Source/Runtime/SlateCore/Private/Styling/StyleColors.cpp
--- a/Engine/Source/Runtime/SlateCore/Private/Styling/StyleColors.cpp
+++ b/Engine/Source/Runtime/SlateCore/Private/Styling/StyleColors.cpp
@@ -69,6 +69,9 @@ const FSlateColor FStyleColors::Error = EStyleColor::Error;
 const FSlateColor FStyleColors::Success = EStyleColor::Success;
 
 USlateThemeManager::USlateThemeManager()
+#if ALLOW_THEMES
+: bThemesLoaded(false)
+#endif
 {
 	InitalizeDefaults();
 }
@@ -153,10 +156,16 @@ void USlateThemeManager::LoadThemes()
 
 	EnsureValidCurrentTheme();
 
+	bThemesLoaded = true;
 
 	ApplyTheme(CurrentThemeId);
 }
 
+SLATECORE_API bool USlateThemeManager::AreThemesLoaded()
+{
+	return bThemesLoaded;
+}
+
 void USlateThemeManager::SaveCurrentThemeAs(const FString& Filename)
 {
 	FStyleTheme& CurrentTheme = GetMutableCurrentTheme();
@@ -204,6 +213,8 @@ void USlateThemeManager::SaveCurrentThemeAs(const FString& Filename)
 
 void USlateThemeManager::ApplyTheme(FGuid ThemeId)
 {
+	ensureMsgf(AreThemesLoaded(), TEXT("Cannot apply theme before theme system is ready. Calls should be deferred via FSlateApplication::OnThemeSystemReady."));
+
 	if (ThemeId.IsValid())
 	{
 		FStyleTheme* CurrentTheme = nullptr;
diff --git a/Engine/Source/Runtime/SlateCore/Public/Styling/StyleColors.h b/Engine/Source/Runtime/SlateCore/Public/Styling/StyleColors.h
--- a/Engine/Source/Runtime/SlateCore/Public/Styling/StyleColors.h
+++ b/Engine/Source/Runtime/SlateCore/Public/Styling/StyleColors.h
@@ -199,6 +199,11 @@ public:
 	 */
 	SLATECORE_API void LoadThemes();
 
+	/**
+	 * Returns true if themes have been loaded at least once.
+	 */
+	SLATECORE_API bool AreThemesLoaded();
+
 	/**
 	 * Saves the current theme
 	 * 
@@ -304,6 +309,7 @@ private:
 	FStyleTheme DefaultDarkTheme;
 	TArray<FStyleTheme> Themes;
 	FLinearColor DefaultColors[(int32)EStyleColor::MAX];
+	bool bThemesLoaded;
 
 	/** Broadcasts whenever the theme changes **/
 	FThemeChangedEvent ThemeChangedEvent;
