[STableViewBase] Fixed FixedLineScrollOffset not working on horizontal lists or vertical lists with few columns

#rnx

diff --git a/Engine/Source/Runtime/Slate/Private/Widgets/Views/STableViewBase.cpp b/Engine/Source/Runtime/Slate/Private/Widgets/Views/STableViewBase.cpp
--- a/Engine/Source/Runtime/Slate/Private/Widgets/Views/STableViewBase.cpp
+++ b/Engine/Source/Runtime/Slate/Private/Widgets/Views/STableViewBase.cpp
@@ -303,15 +303,16 @@ void STableViewBase::Tick( const FGeometry& AllottedGeometry, const double InCur
 				InertialScrollManager.ResetShouldStopScrollNow();
 			}
 
-			if((bStartedTouchInteraction && bEnableTouchAnimatedScrolling) || (!bStartedTouchInteraction && bEnableAnimatedScrolling))
+			bool bEndScroll = true;
+			if ((bStartedTouchInteraction && bEnableTouchAnimatedScrolling) || (!bStartedTouchInteraction && bEnableAnimatedScrolling))
 			{
 				CurrentScrollOffset = FMath::FInterpTo(CurrentScrollOffset, TargetScrollOffset, (double)InDeltaTime, ScrollingAnimationInterpolationSpeed);
-				if (FMath::IsNearlyEqual(CurrentScrollOffset, TargetScrollOffset, 0.01))
-				{
-					CurrentScrollOffset = TargetScrollOffset;
-				}
+				// We may need to wait until a new target scroll offset is recalculated at 0 velocity before we can end the scroll
+				const bool bWaitForZeroVelocityAutoScroll = ShouldScrollToFixedLineOffsetAtZeroVelocity() && InertialScrollManager.GetScrollVelocity() != 0.f;
+				bEndScroll = !bWaitForZeroVelocityAutoScroll && FMath::IsNearlyEqual(CurrentScrollOffset, TargetScrollOffset, 0.01);
 			}
-			else
+
+			if (bEndScroll)
 			{
 				CurrentScrollOffset = TargetScrollOffset;
 			}
@@ -393,6 +394,10 @@ void STableViewBase::Tick( const FGeometry& AllottedGeometry, const double InCur
 			}
 			else if (CurrentScrollOffset == TargetScrollOffset)
 			{
+				// Sync the desired scroll offset with the target scroll offset when the scroll ends.
+				// This prevents teleporting the table upon the next scroll. 
+				DesiredScrollOffset = TargetScrollOffset;
+				
 				NotifyFinishedScrolling();
 			}
 
@@ -823,7 +828,7 @@ STableViewBase::~STableViewBase() = default;
 
 double STableViewBase::GetTargetScrollOffset() const
 {
-	if (!bWasAtEndOfList && FixedLineScrollOffset.IsSet() && !IsRightClickScrolling() && InertialScrollManager.GetScrollVelocity() == 0.f)
+	if (ShouldScrollToFixedLineOffsetAtZeroVelocity() && InertialScrollManager.GetScrollVelocity() == 0.f)
 	{
 		// The scroll has ended, we calculate the appropriate the scroll offset to snap on a line.
 		// We don't do that if we are at the end of the list because in this case showing the last line correctly is more important.
@@ -836,6 +841,11 @@ double STableViewBase::GetTargetScrollOffset() const
 	return DesiredScrollOffset;
 }
 
+bool STableViewBase::ShouldScrollToFixedLineOffsetAtZeroVelocity() const
+{
+	return !bWasAtEndOfList && FixedLineScrollOffset.IsSet() && !IsRightClickScrolling() && !bStartedTouchInteraction;
+}
+
 float STableViewBase::ScrollBy(const FGeometry& MyGeometry, float ScrollByAmountInSlateUnits, EAllowOverscroll InAllowOverscroll)
 {
 	const int32 NumItemsBeingObserved = GetNumItemsBeingObserved();
diff --git a/Engine/Source/Runtime/Slate/Public/Widgets/Views/STableViewBase.h b/Engine/Source/Runtime/Slate/Public/Widgets/Views/STableViewBase.h
--- a/Engine/Source/Runtime/Slate/Public/Widgets/Views/STableViewBase.h
+++ b/Engine/Source/Runtime/Slate/Public/Widgets/Views/STableViewBase.h
@@ -561,6 +561,9 @@ private:
 	/** When true, a refresh should occur the next tick */
 	bool bItemsNeedRefresh = false;
 
+	/** Whether or not the view should scroll to a fixed line offset if the scroll velocity is 0 */
+	bool ShouldScrollToFixedLineOffsetAtZeroVelocity() const;
+
 	union
 	{
 		// vertical scroll bar is stored in horizontal box and vice versa
