[SListView] 'ClearScrollVelocityOnSelection' option added so that this behaviour can be disabled. Still true by default.

#rnx

diff --git a/Engine/Source/Runtime/Slate/Public/Widgets/Views/SListView.h b/Engine/Source/Runtime/Slate/Public/Widgets/Views/SListView.h
--- a/Engine/Source/Runtime/Slate/Public/Widgets/Views/SListView.h
+++ b/Engine/Source/Runtime/Slate/Public/Widgets/Views/SListView.h
@@ -116,6 +116,7 @@ public:
 		, _HandleSpacebarSelection(false)
 		, _IsFocusable(true)
 		, _ReturnFocusToSelection()
+		, _ClearScrollVelocityOnSelection(true)
 		, _OnItemToString_Debug()
 		, _OnEnteredBadState()
 		{
@@ -198,7 +199,9 @@ public:
 
 		SLATE_ATTRIBUTE(bool, IsFocusable)
 
-		SLATE_ARGUMENT(bool, ReturnFocusToSelection)
+		SLATE_ARGUMENT(bool, ReturnFocusToSelection)	
+		
+		SLATE_ARGUMENT(bool, ClearScrollVelocityOnSelection)
 
 		/** Assign this to get more diagnostics from the list view. */
 		SLATE_EVENT(FOnItemToString_Debug, OnItemToString_Debug)
@@ -248,6 +251,8 @@ public:
 		this->IsFocusable = InArgs._IsFocusable;
 
 		this->bReturnFocusToSelection = InArgs._ReturnFocusToSelection;
+			
+		this->bClearScrollVelocityOnSelection = InArgs._ClearScrollVelocityOnSelection;
 
 		this->bEnableAnimatedScrolling = InArgs._EnableAnimatedScrolling;
 		this->FixedLineScrollOffset = InArgs._FixedLineScrollOffset;
@@ -925,7 +930,11 @@ public:
 			RangeSelectionStart = TheItem;
 		}
 
-		this->InertialScrollManager.ClearScrollVelocity();
+		if (bClearScrollVelocityOnSelection)
+		{
+			this->InertialScrollManager.ClearScrollVelocity();
+		}
+
 #if WITH_ACCESSIBILITY
 		// On certain platforms, we need to pass accessibility focus to a widget for screen readers to announce 
 		// accessibility information. STableRows cannot accept keyboard focus, 
@@ -950,7 +959,10 @@ public:
 	{
 		SelectedItems.Empty();
 
-		this->InertialScrollManager.ClearScrollVelocity();
+		if (bClearScrollVelocityOnSelection)
+		{
+			this->InertialScrollManager.ClearScrollVelocity();
+		}
 	}
 
 	virtual void Private_SelectRangeFromCurrentTo( ItemType InRangeSelectionEnd ) override
@@ -987,8 +999,11 @@ public:
 			// The above loop won't add the last item, so manually add it here
 			SelectedItems.Add(ItemsSourceRef[ItemIndex]);
 		}
-
-		this->InertialScrollManager.ClearScrollVelocity();
+		
+		if (bClearScrollVelocityOnSelection)
+		{
+			this->InertialScrollManager.ClearScrollVelocity();
+		}
 	}
 
 	virtual void Private_SignalSelectionChanged(ESelectInfo::Type SelectInfo) override
@@ -2656,6 +2671,9 @@ protected:
 	/** If true, the focus will be returned to the last selected object in a list when navigated to. */
 	bool bReturnFocusToSelection;
 
+	/** True to force the scroll velocity to reset to 0 when an item of the list is selected */
+	bool bClearScrollVelocityOnSelection = true;
+
 	/** If true, the item currently slated to be scrolled into view will also be navigated to after being scrolled in */
 	bool bNavigateOnScrollIntoView = false;
 
diff --git a/Engine/Source/Runtime/UMG/Public/Components/ListView.h b/Engine/Source/Runtime/UMG/Public/Components/ListView.h
--- a/Engine/Source/Runtime/UMG/Public/Components/ListView.h
+++ b/Engine/Source/Runtime/UMG/Public/Components/ListView.h
@@ -171,6 +171,7 @@ protected:
 		Args.bClearSelectionOnClick = bClearSelectionOnClick;
 		Args.ConsumeMouseWheel = ConsumeMouseWheel;
 		Args.bReturnFocusToSelection = bReturnFocusToSelection;
+		Args.bClearScrollVelocityOnSelection = bClearScrollVelocityOnSelection;
 		Args.Orientation = Orientation;
 		Args.ListViewStyle = &WidgetStyle;
 		Args.ScrollBarStyle = &ScrollBarStyle;
@@ -209,6 +210,9 @@ protected:
 
 	UPROPERTY(EditAnywhere, BlueprintReadOnly, Category = ListView)
 	bool bIsFocusable = true;
+	
+	UPROPERTY(EditAnywhere, BlueprintReadOnly, Category = ListView)
+	bool bClearScrollVelocityOnSelection = true;
 
 	UPROPERTY(EditAnywhere, BlueprintReadOnly, Category = ListView)
 	bool bReturnFocusToSelection = false;
diff --git a/Engine/Source/Runtime/UMG/Public/Components/ListViewBase.h b/Engine/Source/Runtime/UMG/Public/Components/ListViewBase.h
--- a/Engine/Source/Runtime/UMG/Public/Components/ListViewBase.h
+++ b/Engine/Source/Runtime/UMG/Public/Components/ListViewBase.h
@@ -237,6 +237,7 @@ protected:
 		bool bClearSelectionOnClick = false;
 		EConsumeMouseWheel ConsumeMouseWheel = EConsumeMouseWheel::WhenScrollingPossible;
 		bool bReturnFocusToSelection = false;
+		bool bClearScrollVelocityOnSelection = true;
 		EOrientation Orientation = Orient_Vertical;
 		EScrollIntoViewAlignment ScrollIntoViewAlignment = EScrollIntoViewAlignment::CenterAligned;
 		const FTableViewStyle* ListViewStyle = &FUMGCoreStyle::Get().GetWidgetStyle<FTableViewStyle>("ListView");
@@ -259,6 +260,7 @@ protected:
 			.ConsumeMouseWheel(Args.ConsumeMouseWheel)
 			.SelectionMode(Args.SelectionMode)
 			.ReturnFocusToSelection(Args.bReturnFocusToSelection)
+			.ClearScrollVelocityOnSelection(Args.bClearScrollVelocityOnSelection)
 			.Orientation(Args.Orientation)
 			.ScrollIntoViewAlignment(Args.ScrollIntoViewAlignment)
 			.ListViewStyle(Args.ListViewStyle)
