Fixed unreachable code warning in DoAnyWindowDescendantsNeedPrepass.

diff --git a/Engine/Source/Runtime/Slate/Private/Framework/Application/SlateApplication.cpp b/Engine/Source/Runtime/Slate/Private/Framework/Application/SlateApplication.cpp
--- a/Engine/Source/Runtime/Slate/Private/Framework/Application/SlateApplication.cpp
+++ b/Engine/Source/Runtime/Slate/Private/Framework/Application/SlateApplication.cpp
@@ -1235,19 +1235,21 @@ void FSlateApplication::DrawWindowAndChildren( const TSharedRef<SWindow>& Window
 
 static bool DoAnyWindowDescendantsNeedPrepass(TSharedRef<SWindow> WindowToPrepass)
 {
-	for (const TSharedRef<SWindow>& ChildWindow : WindowToPrepass->GetChildWindows())
+	TArray<TSharedRef<SWindow>>& ChildWindows = WindowToPrepass->GetChildWindows();
+	if (ChildWindows.IsEmpty())
 	{
-		if (ChildWindow->IsVisible() && !ChildWindow->IsWindowMinimized())
-		{
-			return true;
-		}
-		else
-		{
-			return DoAnyWindowDescendantsNeedPrepass(ChildWindow);
-		}
+		return false;
 	}
 
-	return false;
+	const TSharedRef<SWindow>& ChildWindow = ChildWindows[0];
+	if (ChildWindow->IsVisible() && !ChildWindow->IsWindowMinimized())
+	{
+		return true;
+	}
+	else
+	{
+		return DoAnyWindowDescendantsNeedPrepass(ChildWindow);
+	}
 }
 
 static void PrepassWindowAndChildren(TSharedRef<SWindow> WindowToPrepass, const TSharedPtr<SWindow>& DebuggingWindow, TWeakPtr<SWidget>& CurrentContext)
