[Slate] Added NumSlots() to various multi-slot widgets

diff --git a/Engine/Source/Runtime/Slate/Public/Widgets/Layout/SUniformWrapPanel.h b/Engine/Source/Runtime/Slate/Public/Widgets/Layout/SUniformWrapPanel.h
--- a/Engine/Source/Runtime/Slate/Public/Widgets/Layout/SUniformWrapPanel.h
+++ b/Engine/Source/Runtime/Slate/Public/Widgets/Layout/SUniformWrapPanel.h
@@ -138,6 +138,9 @@ public:
 	/** Removes all slots from the panel */
 	SLATE_API void ClearChildren();
 
+	/** @return the number of slots. */
+	int32 NumSlots() const { return Children.Num(); }
+
 protected:
 	// Begin SWidget overrides.
 	SLATE_API virtual FVector2D ComputeDesiredSize(float) const override;
diff --git a/Engine/Source/Runtime/Slate/Public/Widgets/SCanvas.h b/Engine/Source/Runtime/Slate/Public/Widgets/SCanvas.h
--- a/Engine/Source/Runtime/Slate/Public/Widgets/SCanvas.h
+++ b/Engine/Source/Runtime/Slate/Public/Widgets/SCanvas.h
@@ -137,6 +137,9 @@ public:
 	 */
 	SLATE_API void ClearChildren();
 
+	/** @return the number of slots. */
+	int32 NumSlots() const { return Children.Num(); }
+
 public:
 
 	//~ SWidget overrides
