SNumericVectorInputBox : Only query GetAxisDisplayName if NumberOfComponents == 3

diff --git a/Engine/Source/Runtime/Slate/Public/Widgets/Input/SVectorInputBox.h b/Engine/Source/Runtime/Slate/Public/Widgets/Input/SVectorInputBox.h
--- a/Engine/Source/Runtime/Slate/Public/Widgets/Input/SVectorInputBox.h
+++ b/Engine/Source/Runtime/Slate/Public/Widgets/Input/SVectorInputBox.h
@@ -472,10 +472,19 @@ private:
 	 */
 	void ConstructX(const FArguments& InArgs, TSharedRef<SHorizontalBox> HorizontalBox)
 	{
+		FText DisplayName;
+		if constexpr (NumberOfComponents == 3)
+		{
+			DisplayName = AxisDisplayInfo::GetAxisDisplayName(EAxisList::X);
+		}
+		else
+		{
+			DisplayName = NSLOCTEXT("SVectorInputBox", "X_DisplayName", "X");
+		}
 		ConstructComponent(0,
 			InArgs,
 			AxisDisplayInfo::GetAxisColor(EAxisList::X),
-			AxisDisplayInfo::GetAxisDisplayName(EAxisList::X),
+			DisplayName,
 			HorizontalBox,
 			InArgs._X,
 			InArgs._OnXChanged,
@@ -493,10 +502,19 @@ private:
 	 */
 	void ConstructY(const FArguments& InArgs, TSharedRef<SHorizontalBox> HorizontalBox)
 	{
+		FText DisplayName;
+		if constexpr (NumberOfComponents == 3)
+		{
+			DisplayName = AxisDisplayInfo::GetAxisDisplayName(EAxisList::Y);
+		}
+		else
+		{
+			DisplayName = NSLOCTEXT("SVectorInputBox", "Y_DisplayName", "Y");
+		}
 		ConstructComponent(1,
 			InArgs,
 			AxisDisplayInfo::GetAxisColor(EAxisList::Y),
-			AxisDisplayInfo::GetAxisDisplayName(EAxisList::Y),
+			DisplayName,
 			HorizontalBox,
 			InArgs._Y,
 			InArgs._OnYChanged,
@@ -514,10 +532,19 @@ private:
 	 */
 	void ConstructZ(const FArguments& InArgs, TSharedRef<SHorizontalBox> HorizontalBox)
 	{
+		FText DisplayName;
+		if constexpr (NumberOfComponents == 3)
+		{
+			DisplayName = AxisDisplayInfo::GetAxisDisplayName(EAxisList::Z);
+		}
+		else
+		{
+			DisplayName = NSLOCTEXT("SVectorInputBox", "Z_DisplayName", "Z");
+		}
 		ConstructComponent(2,
 			InArgs,
 			AxisDisplayInfo::GetAxisColor(EAxisList::Z),
-			AxisDisplayInfo::GetAxisDisplayName(EAxisList::Z),
+			DisplayName,
 			HorizontalBox,
 			InArgs._Z,
 			InArgs._OnZChanged,
