[Slate] Add missing virtual destructor for ISlateInputManager

#rnx

diff --git a/Engine/Source/Runtime/Slate/Public/Framework/Application/SlateApplication.h b/Engine/Source/Runtime/Slate/Public/Framework/Application/SlateApplication.h
--- a/Engine/Source/Runtime/Slate/Public/Framework/Application/SlateApplication.h
+++ b/Engine/Source/Runtime/Slate/Public/Framework/Application/SlateApplication.h
@@ -119,6 +119,9 @@ class FPopupSupport
 class ISlateInputManager
 {
 public:
+
+	virtual ~ISlateInputManager() = default;
+	
 	virtual int32 GetUserIndexForMouse() const = 0;
 	virtual int32 GetUserIndexForKeyboard() const = 0;
 
