Toolbar Toggle Buttons and Radio Buttons with submenus revert to being combo buttons when there is no bound action. Check state icon coloration remains supported.
- Deprecate non-const version of IMultiBlockBaseWidget::GetMultiBlockLocation() and add the needed PRAGMA calls to work around the internal warnings.

diff --git a/Engine/Source/Runtime/Slate/Private/Framework/MultiBox/MultiBox.cpp b/Engine/Source/Runtime/Slate/Private/Framework/MultiBox/MultiBox.cpp
--- a/Engine/Source/Runtime/Slate/Private/Framework/MultiBox/MultiBox.cpp
+++ b/Engine/Source/Runtime/Slate/Private/Framework/MultiBox/MultiBox.cpp
@@ -181,6 +181,11 @@ EMultiBlockLocation::Type SMultiBlockBaseWidget::GetMultiBlockLocation()
 	return Location;
 }
 
+EMultiBlockLocation::Type SMultiBlockBaseWidget::GetMultiBlockLocation() const
+{
+	return Location;
+}
+
 void SMultiBlockBaseWidget::OnDragEnter( const FGeometry& MyGeometry, const FDragDropEvent& DragDropEvent )
 {
 	if ( DragDropEvent.GetOperationAs<FUICommandDragDropOp>().IsValid() )
diff --git a/Engine/Source/Runtime/Slate/Private/Framework/MultiBox/SToolBarButtonBlock.cpp b/Engine/Source/Runtime/Slate/Private/Framework/MultiBox/SToolBarButtonBlock.cpp
--- a/Engine/Source/Runtime/Slate/Private/Framework/MultiBox/SToolBarButtonBlock.cpp
+++ b/Engine/Source/Runtime/Slate/Private/Framework/MultiBox/SToolBarButtonBlock.cpp
@@ -364,7 +364,9 @@ void SToolBarButtonBlock::BuildMultiBlockWidget(const ISlateStyle* StyleSet, con
 			];
 		}
 
+	PRAGMA_DISABLE_DEPRECATION_WARNINGS
 	EMultiBlockLocation::Type BlockLocation = GetMultiBlockLocation();
+	PRAGMA_ENABLE_DEPRECATION_WARNINGS
 	
 	// What type of UI should we create for this block?
 	TWeakPtr<const FUICommandInfo> Action = ToolBarButtonBlock->GetAction();
diff --git a/Engine/Source/Runtime/Slate/Private/Framework/MultiBox/SToolBarComboButtonBlock.cpp b/Engine/Source/Runtime/Slate/Private/Framework/MultiBox/SToolBarComboButtonBlock.cpp
--- a/Engine/Source/Runtime/Slate/Private/Framework/MultiBox/SToolBarComboButtonBlock.cpp
+++ b/Engine/Source/Runtime/Slate/Private/Framework/MultiBox/SToolBarComboButtonBlock.cpp
@@ -217,7 +217,9 @@ void SToolBarComboButtonBlock::BuildMultiBlockWidget(const ISlateStyle* StyleSet
 		// clang-format on
 	}
 
+	PRAGMA_DISABLE_DEPRECATION_WARNINGS
 	EMultiBlockLocation::Type BlockLocation = GetMultiBlockLocation();
+	PRAGMA_ENABLE_DEPRECATION_WARNINGS
 	FName BlockStyle = EMultiBlockLocation::ToName(ISlateStyle::Join(StyleName, ".Button"), BlockLocation);
 	const FButtonStyle* ButtonStyle = BlockLocation == EMultiBlockLocation::None ? &ToolBarStyle.ButtonStyle : &StyleSet->GetWidgetStyle<FButtonStyle>(BlockStyle);
 	const FComboButtonStyle* ComboStyle = &ToolBarStyle.ComboButtonStyle;
@@ -246,7 +248,18 @@ void SToolBarComboButtonBlock::BuildMultiBlockWidget(const ISlateStyle* StyleSet
 		// If we have a UICommand, then this is specified in the command.
 		UserInterfaceType = UICommand->GetUserInterfaceType();
 	}
-
+	
+	if ((UserInterfaceType == EUserInterfaceActionType::ToggleButton || UserInterfaceType == EUserInterfaceActionType::RadioButton)
+		&& !((MultiBlock->GetActionList().IsValid() && MultiBlock->GetAction().IsValid()) || MultiBlock->GetDirectActions().IsBound()))
+	{
+		// There are no bound actions; a toggle or radio button wouldn't actually do anything.
+		// Interpret this button as a normal button.
+		UserInterfaceType = EUserInterfaceActionType::Button;
+	}
+	
+	// Cache the checkbox style for icon coloration
+	CheckBoxStyle = GetCheckBoxStyle(StyleSet, StyleName);
+	
 	if (UserInterfaceType == EUserInterfaceActionType::Button)
 	{
 		// clang-format off
@@ -275,31 +288,10 @@ void SToolBarComboButtonBlock::BuildMultiBlockWidget(const ISlateStyle* StyleSet
 				 || UserInterfaceType == EUserInterfaceActionType::RadioButton
 			 ))
 	{
-		FName ToggleBlockStyle = EMultiBlockLocation::ToName(ISlateStyle::Join(StyleName, ".ToggleButton"), BlockLocation);
-
-		const FCheckBoxStyle* CheckStyle;
-		if (!MultiBlock->HasIcon())
-		{
-			CheckStyle = &FCoreStyle::Get().GetWidgetStyle<FCheckBoxStyle>("Checkbox");
-		}
-		else if (BlockLocation == EMultiBlockLocation::None)
-		{
-			CheckStyle = &ToolBarStyle.ToggleButton;
-		}
-		else
-		{
-			CheckStyle = &StyleSet->GetWidgetStyle<FCheckBoxStyle>(ToggleBlockStyle);
-		}
-
-		if (OptionsBlockWidget.IsValid())
-		{
-			CheckStyle = &ToolBarStyle.SettingsToggleButton;
-		}
-
 		// clang-format off
 		SAssignNew(CheckBoxWidget, SCheckBox)
 			// Use the tool bar style for this check box
-			.Style(CheckStyle)
+			.Style(CheckBoxStyle)
 			.CheckBoxContentUsesAutoWidth(false)
 			.OnCheckStateChanged(this, &SToolBarComboButtonBlock::OnCheckStateChanged)
 			.OnGetMenuContent(this, &SToolBarComboButtonBlock::OnGetMenuContent)
@@ -414,6 +406,34 @@ void SToolBarComboButtonBlock::OnCheckStateChanged(const ECheckBoxState NewCheck
 	OnClicked();
 }
 
+const FCheckBoxStyle* SToolBarComboButtonBlock::GetCheckBoxStyle(const ISlateStyle* StyleSet, const FName& StyleName) const
+{
+	const FToolBarStyle& ToolBarStyle = StyleSet->GetWidgetStyle<FToolBarStyle>(StyleName);
+	EMultiBlockLocation::Type BlockLocation = GetMultiBlockLocation();
+
+	const FCheckBoxStyle* CheckStyle;
+	
+	if (!MultiBlock->HasIcon())
+	{
+		CheckStyle = &FCoreStyle::Get().GetWidgetStyle<FCheckBoxStyle>("Checkbox");
+	}
+	else if (BlockLocation == EMultiBlockLocation::None)
+	{
+		CheckStyle = &ToolBarStyle.ToggleButton;
+	}
+	else
+	{
+		CheckStyle = &StyleSet->GetWidgetStyle<FCheckBoxStyle>(EMultiBlockLocation::ToName(ISlateStyle::Join(StyleName, ".ToggleButton"), BlockLocation));
+	}
+
+	if (OptionsBlockWidget.IsValid())
+	{
+		CheckStyle = &ToolBarStyle.SettingsToggleButton;
+	}
+	
+	return CheckStyle;
+}
+
 ECheckBoxState SToolBarComboButtonBlock::GetCheckState() const
 {
 	TSharedPtr<const FUICommandList> ActionList = MultiBlock->GetActionList();
@@ -502,7 +522,20 @@ FSlateColor SToolBarComboButtonBlock::GetIconForegroundColor() const
 	{
 		return FLinearColor::White;
 	}
-
+	
+	if (CheckBoxStyle)
+	{
+		switch (GetCheckState())
+		{
+		case ECheckBoxState::Unchecked:
+			return CheckBoxStyle->ForegroundColor;
+		case ECheckBoxState::Checked:
+			return CheckBoxStyle->CheckedForeground;
+		case ECheckBoxState::Undetermined:
+			return CheckBoxStyle->UndeterminedForeground;
+		}
+	}
+	
 	return FSlateColor::UseForeground();
 }
 
diff --git a/Engine/Source/Runtime/Slate/Private/Framework/MultiBox/SToolBarStackButtonBlock.cpp b/Engine/Source/Runtime/Slate/Private/Framework/MultiBox/SToolBarStackButtonBlock.cpp
--- a/Engine/Source/Runtime/Slate/Private/Framework/MultiBox/SToolBarStackButtonBlock.cpp
+++ b/Engine/Source/Runtime/Slate/Private/Framework/MultiBox/SToolBarStackButtonBlock.cpp
@@ -262,7 +262,9 @@ void SToolBarStackButtonBlock::BuildMultiBlockWidget(const ISlateStyle* StyleSet
 	}
 
 	// Create the content for our button
+	PRAGMA_DISABLE_DEPRECATION_WARNINGS
 	EMultiBlockLocation::Type BlockLocation = GetMultiBlockLocation();
+	PRAGMA_ENABLE_DEPRECATION_WARNINGS
 	FName BlockStyle = EMultiBlockLocation::ToName(ISlateStyle::Join(StyleName, ".Button"), BlockLocation);
 	const FComboButtonStyle* ComboButtonStyle = &FAppStyle::Get().GetWidgetStyle<FComboButtonStyle>("SimpleComboButton");
 	const FButtonStyle* ButtonStyle = BlockLocation == EMultiBlockLocation::None ? &ToolBarStyle.ButtonStyle : &StyleSet->GetWidgetStyle<FButtonStyle>(BlockStyle);
diff --git a/Engine/Source/Runtime/Slate/Public/Framework/MultiBox/MultiBox.h b/Engine/Source/Runtime/Slate/Public/Framework/MultiBox/MultiBox.h
--- a/Engine/Source/Runtime/Slate/Public/Framework/MultiBox/MultiBox.h
+++ b/Engine/Source/Runtime/Slate/Public/Framework/MultiBox/MultiBox.h
@@ -593,7 +593,13 @@ public:
 	/**
 	 * Returns this MultiBlocks location
 	 */
+	UE_DEPRECATED(5.6, "Use the const version of GetMultiBlockLocation")
 	virtual EMultiBlockLocation::Type GetMultiBlockLocation() = 0;
+	
+	/**
+	 * Returns this MultiBlocks location
+	 */
+	virtual EMultiBlockLocation::Type GetMultiBlockLocation() const = 0;
 
 	/**
 	 * Returns true if editing this widget
@@ -621,7 +627,9 @@ public:
 	SLATE_API virtual void SetMultiBlock(TSharedRef<const FMultiBlock> InMultiBlock) override;
 	SLATE_API virtual void SetOptionsBlockWidget(TSharedPtr<SWidget> InOptionsBlock) override;
 	SLATE_API virtual void SetMultiBlockLocation(EMultiBlockLocation::Type InLocation, bool bInSectionContainsIcons) override;
+	UE_DEPRECATED(5.6, "Use the const version of GetMultiBlockLocation")
 	SLATE_API virtual EMultiBlockLocation::Type GetMultiBlockLocation() override;
+	SLATE_API virtual EMultiBlockLocation::Type GetMultiBlockLocation() const override;
 	SLATE_API virtual bool IsInEditMode() const override;
 
 	/** SWidget Interface */
diff --git a/Engine/Source/Runtime/Slate/Public/Framework/MultiBox/SToolBarComboButtonBlock.h b/Engine/Source/Runtime/Slate/Public/Framework/MultiBox/SToolBarComboButtonBlock.h
--- a/Engine/Source/Runtime/Slate/Public/Framework/MultiBox/SToolBarComboButtonBlock.h
+++ b/Engine/Source/Runtime/Slate/Public/Framework/MultiBox/SToolBarComboButtonBlock.h
@@ -183,6 +183,11 @@ protected:
 	 */
 	SLATE_API bool IsEnabled() const;
 
+	/**
+	 * @return Gets the effective checkbox style for this widget. 
+	 */
+	const FCheckBoxStyle* GetCheckBoxStyle(const ISlateStyle* StyleSet, const FName& StyleName) const;
+	
 	/**
 	 * Called by Slate to determine if this button is visible
 	 *
@@ -231,6 +236,9 @@ private:
 
 	/** The hovered style for the entire block used when the block has multiple widgets */
 	const FSlateBrush* BlockHovered = nullptr;
+	
+	/** The checkbox style to be used for simulating checkbox foreground color. */
+	const FCheckBoxStyle* CheckBoxStyle = nullptr;
 
 	/** Whether this toolbar should always use small icons, regardless of the current settings */
 	bool bForceSmallIcons;
