SNumericVectorInputBox : Add attributes for XDisplayName, YDisplayName, ZDisplayName, and WDisplayName

diff --git a/Engine/Source/Runtime/Slate/Public/Widgets/Input/SVectorInputBox.h b/Engine/Source/Runtime/Slate/Public/Widgets/Input/SVectorInputBox.h
--- a/Engine/Source/Runtime/Slate/Public/Widgets/Input/SVectorInputBox.h
+++ b/Engine/Source/Runtime/Slate/Public/Widgets/Input/SVectorInputBox.h
@@ -54,7 +54,16 @@ private:
 
 		FORCENOINLINE FVectorXArguments()
 			: _ToggleXChecked(ECheckBoxState::Checked)
-		{}
+			, _XDisplayName(NSLOCTEXT("SVectorInputBox", "X_DisplayName", "X"))
+		{
+			if constexpr (NumberOfComponents == 3)
+			{
+				if (AxisDisplayInfo::UseForwardRightUpDisplayNames())
+				{
+					_XDisplayName.Set(AxisDisplayInfo::GetAxisDisplayName(EAxisList::X));
+				}
+			}
+		}
 
 		/** X Component of the vector */
 		SLATE_ATTRIBUTE(TOptional<NumericType>, X)
@@ -79,6 +88,9 @@ private:
 
 		/** Called when the x value of the vector slider ended movement */
 		SLATE_EVENT(FOnNumericValueChanged, OnXEndSliderMovement)
+
+		/** The DisplayName of the X component of the vector */
+		SLATE_ATTRIBUTE(FText, XDisplayName)
 	};
 
 	struct FVectorYArgumentsEmpty {};
@@ -89,7 +101,16 @@ private:
 
 		FORCENOINLINE FVectorYArguments()
 			: _ToggleYChecked(ECheckBoxState::Checked)
-		{}
+			, _YDisplayName(NSLOCTEXT("SVectorInputBox", "Y_DisplayName", "Y"))
+		{
+			if constexpr (NumberOfComponents == 3)
+			{
+				if (AxisDisplayInfo::UseForwardRightUpDisplayNames())
+				{
+					_YDisplayName.Set(AxisDisplayInfo::GetAxisDisplayName(EAxisList::Y));
+				}
+			}
+		}
 
 		/** Y Component of the vector */
 		SLATE_ATTRIBUTE(TOptional<NumericType>, Y)
@@ -114,6 +135,9 @@ private:
 
 		/** Called when the y value of the vector slider ended movement */
 		SLATE_EVENT(FOnNumericValueChanged, OnYEndSliderMovement)
+
+		/** The DisplayName of the Y component of the vector */
+		SLATE_ATTRIBUTE(FText, YDisplayName)
 	};
 
 	struct FVectorZArgumentsEmpty {};
@@ -124,7 +148,16 @@ private:
 
 		FORCENOINLINE FVectorZArguments()
 			: _ToggleZChecked(ECheckBoxState::Checked)
-		{}
+			, _ZDisplayName(NSLOCTEXT("SVectorInputBox", "Z_DisplayName", "Z"))
+		{
+			if constexpr (NumberOfComponents == 3)
+			{
+				if (AxisDisplayInfo::UseForwardRightUpDisplayNames())
+				{
+					_ZDisplayName.Set(AxisDisplayInfo::GetAxisDisplayName(EAxisList::Z));
+				}
+			}
+		}
 
 		/** Z Component of the vector */
 		SLATE_ATTRIBUTE(TOptional<NumericType>, Z)
@@ -149,6 +182,9 @@ private:
 
 		/** Called when the z value of the vector slider ended movement */
 		SLATE_EVENT(FOnNumericValueChanged, OnZEndSliderMovement)
+
+		/** The DisplayName of the Z component of the vector */
+		SLATE_ATTRIBUTE(FText, ZDisplayName)
 	};
 
 	struct FVectorWArgumentsEmpty {};
@@ -159,6 +195,7 @@ private:
 
 		FORCENOINLINE FVectorWArguments()
 			: _ToggleWChecked(ECheckBoxState::Checked)
+			, _WDisplayName(NSLOCTEXT("SVectorInputBox", "W_DisplayName", "W"))
 		{}
 
 		/** W Component of the vector */
@@ -184,6 +221,9 @@ private:
 
 		/** Called when the w value of the vector slider ended movement */
 		SLATE_EVENT(FOnNumericValueChanged, OnWEndSliderMovement)
+
+		/** The DisplayName of the W component of the vector */
+		SLATE_ATTRIBUTE(FText, WDisplayName)
 	};
 
 public:
@@ -405,7 +445,7 @@ private:
 	void ConstructComponent(int32 ComponentIndex,
 		const FArguments& InArgs,
 		const FLinearColor& LabelColor,
-		const FText& ComponentDisplayName,
+		const TAttribute<FText>& ComponentDisplayName,
 		TSharedRef<SHorizontalBox>& HorizontalBox,
 		const TAttribute<TOptional<NumericType>>& Component,
 		const FOnNumericValueChanged& OnComponentChanged,
@@ -423,8 +463,11 @@ private:
 		}
 
 		TStringBuilder<32> ToolTipTextFormatString;
-		ToolTipTextFormatString.Append(ComponentDisplayName.ToString());
-		ToolTipTextFormatString.Append(TEXT(": {0}"));
+		if (ComponentDisplayName.IsSet())
+		{			
+			ToolTipTextFormatString.Append(FText::Format(NSLOCTEXT("SVectorInputBox", "ToolTipTextFormatDisplayName", "{0}: "), ComponentDisplayName.Get()).ToString());
+		}
+		ToolTipTextFormatString.Append(TEXT("{0}"));
 
 		TAttribute<TOptional<NumericType>> Value = CreatePerComponentGetter(ComponentIndex, Component, VectorAttribute);
 
@@ -472,19 +515,10 @@ private:
 	 */
 	void ConstructX(const FArguments& InArgs, TSharedRef<SHorizontalBox> HorizontalBox)
 	{
-		FText DisplayName;
-		if constexpr (NumberOfComponents == 3)
-		{
-			DisplayName = AxisDisplayInfo::GetAxisDisplayName(EAxisList::X);
-		}
-		else
-		{
-			DisplayName = NSLOCTEXT("SVectorInputBox", "X_DisplayName", "X");
-		}
 		ConstructComponent(0,
 			InArgs,
 			AxisDisplayInfo::GetAxisColor(EAxisList::X),
-			DisplayName,
+			InArgs._XDisplayName,
 			HorizontalBox,
 			InArgs._X,
 			InArgs._OnXChanged,
@@ -502,19 +536,10 @@ private:
 	 */
 	void ConstructY(const FArguments& InArgs, TSharedRef<SHorizontalBox> HorizontalBox)
 	{
-		FText DisplayName;
-		if constexpr (NumberOfComponents == 3)
-		{
-			DisplayName = AxisDisplayInfo::GetAxisDisplayName(EAxisList::Y);
-		}
-		else
-		{
-			DisplayName = NSLOCTEXT("SVectorInputBox", "Y_DisplayName", "Y");
-		}
 		ConstructComponent(1,
 			InArgs,
 			AxisDisplayInfo::GetAxisColor(EAxisList::Y),
-			DisplayName,
+			InArgs._YDisplayName,
 			HorizontalBox,
 			InArgs._Y,
 			InArgs._OnYChanged,
@@ -532,19 +557,10 @@ private:
 	 */
 	void ConstructZ(const FArguments& InArgs, TSharedRef<SHorizontalBox> HorizontalBox)
 	{
-		FText DisplayName;
-		if constexpr (NumberOfComponents == 3)
-		{
-			DisplayName = AxisDisplayInfo::GetAxisDisplayName(EAxisList::Z);
-		}
-		else
-		{
-			DisplayName = NSLOCTEXT("SVectorInputBox", "Z_DisplayName", "Z");
-		}
 		ConstructComponent(2,
 			InArgs,
 			AxisDisplayInfo::GetAxisColor(EAxisList::Z),
-			DisplayName,
+			InArgs._ZDisplayName,
 			HorizontalBox,
 			InArgs._Z,
 			InArgs._OnZChanged,
@@ -565,7 +581,7 @@ private:
 		ConstructComponent(3,
 			InArgs,
 			SNumericEntryBox<NumericType>::LilacLabelBackgroundColor,
-			NSLOCTEXT("SVectorInputBox", "W_DisplayName", "W"),
+			InArgs._WDisplayName,
 			HorizontalBox,
 			InArgs._W,
 			InArgs._OnWChanged,
