Viewport Toolbar: Convert coordinates space combo button to an interactive combo button that cycles through the different spaces.

diff --git a/Engine/Source/Developer/ToolMenus/Private/ToolMenus.cpp b/Engine/Source/Developer/ToolMenus/Private/ToolMenus.cpp
--- a/Engine/Source/Developer/ToolMenus/Private/ToolMenus.cpp
+++ b/Engine/Source/Developer/ToolMenus/Private/ToolMenus.cpp
@@ -1767,7 +1767,17 @@ void UToolMenus::PopulateToolBarBuilderWithEntry(
 		);
 	}
 
-	const FUIAction UIAction = UToolMenus::ConvertUIAction(Block, MenuData->Context);
+	const FUIAction UIAction = [&Block, MenuData]()
+		{
+			if (Block.ToolBarData.ActionOverride.IsSet())
+			{
+				return UToolMenus::ConvertUIAction(Block.ToolBarData.ActionOverride.GetValue(), MenuData->Context);
+			}
+			else
+			{
+				return UToolMenus::ConvertUIAction(Block, MenuData->Context);
+			}
+		}();
 
 	TAttribute<FText> ToolbarLabelOverride;
 	if (Block.ToolBarData.LabelOverride.IsSet())
diff --git a/Engine/Source/Developer/ToolMenus/Public/ToolMenuEntry.h b/Engine/Source/Developer/ToolMenus/Public/ToolMenuEntry.h
--- a/Engine/Source/Developer/ToolMenus/Public/ToolMenuEntry.h
+++ b/Engine/Source/Developer/ToolMenus/Public/ToolMenuEntry.h
@@ -83,6 +83,9 @@ public:
 	TAttribute<EMenuPlacement> PlacementOverride;
 
 	FMenuEntryResizeParams ResizeParams;
+
+	/** Optional action override to use when the entry appears in a toolbar */
+	TOptional<FToolUIAction> ActionOverride;
 };
 
 
diff --git a/Engine/Source/Editor/UnrealEd/Private/ViewportToolbar/UnrealEdViewportToolbar.cpp b/Engine/Source/Editor/UnrealEd/Private/ViewportToolbar/UnrealEdViewportToolbar.cpp
--- a/Engine/Source/Editor/UnrealEd/Private/ViewportToolbar/UnrealEdViewportToolbar.cpp
+++ b/Engine/Source/Editor/UnrealEd/Private/ViewportToolbar/UnrealEdViewportToolbar.cpp
@@ -866,7 +866,20 @@ FToolMenuEntry CreateTransformsSubmenu()
 											);
 										}
 
+										FToolUIAction CycleCoordSystemAction;
+										{
+											CycleCoordSystemAction.ExecuteAction = FToolMenuExecuteAction::CreateLambda(
+												[](const FToolMenuContext& InContext)
+												{
+													if (const FUIAction* Action = InContext.GetActionForCommand(FEditorViewportCommands::Get().CycleTransformGizmoCoordSystem))
+													{
+														Action->Execute();
+													}
+												});
+										}
+
 										CoordinateSystemSubmenu.ToolBarData.LabelOverride = FText();
+										CoordinateSystemSubmenu.ToolBarData.ActionOverride = CycleCoordSystemAction;
 										CoordinateSystemSubmenu.SetShowInToolbarTopLevel(true);
 									}
 								}
diff --git a/Engine/Source/Runtime/Slate/Private/Framework/MultiBox/SToolBarComboButtonBlock.cpp b/Engine/Source/Runtime/Slate/Private/Framework/MultiBox/SToolBarComboButtonBlock.cpp
--- a/Engine/Source/Runtime/Slate/Private/Framework/MultiBox/SToolBarComboButtonBlock.cpp
+++ b/Engine/Source/Runtime/Slate/Private/Framework/MultiBox/SToolBarComboButtonBlock.cpp
@@ -12,6 +12,7 @@
 #include "Widgets/Layout/SBox.h"
 #include "Styling/ToolBarStyle.h"
 #include "Widgets/Images/SLayeredImage.h"
+#include "Widgets/Input/SButton.h"
 #include "Widgets/SToolTip.h"
 
 FToolBarComboButtonBlock::FToolBarComboButtonBlock(
@@ -153,6 +154,14 @@ void SToolBarComboButtonBlock::BuildMultiBlockWidget(const ISlateStyle* StyleSet
 	if (MultiBox->GetType() == EMultiBoxType::SlimHorizontalToolBar
 		|| MultiBox->GetType() == EMultiBoxType::SlimHorizontalUniformToolBar)
 	{
+		TSharedPtr<SWidget> ButtonIcon = IconWidget;
+
+		if (ToolBarComboButtonBlock->UserInterfaceActionType == EUserInterfaceActionType::Button && HasAction())
+		{
+			// Don't add an icon to the SComboButton if we're a button with an action, since it'll have a separate SButton to host the icon
+			ButtonIcon = SNullWidget::NullWidget;
+		}
+
 		// clang-format off
 		ButtonContent =
 			SNew(SHorizontalBox)
@@ -162,7 +171,7 @@ void SToolBarComboButtonBlock::BuildMultiBlockWidget(const ISlateStyle* StyleSet
 			.VAlign(VAlign_Center)
 			.HAlign(HAlign_Center)
 			[
-				IconWidget
+				ButtonIcon.ToSharedRef()
 			]
 			// Label text
 			+ SHorizontalBox::Slot()
@@ -250,7 +259,7 @@ void SToolBarComboButtonBlock::BuildMultiBlockWidget(const ISlateStyle* StyleSet
 	}
 	
 	if ((UserInterfaceType == EUserInterfaceActionType::ToggleButton || UserInterfaceType == EUserInterfaceActionType::RadioButton)
-		&& !((MultiBlock->GetActionList().IsValid() && MultiBlock->GetAction().IsValid()) || MultiBlock->GetDirectActions().IsBound()))
+		&& !HasAction())
 	{
 		// There are no bound actions; a toggle or radio button wouldn't actually do anything.
 		// Interpret this button as a normal button.
@@ -260,25 +269,42 @@ void SToolBarComboButtonBlock::BuildMultiBlockWidget(const ISlateStyle* StyleSet
 	// Cache the checkbox style for icon coloration
 	CheckBoxStyle = GetCheckBoxStyle(StyleSet, StyleName);
 	
+	LeftHandSideWidget = SNullWidget::NullWidget;
+
 	if (UserInterfaceType == EUserInterfaceActionType::Button)
 	{
-		// clang-format off
-		ChildSlot
-		[
-			SAssignNew(ComboButtonWidget, SComboButton)
-			.AddMetaData<FTagMetaData>(FTagMetaData(TagName))
-			.ContentPadding(0.f)
-			.ComboButtonStyle(ComboStyle)
+		if (HasAction())
+		{
+			ComboStyle = &ToolBarStyle.SettingsComboButton;
+
+			// clang-format off
+			SAssignNew(LeftHandSideWidget, SButton)
 			.ButtonStyle(ButtonStyle)
-			.ToolTipText(ToolBarComboButtonBlock->ToolTip)
-			.MenuPlacement(ToolBarComboButtonBlock->PlacementOverride)
-			.ForegroundColor(this, &SToolBarComboButtonBlock::OnGetForegroundColor)
-			// Route the content generator event
-			.OnGetMenuContent(this, &SToolBarComboButtonBlock::OnGetMenuContent)
-			.ButtonContent()
+			.HAlign(HAlign_Center)
+			.VAlign(VAlign_Center)
+			.ContentPadding(0.f)
+			.OnClicked(this, &SToolBarComboButtonBlock::OnClicked)
+			.IsEnabled(this, &SToolBarComboButtonBlock::IsEnabled)
 			[
-				ButtonContent
-			]
+				IconWidget
+			];
+			// clang-format on
+		}
+
+		// clang-format off
+		SAssignNew(ComboButtonWidget, SComboButton)
+		.AddMetaData<FTagMetaData>(FTagMetaData(TagName))
+		.ContentPadding(0.f)
+		.ComboButtonStyle(ComboStyle)
+		.ButtonStyle(ButtonStyle)
+		.ToolTipText(ToolBarComboButtonBlock->ToolTip)
+		.MenuPlacement(ToolBarComboButtonBlock->PlacementOverride)
+		.ForegroundColor(this, &SToolBarComboButtonBlock::OnGetForegroundColor)
+		// Route the content generator event
+		.OnGetMenuContent(this, &SToolBarComboButtonBlock::OnGetMenuContent)
+		.ButtonContent()
+		[
+			ButtonContent
 		];
 		// clang-format on
 	}
@@ -289,7 +315,7 @@ void SToolBarComboButtonBlock::BuildMultiBlockWidget(const ISlateStyle* StyleSet
 			 ))
 	{
 		// clang-format off
-		SAssignNew(CheckBoxWidget, SCheckBox)
+		SAssignNew(LeftHandSideWidget, SCheckBox)
 			// Use the tool bar style for this check box
 			.Style(CheckBoxStyle)
 			.CheckBoxContentUsesAutoWidth(false)
@@ -322,7 +348,10 @@ void SToolBarComboButtonBlock::BuildMultiBlockWidget(const ISlateStyle* StyleSet
 				.Margin(0)
 			];
 		// clang-format on
+	}
 
+	if (LeftHandSideWidget && ComboButtonWidget)
+	{
 		// clang-format off
 		ChildSlot
 		[
@@ -338,7 +367,7 @@ void SToolBarComboButtonBlock::BuildMultiBlockWidget(const ISlateStyle* StyleSet
 				.AutoWidth()
 				.VAlign(VAlign_Fill)
 				[
-					CheckBoxWidget.ToSharedRef()
+					LeftHandSideWidget.ToSharedRef()
 				]
 				+ SHorizontalBox::Slot()
 				.AutoWidth()
@@ -465,6 +494,12 @@ bool SToolBarComboButtonBlock::IsEnabled() const
 	return true;
 }
 
+bool SToolBarComboButtonBlock::HasAction() const
+{
+	return (MultiBlock->GetActionList().IsValid() && MultiBlock->GetAction().IsValid()) ||
+		MultiBlock->GetDirectActions().IsBound();
+}
+
 EVisibility SToolBarComboButtonBlock::GetVisibility() const
 {
 	// Let the visibility override take prescedence here.
@@ -565,7 +600,8 @@ FSlateColor SToolBarComboButtonBlock::OnGetForegroundColor() const
 
 const FSlateBrush* SToolBarComboButtonBlock::GetBorderImage() const
 {
-	if ((ComboButtonWidget && ComboButtonWidget->IsHovered()) || (CheckBoxWidget && CheckBoxWidget->IsHovered()))
+	if ((ComboButtonWidget && ComboButtonWidget->IsHovered()) ||
+		(LeftHandSideWidget && LeftHandSideWidget->IsHovered()))
 	{
 		return BlockHovered;
 	}
diff --git a/Engine/Source/Runtime/Slate/Public/Framework/MultiBox/SToolBarComboButtonBlock.h b/Engine/Source/Runtime/Slate/Public/Framework/MultiBox/SToolBarComboButtonBlock.h
--- a/Engine/Source/Runtime/Slate/Public/Framework/MultiBox/SToolBarComboButtonBlock.h
+++ b/Engine/Source/Runtime/Slate/Public/Framework/MultiBox/SToolBarComboButtonBlock.h
@@ -26,13 +26,13 @@ public:
 	 *
 	 * @param	InAction					UI action that sets the enabled state for this combo button
 	 * @param	InMenuContentGenerator		Delegate that generates a widget for this combo button's menu content.
-	 * Called when the menu is summoned.
+	 *										Called when the menu is summoned.
 	 * @param	InLabel						Optional label for this combo button.
 	 * @param	InToolTip					Tool tip string (required!)
 	 * @param	InIcon						Optional icon to use for the tool bar image
 	 * @param	bInSimpleComboBox			If true, the icon and label won't be displayed
 	 * @param	InToolbarLabelOverride		Optional label to use when the block appears in a toolbar. If omitted, then the
-	 * label override or command name will be used instead.
+	 *										label override or command name will be used instead.
 	 * @param	InPlacementOverride			Optional override placement to customize menu placement once opened via e.g. a toolbar menu button.
 	 * @param	InUserInterfaceActionType	Optional type of user interface action to use. If a command is bound, the user interface action type associated with the command will be used instead.
 	 */
@@ -183,6 +183,9 @@ protected:
 	 */
 	SLATE_API bool IsEnabled() const;
 
+	/** True if we have an action bound */
+	bool HasAction() const;
+
 	/**
 	 * @return Gets the effective checkbox style for this widget. 
 	 */
@@ -228,7 +231,7 @@ private:
 	/** Optional overridden icon for this tool bar button.  IF not set, then the action's icon will be used instead. */
 	TAttribute<FSlateIcon> Icon;
 
-	TSharedPtr<SCheckBox> CheckBoxWidget;
+	TSharedPtr<SWidget> LeftHandSideWidget; // Widget that comes before the combo button
 	TSharedPtr<SComboButton> ComboButtonWidget;
 
 	/** The foreground color for button when the combo button is open */
