Rolled back to the previous version of SColorBlock.

There were unexpected side effects in SChooserColumnHandle. The class that the widget references through an attribute lambda hasn't been fully initialized when the widget is created. As a result of the recent change the lambda is no immediately called resulting in a null pointer crash.

#rnx

diff --git a/Engine/Source/Runtime/Slate/Private/Widgets/Colors/SColorBlock.cpp b/Engine/Source/Runtime/Slate/Private/Widgets/Colors/SColorBlock.cpp
--- a/Engine/Source/Runtime/Slate/Private/Widgets/Colors/SColorBlock.cpp
+++ b/Engine/Source/Runtime/Slate/Private/Widgets/Colors/SColorBlock.cpp
@@ -40,15 +40,15 @@ SColorBlock::~SColorBlock() = default;
  */
 void SColorBlock::Construct(const FArguments& InArgs)
 {
-	Color.Assign(*this, InArgs._Color, ESlateAttributeBindAction::Initialize);
-	AlphaBackgroundBrush.Assign(*this, InArgs._AlphaBackgroundBrush, ESlateAttributeBindAction::Initialize);
-	GradientCornerRadius.Assign(*this, InArgs._CornerRadius, FVector4(0.0f), ESlateAttributeBindAction::Initialize);
-	ColorBlockSize.Assign(*this, InArgs._Size, ESlateAttributeBindAction::Initialize);
+	Color.Assign(*this, InArgs._Color);
+	AlphaBackgroundBrush.Assign(*this, InArgs._AlphaBackgroundBrush);
+	GradientCornerRadius.Assign(*this, InArgs._CornerRadius, FVector4(0.0f));
+	ColorBlockSize.Assign(*this, InArgs._Size);
 	MouseButtonDownHandler = InArgs._OnMouseButtonDown;
-	AlphaDisplayMode.Assign(*this, InArgs._AlphaDisplayMode, ESlateAttributeBindAction::Initialize);
-	ColorIsHSV.Assign(*this, InArgs._ColorIsHSV, ESlateAttributeBindAction::Initialize);
-	ShowBackgroundForAlpha.Assign(*this, InArgs._ShowBackgroundForAlpha, ESlateAttributeBindAction::Initialize);
-	bUseSRGB.Assign(*this, InArgs._UseSRGB, ESlateAttributeBindAction::Initialize);
+	AlphaDisplayMode.Assign(*this, InArgs._AlphaDisplayMode);
+	ColorIsHSV.Assign(*this, InArgs._ColorIsHSV);
+	ShowBackgroundForAlpha.Assign(*this, InArgs._ShowBackgroundForAlpha);
+	bUseSRGB.Assign(*this, InArgs._UseSRGB);
 }
 
 int32 SColorBlock::OnPaint(const FPaintArgs& Args, const FGeometry& AllottedGeometry, const FSlateRect& MyCullingRect, FSlateWindowElementList& OutDrawElements, int32 LayerId, const FWidgetStyle& InWidgetStyle, bool bParentEnabled) const
