[STableViewBase] FixedLineScrollOffset: clarify in the documentation how it works with inertial scrolls. Make it Round the offset instead of Flooring it for inertial scrolls.

#rnx

diff --git a/Engine/Source/Runtime/Slate/Private/Widgets/Views/STableViewBase.cpp b/Engine/Source/Runtime/Slate/Private/Widgets/Views/STableViewBase.cpp
--- a/Engine/Source/Runtime/Slate/Private/Widgets/Views/STableViewBase.cpp
+++ b/Engine/Source/Runtime/Slate/Private/Widgets/Views/STableViewBase.cpp
@@ -231,6 +231,7 @@ EActiveTimerReturnType STableViewBase::UpdateInertialScroll(double InCurrentTime
 			if (CanUseInertialScroll(TickScrollDelta))
 			{
 				InertialScrollManager.AddScrollSample(TickScrollDelta, InCurrentTime);
+				bIsInertialScroll = true;
 			}
 		}
 		else
@@ -247,7 +248,7 @@ EActiveTimerReturnType STableViewBase::UpdateInertialScroll(double InCurrentTime
 				}
 				else
 				{
-					InertialScrollManager.ClearScrollVelocity();
+					EndInertialScrolling();
 				}
 			}
 
@@ -421,7 +422,7 @@ FReply STableViewBase::OnPreviewMouseButtonDown( const FGeometry& MyGeometry, co
 	if (bEnableTouchScrolling && MouseEvent.IsTouchEvent())
 	{
 		// Clear any inertia
-		InertialScrollManager.ClearScrollVelocity(true);
+		EndInertialScrolling(true);
 		// We have started a new interaction; track how far the user has moved since they put their finger down.
 		AmountScrolledWhileRightMouseDown = 0;
 
@@ -441,7 +442,7 @@ FReply STableViewBase::OnPreviewMouseButtonDown( const FGeometry& MyGeometry, co
 FReply STableViewBase::OnMouseButtonDown( const FGeometry& MyGeometry, const FPointerEvent& MouseEvent )
 {
 	// Zero the scroll velocity so the list stops immediately on mouse down, even if the user does not drag
-	InertialScrollManager.ClearScrollVelocity(true);
+	EndInertialScrolling(true);
 
 	if (MouseEvent.GetEffectingButton() == EKeys::RightMouseButton)
 	{
@@ -603,7 +604,7 @@ FReply STableViewBase::OnMouseWheel( const FGeometry& MyGeometry, const FPointer
 	if( bIsPointerScrollingEnabled && !MouseEvent.IsControlDown() )
 	{
 		// Make sure scroll velocity is cleared so it doesn't fight with the mouse wheel input
-		this->InertialScrollManager.ClearScrollVelocity();
+		EndInertialScrolling();
 
 		float AmountScrolledInItems = 0.f;
 		if (FixedLineScrollOffset.IsSet())
@@ -835,9 +836,19 @@ double STableViewBase::GetTargetScrollOffset() const
 		// The scroll has ended, we calculate the appropriate the scroll offset to snap on a line.
 		// We don't do that if we are at the end of the list because in this case showing the last line correctly is more important.
 		const int32 NumItemsPerLine = GetNumItemsPerLine();
-		const double DesiredLineOffset = FMath::FloorToDouble(DesiredScrollOffset / NumItemsPerLine) - FixedLineScrollOffset.GetValue();
-		
-		//return FMath::Max(0.0, FMath::CeilToDouble(DesiredScrollOffset) - AdditionalOffset);
+		double DesiredLineOffset = DesiredScrollOffset / NumItemsPerLine;
+		if (bIsInertialScroll)
+		{
+			// Scroll to the item most in view
+			DesiredLineOffset = FMath::RoundToDouble(DesiredLineOffset);
+		}
+		else
+		{
+			// Scroll to the topmost item.
+			DesiredLineOffset = FMath::FloorToDouble(DesiredLineOffset);
+		}
+		DesiredLineOffset -= FixedLineScrollOffset.GetValue();
+
 		return FMath::Max(0.0, DesiredLineOffset * NumItemsPerLine);
 	}
 	return DesiredScrollOffset;
@@ -891,9 +902,10 @@ void STableViewBase::SetScrollOffset( const float InScrollOffset )
 	}
 }
 
-void STableViewBase::EndInertialScrolling()
+void STableViewBase::EndInertialScrolling(const bool bInShouldStopScrollNow)
 {
-	InertialScrollManager.ClearScrollVelocity();
+	InertialScrollManager.ClearScrollVelocity(bInShouldStopScrollNow);
+	bIsInertialScroll = false;
 }
 
 void STableViewBase::AddScrollOffset(const float InScrollOffsetDelta, bool RefreshList)
diff --git a/Engine/Source/Runtime/Slate/Public/Widgets/Views/STableViewBase.h b/Engine/Source/Runtime/Slate/Public/Widgets/Views/STableViewBase.h
--- a/Engine/Source/Runtime/Slate/Public/Widgets/Views/STableViewBase.h
+++ b/Engine/Source/Runtime/Slate/Public/Widgets/Views/STableViewBase.h
@@ -172,8 +172,9 @@ public:
 	/** Set the scroll offset of this view (in items) */
 	SLATE_API void SetScrollOffset( const float InScrollOffset );
 
-	/** Reset the inertial scroll velocity accumulated in the InertialScrollManager */
-	SLATE_API void EndInertialScrolling();
+	/** Reset the inertial scroll velocity accumulated in the InertialScrollManager
+	 * @param bInShouldStopScrollNow true implies the scroll will stop instantly, else the list will scroll until any accumulated scroll offset is cleared. */
+	SLATE_API void EndInertialScrolling(const bool bInShouldStopScrollNow = false);
 
 	/** Add the scroll offset of this view (in items) */
 	SLATE_API void AddScrollOffset(const float InScrollOffsetDelta, bool RefreshList = false);
@@ -446,8 +447,14 @@ protected:
 
 	/** 
 	 * The fixed offset in lines to always apply to the top/left (depending on orientation) of the list.
-	 * If provided, all non-inertial means of scrolling will settle with exactly this offset of the topmost entry.
-	 * Ex: A value of 0.25 would cause the topmost full entry to be offset down by a quarter length of the preceeding entry.
+	 *
+	 * If provided:
+	 *
+	 * - all non-inertial means of scrolling will settle with exactly this offset of the topmost entry.
+	 *   Ex: A value of 0.25 would cause the topmost full entry to be offset down by a quarter length of the preceding entry.
+	 *   
+	 * - all inertial means of scrolling will, at the end of the scroll, settle with exactly this offset of the entry most in view.
+	 *	 Ex: if at the end of the scroll we see 0.25% of a line and 0.75% of next line at the top/left, the list will scroll to show the second line minus the offset.
 	 */
 	TOptional<double> FixedLineScrollOffset;
 
@@ -567,6 +574,9 @@ private:
 	/** Whether or not the view should scroll to a fixed line offset if the scroll velocity is 0 */
 	bool ShouldScrollToFixedLineOffsetAtZeroVelocity() const;
 
+	/** Whether or not the current scroll is inertial */
+	bool bIsInertialScroll = false;
+
 	union
 	{
 		// vertical scroll bar is stored in horizontal box and vice versa
