Simplify SToolBarComboButtonBlock to have fewer branching permutations.

diff --git a/Engine/Source/Runtime/Slate/Private/Framework/MultiBox/SToolBarComboButtonBlock.cpp b/Engine/Source/Runtime/Slate/Private/Framework/MultiBox/SToolBarComboButtonBlock.cpp
--- a/Engine/Source/Runtime/Slate/Private/Framework/MultiBox/SToolBarComboButtonBlock.cpp
+++ b/Engine/Source/Runtime/Slate/Private/Framework/MultiBox/SToolBarComboButtonBlock.cpp
@@ -86,6 +86,7 @@ void SToolBarComboButtonBlock::BuildMultiBlockWidget(const ISlateStyle* StyleSet
 	const FToolBarStyle& ToolBarStyle = StyleSet->GetWidgetStyle<FToolBarStyle>(StyleName);
 
 	// If override is set use that
+	TAttribute<EVisibility> LabelVisibility;
 	if (LabelVisibilityOverride.IsSet())
 	{
 		LabelVisibility = LabelVisibilityOverride.GetValue();
@@ -148,20 +149,17 @@ void SToolBarComboButtonBlock::BuildMultiBlockWidget(const ISlateStyle* StyleSet
 	{
 		TagName = ToolBarComboButtonBlock->GetTutorialHighlightName();
 	}
+	
+	// When an execute action is present, the combo button is split into two parts:
+	// - The (checkbox) button that handles the action
+	// - The ComboButton that opens the menu
+	const bool bIsSplitButton = HasAction();
 
 	// Create the content for our button
 	TSharedRef<SWidget> ButtonContent = SNullWidget::NullWidget;
 	if (MultiBox->GetType() == EMultiBoxType::SlimHorizontalToolBar
 		|| MultiBox->GetType() == EMultiBoxType::SlimHorizontalUniformToolBar)
 	{
-		TSharedPtr<SWidget> ButtonIcon = IconWidget;
-
-		if (ToolBarComboButtonBlock->UserInterfaceActionType == EUserInterfaceActionType::Button && HasAction())
-		{
-			// Don't add an icon to the SComboButton if we're a button with an action, since it'll have a separate SButton to host the icon
-			ButtonIcon = SNullWidget::NullWidget;
-		}
-
 		// clang-format off
 		ButtonContent =
 			SNew(SHorizontalBox)
@@ -171,7 +169,8 @@ void SToolBarComboButtonBlock::BuildMultiBlockWidget(const ISlateStyle* StyleSet
 			.VAlign(VAlign_Center)
 			.HAlign(HAlign_Center)
 			[
-				ButtonIcon.ToSharedRef()
+				// A split button will have the icon handled in its own widget
+				bIsSplitButton ? SNullWidget::NullWidget : IconWidget
 			]
 			// Label text
 			+ SHorizontalBox::Slot()
@@ -179,9 +178,9 @@ void SToolBarComboButtonBlock::BuildMultiBlockWidget(const ISlateStyle* StyleSet
 			.MinWidth(ToolBarStyle.ComboContentMinWidth)
 			.MaxWidth(ToolBarStyle.ComboContentMaxWidth)
 			.Padding(TAttribute<FMargin>::CreateLambda(
-				[Label, ToolBarStyle]() -> FMargin
+				[bIsSplitButton, Label, ToolBarStyle]() -> FMargin
 				{
-					return Label.Get().IsEmpty() ? FMargin(0) : ToolBarStyle.LabelPadding;
+					return bIsSplitButton || Label.Get().IsEmpty() ? FMargin(0) : ToolBarStyle.LabelPadding;
 				}
 			))
 			.HAlign(ToolBarStyle.ComboContentHorizontalAlignment)
@@ -258,66 +257,23 @@ void SToolBarComboButtonBlock::BuildMultiBlockWidget(const ISlateStyle* StyleSet
 		UserInterfaceType = UICommand->GetUserInterfaceType();
 	}
 	
-	if ((UserInterfaceType == EUserInterfaceActionType::ToggleButton || UserInterfaceType == EUserInterfaceActionType::RadioButton)
-		&& !HasAction())
-	{
-		// There are no bound actions; a toggle or radio button wouldn't actually do anything.
-		// Interpret this button as a normal button.
-		UserInterfaceType = EUserInterfaceActionType::Button;
-	}
-	
 	// Cache the checkbox style for icon coloration
 	CheckBoxStyle = GetCheckBoxStyle(StyleSet, StyleName);
 	
 	LeftHandSideWidget = SNullWidget::NullWidget;
 
-	if (UserInterfaceType == EUserInterfaceActionType::Button)
+	if (bIsSplitButton)
 	{
-		if (HasAction())
+		if (UserInterfaceType == EUserInterfaceActionType::Button)
 		{
-			// clang-format off
-			SAssignNew(LeftHandSideWidget, SButton)
-			.ButtonStyle(ButtonStyle)
-			.HAlign(HAlign_Center)
-			.VAlign(VAlign_Center)
-			.ContentPadding(0.f)
-			.OnClicked(this, &SToolBarComboButtonBlock::OnClicked)
-			.IsEnabled(this, &SToolBarComboButtonBlock::IsEnabled)
-			[
-				IconWidget
-			];
-			// clang-format on
-
+			// When a button is specified, the combo menu is implied to be
+			// the "settings of the button"
 			ComboStyle = &ToolBarStyle.SettingsComboButton;
-			ButtonStyle = &ToolBarStyle.SettingsComboButton.ButtonStyle;
+			ButtonStyle = &ComboStyle->ButtonStyle;
 		}
-
-		// clang-format off
-		SAssignNew(ComboButtonWidget, SComboButton)
-		.AddMetaData<FTagMetaData>(FTagMetaData(TagName))
-		.ContentPadding(0.f)
-		.ComboButtonStyle(ComboStyle)
-		.ButtonStyle(ButtonStyle)
-		.ToolTipText(ToolBarComboButtonBlock->ToolTip)
-		.MenuPlacement(ToolBarComboButtonBlock->PlacementOverride)
-		.ForegroundColor(this, &SToolBarComboButtonBlock::OnGetForegroundColor)
-		// Route the content generator event
-		.OnGetMenuContent(this, &SToolBarComboButtonBlock::OnGetMenuContent)
-		.ButtonContent()
-		[
-			ButtonContent
-		];
-		// clang-format on
-	}
-
-	else if (ensure(
-				 UserInterfaceType == EUserInterfaceActionType::ToggleButton
-				 || UserInterfaceType == EUserInterfaceActionType::RadioButton
-			 ))
-	{
+	
 		// clang-format off
 		SAssignNew(LeftHandSideWidget, SCheckBox)
-			// Use the tool bar style for this check box
 			.Style(CheckBoxStyle)
 			.CheckBoxContentUsesAutoWidth(false)
 			.OnCheckStateChanged(this, &SToolBarComboButtonBlock::OnCheckStateChanged)
@@ -325,61 +281,55 @@ void SToolBarComboButtonBlock::BuildMultiBlockWidget(const ISlateStyle* StyleSet
 			.IsChecked(this, &SToolBarComboButtonBlock::GetCheckState)
 			.IsEnabled(this, &SToolBarComboButtonBlock::IsEnabled)
 			[
-				IconWidget
-			];
-		// clang-format on
-
-		// clang-format off
-		SAssignNew(ComboButtonWidget, SComboButton)
-			.AddMetaData<FTagMetaData>(FTagMetaData(TagName))
-			.ContentPadding(0.f)
-			.ComboButtonStyle(ComboStyle)
-			.ButtonStyle(ButtonStyle)
-			.MenuPlacement(ToolBarComboButtonBlock->PlacementOverride)
-			.ForegroundColor(this, &SToolBarComboButtonBlock::OnGetForegroundColor)
-			// Route the content generator event
-			.OnGetMenuContent(this, &SToolBarComboButtonBlock::OnGetMenuContent)
-			.ButtonContent()
-			[
-				SNew(STextBlock)
-				.Visibility(ToolBarComboButtonBlock->bSimpleComboBox ? EVisibility::Collapsed : LabelVisibility)
-				.Text(Label)
-				// Smaller font for tool tip labels
-				.TextStyle(&ToolBarStyle.LabelStyle)
-				.Margin(0)
+				SNew(SBox)
+				.HAlign(HAlign_Center)
+				.VAlign(VAlign_Center)
+				[
+					IconWidget
+				]
 			];
 		// clang-format on
 	}
-
-	if (LeftHandSideWidget && ComboButtonWidget)
-	{
-		// clang-format off
-		ChildSlot
+	
+	// clang-format off
+	ChildSlot
+	[
+		SNew(SBorder)
+		.BorderImage( this, &SToolBarComboButtonBlock::GetBorderImage )
+		.Padding(0.f)
 		[
-			SNew(SBorder)
-			.BorderImage( this, &SToolBarComboButtonBlock::GetBorderImage )
-			.Padding(0.f)
+			SNew(SHorizontalBox)
+			.ToolTip(FMultiBoxSettings::ToolTipConstructor.Execute(
+				ActualToolTip, nullptr, UICommand, /*ShowActionShortcut=*/true
+			))
+			+ SHorizontalBox::Slot()
+			.AutoWidth()
+			.VAlign(VAlign_Fill)
 			[
-				SNew(SHorizontalBox)
-				.ToolTip(FMultiBoxSettings::ToolTipConstructor.Execute(
-					ActualToolTip, nullptr, UICommand, /*ShowActionShortcut=*/true
-				))
-				+ SHorizontalBox::Slot()
-				.AutoWidth()
-				.VAlign(VAlign_Fill)
-				[
-					LeftHandSideWidget.ToSharedRef()
-				]
-				+ SHorizontalBox::Slot()
-				.AutoWidth()
-				.VAlign(VAlign_Fill)
+				LeftHandSideWidget.ToSharedRef()
+			]
+			+ SHorizontalBox::Slot()
+			.AutoWidth()
+			.VAlign(VAlign_Fill)
+			[
+				SAssignNew(ComboButtonWidget, SComboButton)
+				.AddMetaData<FTagMetaData>(FTagMetaData(TagName))
+				.ContentPadding(0.f)
+				.ComboButtonStyle(ComboStyle)
+				.ButtonStyle(ButtonStyle)
+				.ToolTipText(ToolBarComboButtonBlock->ToolTip)
+				.MenuPlacement(ToolBarComboButtonBlock->PlacementOverride)
+				.ForegroundColor(this, &SToolBarComboButtonBlock::OnGetForegroundColor)
+				// Route the content generator event
+				.OnGetMenuContent(this, &SToolBarComboButtonBlock::OnGetMenuContent)
+				.ButtonContent()
 				[
-					ComboButtonWidget.ToSharedRef()
+					ButtonContent
 				]
 			]
-		];
-		// clang-format on
-	}
+		]
+	];
+	// clang-format on
 
 	FMargin Padding = ToolBarStyle.ComboButtonPadding;
 	if (ToolBarComboButtonBlock->bSimpleComboBox)
@@ -407,7 +357,6 @@ FReply SToolBarComboButtonBlock::OnClicked()
 	// Button was clicked, so trigger the action!
 	TSharedPtr<const FUICommandList> ActionList = MultiBlock->GetActionList();
 	TSharedPtr<const FUICommandInfo> Action = MultiBlock->GetAction();
-	const FUIAction& DirectActions = MultiBlock->GetDirectActions();
 
 	if (ActionList.IsValid() && Action.IsValid())
 	{
diff --git a/Engine/Source/Runtime/Slate/Public/Framework/MultiBox/SToolBarComboButtonBlock.h b/Engine/Source/Runtime/Slate/Public/Framework/MultiBox/SToolBarComboButtonBlock.h
--- a/Engine/Source/Runtime/Slate/Public/Framework/MultiBox/SToolBarComboButtonBlock.h
+++ b/Engine/Source/Runtime/Slate/Public/Framework/MultiBox/SToolBarComboButtonBlock.h
@@ -225,11 +225,11 @@ private:
 	/** Overrides the visibility of the of label. This is used to set up the LabelVisibility attribute */
 	TOptional<EVisibility> LabelVisibilityOverride;
 
-	/** Controls the visibility of the of label, defaults to GetIconVisibility */
-	TAttribute< EVisibility > LabelVisibility;
-
 	/** Optional overridden icon for this tool bar button.  IF not set, then the action's icon will be used instead. */
 	TAttribute<FSlateIcon> Icon;
+	
+	/** Whether this toolbar should always use small icons, regardless of the current settings */
+	bool bForceSmallIcons;
 
 	TSharedPtr<SWidget> LeftHandSideWidget; // Widget that comes before the combo button
 	TSharedPtr<SComboButton> ComboButtonWidget;
@@ -237,12 +237,9 @@ private:
 	/** The foreground color for button when the combo button is open */
 	FSlateColor OpenForegroundColor;
 
-	/** The hovered style for the entire block used when the block has multiple widgets */
-	const FSlateBrush* BlockHovered = nullptr;
-	
 	/** The checkbox style to be used for simulating checkbox foreground color. */
 	const FCheckBoxStyle* CheckBoxStyle = nullptr;
 
-	/** Whether this toolbar should always use small icons, regardless of the current settings */
-	bool bForceSmallIcons;
+	/** The hovered style for the entire block used when the block has multiple widgets */
+	const FSlateBrush* BlockHovered = nullptr;
 };
