[STableViewBase, FixedLineScrollOffset] Fix for when 'bStartedTouchInteraction' is always true and also for when 'DesiredLineOffset' is almost a multiple of 'NumItemsPerLine'.

When a tile view 'HasMouseCapture' it forces 'bStartedTouchInteraction' to remain true even when the player is not touching the screen anymore (see SCommonTileView.h OnMouseLeave for an example of that). Therefore having '!bStartedTouchInteraction' inside 'ShouldScrollFixedLineOffsetAtZeroVelocity' prevented the FixedScrollLineOffset from working in this situation.

'!bStartedTouchInteraction' was initially added as a condition because we noticed that sometimes 'GetTargetScrollOffset' would recaculate the 'DesiredLineOffset' on the first Tick after starting scrolling, when the scroll velocity is still considered 0 and 'IsRightClickScrolling' still returns 'false'. This could cause issues in particular when 'FixedLineScrollOffset' is 0, because then due to double imprecision 'DesiredLineOffset' could be very slightly below a multiple of 'NumItemsPerLine' and the 'FloorToDouble' calculation would teleport to the previous row. We fix this issue by not updating 'DesiredLineOffset' if it is very close to a multiple of 'NumItemsPerLine'.

#rnx

diff --git a/Engine/Source/Runtime/Slate/Private/Widgets/Views/STableViewBase.cpp b/Engine/Source/Runtime/Slate/Private/Widgets/Views/STableViewBase.cpp
--- a/Engine/Source/Runtime/Slate/Private/Widgets/Views/STableViewBase.cpp
+++ b/Engine/Source/Runtime/Slate/Private/Widgets/Views/STableViewBase.cpp
@@ -837,15 +837,22 @@ double STableViewBase::GetTargetScrollOffset() const
 		// We don't do that if we are at the end of the list because in this case showing the last line correctly is more important.
 		const int32 NumItemsPerLine = GetNumItemsPerLine();
 		double DesiredLineOffset = DesiredScrollOffset / NumItemsPerLine;
-		if (bIsInertialScroll)
-		{
-			// Scroll to the item most in view
-			DesiredLineOffset = FMath::RoundToDouble(DesiredLineOffset);
-		}
-		else
+		const double RoundedDesiredLineOffset = FMath::RoundToDouble(DesiredLineOffset);
+
+		// If DesiredLineOffset is almost equal to RoundedDesiredLineOffset:
+		// There is no need to round it and it is better not to floor it in case it is slightly less that the rounded value
+		if (!FMath::IsNearlyEqual(RoundedDesiredLineOffset, DesiredLineOffset, UE_DOUBLE_KINDA_SMALL_NUMBER))
 		{
-			// Scroll to the topmost item.
-			DesiredLineOffset = FMath::FloorToDouble(DesiredLineOffset);
+			if (bIsInertialScroll)
+			{
+				// Scroll to the item most in view
+				DesiredLineOffset = RoundedDesiredLineOffset;
+			}
+			else
+			{
+				// Scroll to the topmost item.
+				DesiredLineOffset = FMath::FloorToDouble(DesiredLineOffset);
+			}
 		}
 		DesiredLineOffset -= FixedLineScrollOffset.GetValue();
 
@@ -856,7 +863,7 @@ double STableViewBase::GetTargetScrollOffset() const
 
 bool STableViewBase::ShouldScrollToFixedLineOffsetAtZeroVelocity() const
 {
-	return !bWasAtEndOfList && FixedLineScrollOffset.IsSet() && !IsRightClickScrolling() && !bStartedTouchInteraction;
+	return !bWasAtEndOfList && FixedLineScrollOffset.IsSet() && !IsRightClickScrolling();
 }
 
 float STableViewBase::ScrollBy(const FGeometry& MyGeometry, float ScrollByAmountInSlateUnits, EAllowOverscroll InAllowOverscroll)
diff --git a/Engine/Source/Runtime/UMG/Public/Components/ListViewBase.h b/Engine/Source/Runtime/UMG/Public/Components/ListViewBase.h
--- a/Engine/Source/Runtime/UMG/Public/Components/ListViewBase.h
+++ b/Engine/Source/Runtime/UMG/Public/Components/ListViewBase.h
@@ -707,7 +707,7 @@ protected:
 	float ScrollingAnimationInterpolationSpeed = 12.f;
 
 	/** True to enable lerped animation when scrolling through the list with touch*/
-	UPROPERTY(EditAnywhere, BlueprintReadOnly, Category = Scrolling)
+	UPROPERTY(EditAnywhere, BlueprintReadOnly, Category = Scrolling, DisplayName = "Enable Touch Animated Scrolling")
 	bool bInEnableTouchAnimatedScrolling = false;
 
 	/**  Disable to stop scrollbars from activating inertial overscrolling */
