SSpinBox : Call EditableText->SetEditableText instead of EditableText->SetText

Previously, any calls to EditableText->SetText ( e.g. from OnMouseWheel ) would break the attribute binding on EditableText.Text leading to a disconnect between the value displayed in text mode and the ValueAttribute.

diff --git a/Engine/Source/Runtime/Slate/Private/Widgets/Input/SEditableText.cpp b/Engine/Source/Runtime/Slate/Private/Widgets/Input/SEditableText.cpp
--- a/Engine/Source/Runtime/Slate/Private/Widgets/Input/SEditableText.cpp
+++ b/Engine/Source/Runtime/Slate/Private/Widgets/Input/SEditableText.cpp
@@ -92,6 +92,11 @@ FText SEditableText::GetText() const
 	return EditableTextLayout->GetText();
 }
 
+bool SEditableText::SetEditableText(const FText& InNewText)
+{
+	return EditableTextLayout->SetEditableText(InNewText);
+}
+
 void SEditableText::SetFont( const TAttribute< FSlateFontInfo >& InNewFont )
 {
 	Font = InNewFont;
diff --git a/Engine/Source/Runtime/Slate/Private/Widgets/Input/SSpinBox.cpp b/Engine/Source/Runtime/Slate/Private/Widgets/Input/SSpinBox.cpp
--- a/Engine/Source/Runtime/Slate/Private/Widgets/Input/SSpinBox.cpp
+++ b/Engine/Source/Runtime/Slate/Private/Widgets/Input/SSpinBox.cpp
@@ -593,12 +593,12 @@ FReply SSpinBox<NumericType>::OnMouseWheel(const FGeometry& MyGeometry, const FP
 
 		TSharedPtr<INumericTypeInterface<NumericType>> Interface = InterfaceAttr.Get();
 
-		// First SetText is to update the value before calling CommitValue. Otherwise, when the text lose
+		// First SetEditableText is to update the value before calling CommitValue. Otherwise, when the text lose
 		// focus from the CommitValue, it will override the value we just committed.
-		// The second SetText is to update the text to the InternalValue since it could have been clamped.
-		EditableText->SetText(FText::FromString(Interface->ToString((NumericType)NewValue)));
+		// The second SetEditableText is to update the text to the InternalValue since it could have been clamped.
+		EditableText->SetEditableText(FText::FromString(Interface->ToString((NumericType)NewValue)));
 		CommitValue(RoundedNewValue, NewValue, CommittedViaSpin, ETextCommit::OnEnter);
-		EditableText->SetText(FText::FromString(Interface->ToString((NumericType)InternalValue)));
+		EditableText->SetEditableText(FText::FromString(Interface->ToString((NumericType)InternalValue)));
 
 		return FReply::Handled();
 	}
@@ -918,7 +918,7 @@ void SSpinBox<NumericType>::TextField_OnTextChanged(const FText& NewText)
 		if (NumValidChars < Data.Len())
 		{
 			FString ValidData = NumValidChars > 0 ? Data.Left(NumValidChars) : GetValueAsString();
-			EditableText->SetText(FText::FromString(ValidData));
+			EditableText->SetEditableText(FText::FromString(ValidData));
 		}
 
 		// we check that the input is numeric, as we don't want to commit the new value on every change when an expression like *= is entered
diff --git a/Engine/Source/Runtime/Slate/Public/Widgets/Input/SEditableText.h b/Engine/Source/Runtime/Slate/Public/Widgets/Input/SEditableText.h
--- a/Engine/Source/Runtime/Slate/Public/Widgets/Input/SEditableText.h
+++ b/Engine/Source/Runtime/Slate/Public/Widgets/Input/SEditableText.h
@@ -204,6 +204,16 @@ public:
 	 */
 	SLATE_API FText GetText() const;
 
+	/**
+	 * Sets the text currently being edited
+	 * Note: Doesn't override the bound Text attribute, nor does it call OnTextChanged
+	 *
+	 * @param InNewText	The new editable text
+	 *
+	 * @return true if the text was updated, false if the text was already up-to-date
+	 */
+	SLATE_API bool SetEditableText(const FText& InNewText);
+
 	/** See the HintText attribute */
 	SLATE_API void SetHintText( const TAttribute< FText >& InHintText );
 	
