[SSpinBox] Spin boxes now revert to the previous value on escape key press

diff --git a/Engine/Source/Runtime/Slate/Private/Widgets/Input/SSpinBox.cpp b/Engine/Source/Runtime/Slate/Private/Widgets/Input/SSpinBox.cpp
--- a/Engine/Source/Runtime/Slate/Private/Widgets/Input/SSpinBox.cpp
+++ b/Engine/Source/Runtime/Slate/Private/Widgets/Input/SSpinBox.cpp
@@ -171,6 +171,7 @@ void SSpinBox<NumericType>::Construct(const FArguments& InArgs)
 				.Font(InArgs._Font)
 				.SelectAllTextWhenFocused(true)
 				.Text(this, &SSpinBox<NumericType>::GetValueAsText)
+				.RevertTextOnEscape(true)
 				.OnIsTypedCharValid(this, &SSpinBox<NumericType>::IsCharacterValid)
 				.OnTextChanged(this, &SSpinBox<NumericType>::TextField_OnTextChanged)
 				.OnTextCommitted(this, &SSpinBox<NumericType>::TextField_OnTextCommitted)
