Add a Slate.DismissMenuStacksOnFocusLost debug cvar so that it's possible to use Widget Reflector on menu stacks

diff --git a/Engine/Source/Runtime/Slate/Private/Framework/Application/MenuStack.cpp b/Engine/Source/Runtime/Slate/Private/Framework/Application/MenuStack.cpp
--- a/Engine/Source/Runtime/Slate/Private/Framework/Application/MenuStack.cpp
+++ b/Engine/Source/Runtime/Slate/Private/Framework/Application/MenuStack.cpp
@@ -20,6 +20,14 @@ namespace FMenuStackDefs
 	const float AnimationDuration = 0.15f;
 }
 
+#if WITH_SLATE_DEBUGGING
+FAutoConsoleVariable CVarSlateDismissMenuStacksOnFocusLost(
+	TEXT("Slate.DismissMenuStacksOnFocusLost"),
+	true,
+	TEXT("Should open menus disappear when whey lose focus?")
+);
+#endif
+
 /** Overlay widget class used to hold menu contents and display them on top of the current window */
 class SMenuPanel : public SOverlay
 {
@@ -769,6 +777,13 @@ void FMenuStack::OnMenuContentLostFocus(const FWidgetPath& InFocussedPath)
 		}
 		else
 		{
+#if WITH_SLATE_DEBUGGING
+			if (!CVarSlateDismissMenuStacksOnFocusLost->GetBool())
+			{
+				return;
+			}
+#endif
+
 			// Focus has moved outside all menus - collapse the stack
 			DismissAll();
 		}
@@ -851,6 +866,13 @@ void FMenuStack::OnWindowActivated( TSharedRef<SWindow> ActivatedWindow )
 		}
 		else
 		{
+#if WITH_SLATE_DEBUGGING
+			if (!CVarSlateDismissMenuStacksOnFocusLost->GetBool())
+			{
+				return;
+			}
+#endif
+
 			// Activated a window that isn't a menu - collapse the stack
 			DismissAll();
 		}
diff --git a/Engine/Source/Runtime/Slate/Private/Framework/Application/SlateApplication.cpp b/Engine/Source/Runtime/Slate/Private/Framework/Application/SlateApplication.cpp
--- a/Engine/Source/Runtime/Slate/Private/Framework/Application/SlateApplication.cpp
+++ b/Engine/Source/Runtime/Slate/Private/Framework/Application/SlateApplication.cpp
@@ -7300,6 +7300,13 @@ bool FSlateApplication::OnWindowAction( const TSharedRef< FGenericWindow >& Plat
 
 	if (InActionType == EWindowAction::ClickedNonClientArea)
 	{
+#if WITH_SLATE_DEBUGGING
+		if (!CVarSlateDismissMenuStacksOnFocusLost->GetBool())
+		{
+			return bResult;
+		}
+#endif
+
 		DismissAllMenus();
 	}
 
diff --git a/Engine/Source/Runtime/Slate/Public/Framework/Application/MenuStack.h b/Engine/Source/Runtime/Slate/Public/Framework/Application/MenuStack.h
--- a/Engine/Source/Runtime/Slate/Public/Framework/Application/MenuStack.h
+++ b/Engine/Source/Runtime/Slate/Public/Framework/Application/MenuStack.h
@@ -61,6 +61,10 @@ struct FPopupTransitionEffect
 
 DECLARE_MULTICAST_DELEGATE_OneParam(FOnMenuDestroyed, const TSharedRef<IMenu>& InMenu);
 
+#if WITH_SLATE_DEBUGGING
+extern FAutoConsoleVariable CVarSlateDismissMenuStacksOnFocusLost;
+#endif
+
 /**
  * Represents a stack of open menus. The last item in the stack is the top most menu
  * Menus are described as IMenus. Implementations of IMenu can control how the menus are created and presented (e.g. in their own window)
