[ToolMenus] Allow to dynamically hide Menu Entries/Widgets/Sections

diff --git a/Engine/Source/Developer/ToolMenus/Private/ToolMenuEntry.cpp b/Engine/Source/Developer/ToolMenus/Private/ToolMenuEntry.cpp
--- a/Engine/Source/Developer/ToolMenus/Private/ToolMenuEntry.cpp
+++ b/Engine/Source/Developer/ToolMenus/Private/ToolMenuEntry.cpp
@@ -339,6 +339,11 @@ void FToolMenuEntry::SetShowInToolbarTopLevel(TAttribute<bool> InTopLevel)
 	ShowInToolbarTopLevel = InTopLevel;
 }
 
+void FToolMenuEntry::SetVisibility(TAttribute<EVisibility> InVisibility)
+{
+	Visibility = InVisibility;
+}
+
 bool FToolMenuEntry::IsScriptObjectDynamicConstruct() const
 {
 	static const FName ConstructMenuEntryName = GET_FUNCTION_NAME_CHECKED(UToolMenuEntryScript, ConstructMenuEntry);
diff --git a/Engine/Source/Developer/ToolMenus/Private/ToolMenuSection.cpp b/Engine/Source/Developer/ToolMenus/Private/ToolMenuSection.cpp
--- a/Engine/Source/Developer/ToolMenus/Private/ToolMenuSection.cpp
+++ b/Engine/Source/Developer/ToolMenus/Private/ToolMenuSection.cpp
@@ -35,6 +35,7 @@ void FToolMenuSection::InitGeneratedSectionCopy(const FToolMenuSection& Source,
 	Construct = Source.Construct;
 	Context = InContext;
 	Alignment = Source.Alignment;
+	Visibility = Source.Visibility;
 }
 
 bool FToolMenuSection::IsRegistering() const
@@ -126,7 +127,14 @@ FToolMenuEntry& FToolMenuSection::AddMenuSeparator(const FName InName)
 
 FToolMenuEntry& FToolMenuSection::AddSeparator(const FName InName)
 {
-	return AddEntry(FToolMenuEntry::InitSeparator(InName));
+	FToolMenuEntry& SeparatorEntry = AddEntry(FToolMenuEntry::InitSeparator(InName));
+
+	if (Visibility.IsSet())
+	{
+		SeparatorEntry.Visibility = Visibility;
+	}
+
+	return SeparatorEntry;
 }
 
 FToolMenuEntry& FToolMenuSection::AddSubMenu(const FName InName, const TAttribute<FText>& InLabel, const TAttribute<FText>& InToolTip, const FNewToolMenuChoice& InMakeMenu, const FToolUIActionChoice& InAction, const EUserInterfaceActionType InUserInterfaceActionType, bool bInOpenSubMenuOnClick, const TAttribute<FSlateIcon>& InIcon, const bool bInShouldCloseWindowAfterMenuSelection)
diff --git a/Engine/Source/Developer/ToolMenus/Private/ToolMenus.cpp b/Engine/Source/Developer/ToolMenus/Private/ToolMenus.cpp
--- a/Engine/Source/Developer/ToolMenus/Private/ToolMenus.cpp
+++ b/Engine/Source/Developer/ToolMenus/Private/ToolMenus.cpp
@@ -93,6 +93,25 @@ FToolUIActionChoice::FToolUIActionChoice(const TSharedPtr< const FUICommandInfo
 
 class FPopulateMenuBuilderWithToolMenuEntry
 {
+private:
+	// Returns visibility for the specified Entry + Section
+	static TAttribute<EVisibility> GetEntryVisibility(FToolMenuEntry& InMenuEntry, FToolMenuSection& InSection)
+	{
+		return TAttribute<EVisibility>::CreateLambda(
+			[BlockVisibility = InMenuEntry.Visibility, SectionVisibility = InSection.Visibility]()
+			{
+				// If section visibility is set, and section is not visible, return Section Visibility
+				if (SectionVisibility.IsSet() && SectionVisibility.Get() != EVisibility::Visible)
+				{
+					return SectionVisibility.Get();
+				}
+
+				// Stick to this Block visibility
+				return BlockVisibility.Get();
+			}
+		);
+	}
+
 public:
 	FPopulateMenuBuilderWithToolMenuEntry(FMenuBuilder& InMenuBuilder, UToolMenu* InMenuData, FToolMenuSection& InSection, FToolMenuEntry& InBlock, bool bInAllowSubMenuCollapse) :
 		MenuBuilder(InMenuBuilder),
@@ -177,15 +196,25 @@ public:
 
 		if (!bSubMenuAdded)
 		{
+			TAttribute<EVisibility> VisibilityOverride = GetEntryVisibility(Block, Section);
+
 			if (Widget.IsValid())
 			{
 				if (bUIActionIsSet)
 				{
-					MenuBuilder.AddSubMenu(UIAction, Widget.ToSharedRef(), NewMenuDelegate, Block.bShouldCloseWindowAfterMenuSelection);
+					MenuBuilder.AddSubMenu(
+						UIAction, Widget.ToSharedRef(), NewMenuDelegate, Block.bShouldCloseWindowAfterMenuSelection, VisibilityOverride
+					);
 				}
 				else
 				{
-					MenuBuilder.AddSubMenu(Widget.ToSharedRef(), NewMenuDelegate, Block.SubMenuData.bOpenSubMenuOnClick, Block.bShouldCloseWindowAfterMenuSelection);
+					MenuBuilder.AddSubMenu(
+						Widget.ToSharedRef(),
+						NewMenuDelegate,
+						Block.SubMenuData.bOpenSubMenuOnClick,
+						Block.bShouldCloseWindowAfterMenuSelection,
+						VisibilityOverride
+					);
 				}
 			}
 			else
@@ -201,7 +230,8 @@ public:
 						Block.UserInterfaceActionType,
 						Block.SubMenuData.bOpenSubMenuOnClick,
 						Block.Icon.Get(),
-						Block.bShouldCloseWindowAfterMenuSelection
+						Block.bShouldCloseWindowAfterMenuSelection,
+						VisibilityOverride
 					);
 				}
 				else
@@ -214,7 +244,8 @@ public:
 						Block.Icon.Get(),
 						Block.bShouldCloseWindowAfterMenuSelection,
 						BlockNameOverride,
-						Block.TutorialHighlightName
+						Block.TutorialHighlightName,
+						VisibilityOverride
 					);
 				}
 			}
@@ -223,6 +254,8 @@ public:
 
 	void AddStandardEntryToMenuBuilder()
 	{
+		TAttribute<EVisibility> VisibilityOverride = GetEntryVisibility(Block, Section);
+
 		// First, check for a ToolUIAction, otherwise do the rest of this (have CommandList and Command)
 		// Need another variable to store if we are using a keybind from a command
 		if (Block.Command.IsValid())
@@ -242,7 +275,9 @@ public:
 					*MenuData->MenuName.ToString());
 			}
 
-			MenuBuilder.AddMenuEntry(Block.Command, BlockNameOverride, LabelToDisplay, Block.ToolTip, Block.Icon.Get());
+			MenuBuilder.AddMenuEntry(
+				Block.Command, BlockNameOverride, LabelToDisplay, Block.ToolTip, Block.Icon.Get(), NAME_None, VisibilityOverride
+			);
 
 			if (bPopCommandList)
 			{
@@ -263,6 +298,7 @@ public:
 			MenuEntryParams.UserInterfaceActionType = Block.UserInterfaceActionType;
 			MenuEntryParams.TutorialHighlightName = Block.TutorialHighlightName;
 			MenuEntryParams.InputBindingOverride = Block.InputBindingLabel;
+			MenuEntryParams.Visibility = VisibilityOverride;
 
 			MenuBuilder.AddMenuEntry(MenuEntryParams);
 		}
@@ -278,6 +314,7 @@ public:
 				MenuEntryParams.UserInterfaceActionType = Block.UserInterfaceActionType;
 				MenuEntryParams.TutorialHighlightName = Block.TutorialHighlightName;
 				MenuEntryParams.InputBindingOverride = Block.InputBindingLabel;
+				MenuEntryParams.Visibility = VisibilityOverride;
 
 				MenuBuilder.AddMenuEntry(MenuEntryParams);
 			}
@@ -291,7 +328,8 @@ public:
 					BlockNameOverride,
 					Block.UserInterfaceActionType,
 					Block.TutorialHighlightName,
-					Block.InputBindingLabel
+					Block.InputBindingLabel,
+					VisibilityOverride
 				);
 			}
 		}
@@ -309,6 +347,8 @@ public:
 			return;
 		}
 
+		TAttribute<EVisibility> VisibilityOverride = GetEntryVisibility(Block, Section);
+
 		UIAction = UToolMenus::ConvertUIAction(Block, MenuData->Context);
 		bUIActionIsSet = UIAction.ExecuteAction.IsBound() || UIAction.CanExecuteAction.IsBound() || UIAction.GetActionCheckState.IsBound() || UIAction.IsActionVisibleDelegate.IsBound();
 
@@ -346,7 +386,7 @@ public:
 		}
 		else if (Block.Type == EMultiBlockType::Separator)
 		{
-			MenuBuilder.AddSeparator(BlockNameOverride);
+			MenuBuilder.AddSeparator(BlockNameOverride, VisibilityOverride);
 		}
 		else if (Block.Type == EMultiBlockType::Widget)
 		{
@@ -361,13 +401,22 @@ public:
 				MenuEntryParams.UserInterfaceActionType = Block.UserInterfaceActionType;
 				MenuEntryParams.TutorialHighlightName = Block.TutorialHighlightName;
 				MenuEntryParams.InputBindingOverride = Block.InputBindingLabel;
-				
+				MenuEntryParams.Visibility = VisibilityOverride;
+
 				MenuBuilder.AddMenuEntry(MenuEntryParams);
 			}
 			else
 			{
 				Block.WidgetData.StyleParams.bNoIndent = Block.WidgetData.bNoIndent;
-				MenuBuilder.AddWidget(Widget.ToSharedRef(), LabelToDisplay.Get(), Block.WidgetData.StyleParams, Block.WidgetData.bSearchable, Block.ToolTip.Get(), Block.Icon);
+				MenuBuilder.AddWidget(
+					Widget.ToSharedRef(),
+					LabelToDisplay.Get(),
+					Block.WidgetData.StyleParams,
+					Block.WidgetData.bSearchable,
+					Block.ToolTip.Get(),
+					Block.Icon,
+					VisibilityOverride
+				);
 			}
 		}
 		else
@@ -851,6 +900,11 @@ void UToolMenus::AssembleMenu(UToolMenu* GeneratedMenu, const UToolMenu* Other)
 				{
 					Section->Construct = RemainingSection.Construct;
 				}
+
+				if (!Section->Visibility.IsSet() && RemainingSection.Visibility.IsSet())
+				{
+					Section->Visibility = RemainingSection.Visibility;
+				}
 			}
 
 			AssembleMenuSection(GeneratedMenu, Other, Section, RemainingSection);
@@ -1424,11 +1478,11 @@ void UToolMenus::PopulateMenuBuilder(FMenuBuilder& MenuBuilder, UToolMenu* MenuD
 			{
 				LabelText = FText::FromName(Section.Name);
 			}
-			MenuBuilder.BeginSection(Section.Name, LabelText);
+			MenuBuilder.BeginSection(Section.Name, LabelText, Section.Visibility);
 		}
 		else
 		{
-			MenuBuilder.BeginSection(Section.Name, Section.Label);
+			MenuBuilder.BeginSection(Section.Name, Section.Label, Section.Visibility);
 		}
 
 		for (FToolMenuEntry& Block : Section.Blocks)
diff --git a/Engine/Source/Developer/ToolMenus/Public/ToolMenuEntry.h b/Engine/Source/Developer/ToolMenus/Public/ToolMenuEntry.h
--- a/Engine/Source/Developer/ToolMenus/Public/ToolMenuEntry.h
+++ b/Engine/Source/Developer/ToolMenus/Public/ToolMenuEntry.h
@@ -203,6 +203,8 @@ struct TOOLMENUS_API FToolMenuEntry
 	 */
 	void SetShowInToolbarTopLevel(TAttribute<bool> InTopLevel);
 
+	void SetVisibility(TAttribute<EVisibility> InVisibility);
+
 private:
 
 	void SetCommand(const TSharedPtr< const FUICommandInfo >& InCommand, TOptional<FName> InName, const TAttribute<FText>& InLabel, const TAttribute<FText>& InToolTip, const TAttribute<FSlateIcon>& InIcon);
@@ -281,4 +283,5 @@ private:
 	bool bCommandIsKeybindOnly;
 
 	TAttribute<bool> ShowInToolbarTopLevel;
+	TAttribute<EVisibility> Visibility;
 };
diff --git a/Engine/Source/Developer/ToolMenus/Public/ToolMenuSection.h b/Engine/Source/Developer/ToolMenus/Public/ToolMenuSection.h
--- a/Engine/Source/Developer/ToolMenus/Public/ToolMenuSection.h
+++ b/Engine/Source/Developer/ToolMenus/Public/ToolMenuSection.h
@@ -133,6 +133,7 @@ public:
 	TObjectPtr<UToolMenuSectionDynamic> ToolMenuSectionDynamic;
 
 	TAttribute<FText> Label;
+	TAttribute<EVisibility> Visibility;
 
 	FNewSectionConstructChoice Construct;
 
diff --git a/Engine/Source/Runtime/Slate/Private/Framework/MultiBox/MultiBoxBuilder.cpp b/Engine/Source/Runtime/Slate/Private/Framework/MultiBox/MultiBoxBuilder.cpp
--- a/Engine/Source/Runtime/Slate/Private/Framework/MultiBox/MultiBoxBuilder.cpp
+++ b/Engine/Source/Runtime/Slate/Private/Framework/MultiBox/MultiBoxBuilder.cpp
@@ -148,7 +148,7 @@ FBaseMenuBuilder::FBaseMenuBuilder( const EMultiBoxType InType, const bool bInSh
 	MultiBox->SetStyle(InStyleSet, "Menu");
 }
 
-void FBaseMenuBuilder::AddMenuEntry( const TSharedPtr< const FUICommandInfo > InCommand, FName InExtensionHook, const TAttribute<FText>& InLabelOverride, const TAttribute<FText>& InToolTipOverride, const FSlateIcon& InIconOverride, FName InTutorialHighlightName )
+void FBaseMenuBuilder::AddMenuEntry( const TSharedPtr< const FUICommandInfo > InCommand, FName InExtensionHook, const TAttribute<FText>& InLabelOverride, const TAttribute<FText>& InToolTipOverride, const FSlateIcon& InIconOverride, FName InTutorialHighlightName, const TAttribute<EVisibility>& InVisibility )
 {
 	ApplySectionBeginning();
 
@@ -159,12 +159,18 @@ void FBaseMenuBuilder::AddMenuEntry( const TSharedPtr< const FUICommandInfo > In
 	TSharedRef< FMenuEntryBlock > NewMenuEntryBlock = MakeShared<FMenuEntryBlock>( InExtensionHook, InCommand, CommandListStack.Last(), InLabelOverride, InToolTipOverride, InIconOverride, bCloseSelfOnly );
 	NewMenuEntryBlock->SetTutorialHighlightName(GenerateTutorialIdentifierName(TutorialHighlightName, InTutorialHighlightName, InCommand, MultiBox->GetBlocks().Num()));
 	NewMenuEntryBlock->SetCheckBoxStyle(CheckBoxStyle);
+
+	if (InVisibility.IsSet())
+	{
+		NewMenuEntryBlock->SetVisibilityOverride(InVisibility);
+	}
+
 	MultiBox->AddMultiBlock( NewMenuEntryBlock );
 
 	ApplyHook(InExtensionHook, EExtensionHook::After);
 }
 
-void FBaseMenuBuilder::AddMenuEntry( const TAttribute<FText>& InLabel, const TAttribute<FText>& InToolTip, const FSlateIcon& InIcon, const FUIAction& InAction, FName InExtensionHook, const EUserInterfaceActionType UserInterfaceActionType, FName InTutorialHighlightName, const TAttribute<FText>& InInputBindingOverride)
+void FBaseMenuBuilder::AddMenuEntry( const TAttribute<FText>& InLabel, const TAttribute<FText>& InToolTip, const FSlateIcon& InIcon, const FUIAction& InAction, FName InExtensionHook, const EUserInterfaceActionType UserInterfaceActionType, FName InTutorialHighlightName, const TAttribute<FText>& InInputBindingOverride, const TAttribute<EVisibility>& InVisibility)
 {
 	ApplySectionBeginning();
 
@@ -173,12 +179,18 @@ void FBaseMenuBuilder::AddMenuEntry( const TAttribute<FText>& InLabel, const TAt
 	constexpr bool bShouldCloseWindowAfterMenuSelection = true;
 	TSharedRef< FMenuEntryBlock > NewMenuEntryBlock = MakeShared<FMenuEntryBlock>(InExtensionHook, InLabel, InToolTip, InIcon, InAction, UserInterfaceActionType, bCloseSelfOnly, bShouldCloseWindowAfterMenuSelection, InInputBindingOverride);
 	NewMenuEntryBlock->SetTutorialHighlightName(GenerateTutorialIdentifierName(TutorialHighlightName, InTutorialHighlightName, nullptr, MultiBox->GetBlocks().Num()));
+
+	if (InVisibility.IsSet())
+	{
+		NewMenuEntryBlock->SetVisibilityOverride(InVisibility);
+	}
+
 	MultiBox->AddMultiBlock( NewMenuEntryBlock );
 	
 	ApplyHook(InExtensionHook, EExtensionHook::After);
 }
 
-void FBaseMenuBuilder::AddMenuEntry( const FUIAction& UIAction, const TSharedRef< SWidget > Contents, const FName& InExtensionHook, const TAttribute<FText>& InToolTip, const EUserInterfaceActionType UserInterfaceActionType, FName InTutorialHighlightName )
+void FBaseMenuBuilder::AddMenuEntry( const FUIAction& UIAction, const TSharedRef< SWidget > Contents, const FName& InExtensionHook, const TAttribute<FText>& InToolTip, const EUserInterfaceActionType UserInterfaceActionType, FName InTutorialHighlightName, const TAttribute<EVisibility>& InVisibility )
 {
 	ApplySectionBeginning();
 
@@ -186,6 +198,12 @@ void FBaseMenuBuilder::AddMenuEntry( const FUIAction& UIAction, const TSharedRef
 
 	TSharedRef< FMenuEntryBlock > NewMenuEntryBlock = MakeShared<FMenuEntryBlock>( InExtensionHook, UIAction, Contents, InToolTip, UserInterfaceActionType, bCloseSelfOnly );
 	NewMenuEntryBlock->SetTutorialHighlightName(GenerateTutorialIdentifierName(TutorialHighlightName, InTutorialHighlightName, nullptr, MultiBox->GetBlocks().Num()));
+
+	if (InVisibility.IsSet())
+	{
+		NewMenuEntryBlock->SetVisibilityOverride(InVisibility);
+	}
+
 	MultiBox->AddMultiBlock( NewMenuEntryBlock );
 
 	ApplyHook(InExtensionHook, EExtensionHook::After);
@@ -199,6 +217,12 @@ void FBaseMenuBuilder::AddMenuEntry(const FMenuEntryParams& InMenuEntryParams)
 	
 	TSharedPtr< FMenuEntryBlock > NewMenuEntryBlock = MakeShared<FMenuEntryBlock>(InMenuEntryParams);
 	NewMenuEntryBlock->SetTutorialHighlightName(GenerateTutorialIdentifierName(TutorialHighlightName, InMenuEntryParams.TutorialHighlightName, nullptr, MultiBox->GetBlocks().Num()));
+
+	if (InMenuEntryParams.Visibility.IsSet())
+	{
+		NewMenuEntryBlock->SetVisibilityOverride(InMenuEntryParams.Visibility);
+	}
+	
 	MultiBox->AddMultiBlock(NewMenuEntryBlock.ToSharedRef());
 
 	ApplyHook(InMenuEntryParams.ExtensionHook, EExtensionHook::After);
@@ -219,7 +243,7 @@ TSharedRef< class SWidget > FMenuBuilder::MakeWidget( FMultiBox::FOnMakeMultiBox
 	return MultiBox->MakeWidget(bSearchable, InMakeMultiBoxBuilderOverride, MaxHeightAttribute);
 }
 
-void FMenuBuilder::BeginSection( FName InExtensionHook, const TAttribute< FText >& InHeadingText )
+void FMenuBuilder::BeginSection(FName InExtensionHook, const TAttribute< FText >& InHeadingText, const TAttribute<EVisibility>& InVisibility)
 {
 	check(CurrentSectionExtensionHook == NAME_None && !bSectionNeedsToBeApplied);
 
@@ -230,6 +254,7 @@ void FMenuBuilder::BeginSection( FName InExtensionHook, const TAttribute< FText
 	bSectionNeedsToBeApplied = true;
 	CurrentSectionExtensionHook = InExtensionHook;
 	CurrentSectionHeadingText = InHeadingText.Get();
+	CurrentSectionVisibility = InVisibility;
 
 	// Do apply the section beginning if we are in developer "show me all the hooks" mode
 	if (FMultiBoxSettings::DisplayMultiboxHooks.Get())
@@ -246,16 +271,17 @@ void FMenuBuilder::EndSection()
 	CurrentSectionExtensionHook = NAME_None;
 	bSectionNeedsToBeApplied = false;
 	CurrentSectionHeadingText = FText::GetEmpty();
+	CurrentSectionVisibility = EVisibility::Visible;
 
 	ApplyHook(SectionExtensionHook, EExtensionHook::After);
 }
 
-void FMenuBuilder::AddMenuSeparator(FName InExtensionHook)
+void FMenuBuilder::AddMenuSeparator(FName InExtensionHook, const TAttribute<EVisibility>& InVisibility)
 {
-	AddSeparator(InExtensionHook);
+	AddSeparator(InExtensionHook, InVisibility);
 }
 
-void FMenuBuilder::AddSeparator(FName InExtensionHook /*= NAME_None*/)
+void FMenuBuilder::AddSeparator(FName InExtensionHook /*= NAME_None*/, const TAttribute<EVisibility>& InVisibility)
 {
 	ApplySectionBeginning();
 
@@ -265,13 +291,19 @@ void FMenuBuilder::AddSeparator(FName InExtensionHook /*= NAME_None*/)
 	if (MultiBox->GetBlocks().Num() > 0 || FMultiBoxSettings::DisplayMultiboxHooks.Get())
 	{
 		TSharedRef< FMenuSeparatorBlock > NewMenuSeparatorBlock(new FMenuSeparatorBlock(InExtensionHook, /* bInIsPartOfHeading=*/ false));
+
+		if (InVisibility.IsSet())
+		{
+			NewMenuSeparatorBlock->SetVisibilityOverride(InVisibility);
+		}
+
 		MultiBox->AddMultiBlock(NewMenuSeparatorBlock);
 	}
 
 	ApplyHook(InExtensionHook, EExtensionHook::After);
 }
 
-void FMenuBuilder::AddSubMenu( const TAttribute<FText>& InMenuLabel, const TAttribute<FText>& InToolTip, const FNewMenuDelegate& InSubMenu, const FUIAction& InUIAction, FName InExtensionHook, const EUserInterfaceActionType InUserInterfaceActionType, const bool bInOpenSubMenuOnClick, const FSlateIcon& InIcon, const bool bInShouldCloseWindowAfterMenuSelection /*= true*/ )
+void FMenuBuilder::AddSubMenu( const TAttribute<FText>& InMenuLabel, const TAttribute<FText>& InToolTip, const FNewMenuDelegate& InSubMenu, const FUIAction& InUIAction, FName InExtensionHook, const EUserInterfaceActionType InUserInterfaceActionType, const bool bInOpenSubMenuOnClick, const FSlateIcon& InIcon, const bool bInShouldCloseWindowAfterMenuSelection /*= true*/, const TAttribute<EVisibility>& InVisibility )
 {
 	ApplySectionBeginning();
 
@@ -279,10 +311,15 @@ void FMenuBuilder::AddSubMenu( const TAttribute<FText>& InMenuLabel, const TAttr
 	TSharedRef< FMenuEntryBlock > NewMenuEntryBlock = MakeShared<FMenuEntryBlock>( InExtensionHook, InMenuLabel, InToolTip, InSubMenu, ExtenderStack.Top(), bIsSubMenu, bInOpenSubMenuOnClick, InIcon, InUIAction, InUserInterfaceActionType, bCloseSelfOnly, bInShouldCloseWindowAfterMenuSelection, CommandListStack.Last() );
 	NewMenuEntryBlock->SetRecursivelySearchable(bRecursivelySearchable);
 
+	if (InVisibility.IsSet())
+	{
+		NewMenuEntryBlock->SetVisibilityOverride(InVisibility);
+	}
+
 	MultiBox->AddMultiBlock( NewMenuEntryBlock );
 }
 
-void FMenuBuilder::AddSubMenu( const TAttribute<FText>& InMenuLabel, const TAttribute<FText>& InToolTip, const FNewMenuDelegate& InSubMenu, const bool bInOpenSubMenuOnClick /*= false*/, const FSlateIcon& InIcon /*= FSlateIcon()*/, const bool bInShouldCloseWindowAfterMenuSelection /*= true*/, FName InExtensionHook /*=NAME_None*/, FName InTutorialHighlightName /*= NAME_None*/)
+void FMenuBuilder::AddSubMenu( const TAttribute<FText>& InMenuLabel, const TAttribute<FText>& InToolTip, const FNewMenuDelegate& InSubMenu, const bool bInOpenSubMenuOnClick /*= false*/, const FSlateIcon& InIcon /*= FSlateIcon()*/, const bool bInShouldCloseWindowAfterMenuSelection /*= true*/, FName InExtensionHook /*=NAME_None*/, FName InTutorialHighlightName /*= NAME_None*/, const TAttribute<EVisibility>& InVisibility)
 {
 	ApplySectionBeginning();
 
@@ -291,11 +328,16 @@ void FMenuBuilder::AddSubMenu( const TAttribute<FText>& InMenuLabel, const TAttr
 	NewMenuEntryBlock->SetTutorialHighlightName(GenerateTutorialIdentifierName(TutorialHighlightName, InTutorialHighlightName, nullptr, MultiBox->GetBlocks().Num()));
 	NewMenuEntryBlock->SetRecursivelySearchable(bRecursivelySearchable);
 	NewMenuEntryBlock->SetCheckBoxStyle(CheckBoxStyle);
-	
+
+	if (InVisibility.IsSet())
+	{
+		NewMenuEntryBlock->SetVisibilityOverride(InVisibility);
+	}
+
 	MultiBox->AddMultiBlock( NewMenuEntryBlock );
 }
 
-void FMenuBuilder::AddSubMenu( const TSharedRef< SWidget > Contents, const FNewMenuDelegate& InSubMenu, const bool bInOpenSubMenuOnClick /*= false*/, const bool bInShouldCloseWindowAfterMenuSelection /*= true*/ )
+void FMenuBuilder::AddSubMenu( const TSharedRef< SWidget > Contents, const FNewMenuDelegate& InSubMenu, const bool bInOpenSubMenuOnClick /*= false*/, const bool bInShouldCloseWindowAfterMenuSelection /*= true*/, const TAttribute<EVisibility>& InVisibility )
 {
 	ApplySectionBeginning();
 
@@ -303,10 +345,15 @@ void FMenuBuilder::AddSubMenu( const TSharedRef< SWidget > Contents, const FNewM
 	TSharedRef< FMenuEntryBlock > NewMenuEntryBlock = MakeShared<FMenuEntryBlock>( NAME_None, Contents, InSubMenu, ExtenderStack.Top(), bIsSubMenu, bInOpenSubMenuOnClick, CommandListStack.Last(), bCloseSelfOnly, bInShouldCloseWindowAfterMenuSelection );
 	NewMenuEntryBlock->SetRecursivelySearchable(bRecursivelySearchable);
 
+	if (InVisibility.IsSet())
+	{
+		NewMenuEntryBlock->SetVisibilityOverride(InVisibility);
+	}
+
 	MultiBox->AddMultiBlock( NewMenuEntryBlock );
 }
 
-void FMenuBuilder::AddSubMenu( const FUIAction& UIAction, const TSharedRef< SWidget > Contents, const FNewMenuDelegate& InSubMenu, const bool bInShouldCloseWindowAfterMenuSelection /*= true*/ )
+void FMenuBuilder::AddSubMenu( const FUIAction& UIAction, const TSharedRef< SWidget > Contents, const FNewMenuDelegate& InSubMenu, const bool bInShouldCloseWindowAfterMenuSelection /*= true*/, const TAttribute<EVisibility>& InVisibility )
 {
 	ApplySectionBeginning();
 
@@ -314,6 +361,11 @@ void FMenuBuilder::AddSubMenu( const FUIAction& UIAction, const TSharedRef< SWid
 	TSharedRef< FMenuEntryBlock > NewMenuEntryBlock = MakeShared<FMenuEntryBlock>( NAME_None, UIAction, Contents, InSubMenu, ExtenderStack.Top(), bIsSubMenu, CommandListStack.Last(), bCloseSelfOnly, bInShouldCloseWindowAfterMenuSelection );
 	NewMenuEntryBlock->SetRecursivelySearchable(bRecursivelySearchable);
 
+	if (InVisibility.IsSet())
+	{
+		NewMenuEntryBlock->SetVisibilityOverride(InVisibility);
+	}
+
 	MultiBox->AddMultiBlock( NewMenuEntryBlock );
 }
 
@@ -358,13 +410,18 @@ void FMenuBuilder::AddWidget( TSharedRef<SWidget> InWidget, const FText& Label,
 	return AddWidget(InWidget, Label, StyleParams, bInSearchable, InToolTipText);
 }
 
-void FMenuBuilder::AddWidget(const TSharedRef<SWidget>& InWidget, const FText& InLabel, const FMenuEntryStyleParams& InStyleParams, const bool bInSearchable, const TAttribute<FText>& InToolTipText, const TAttribute<FSlateIcon>& InIcon)
+void FMenuBuilder::AddWidget(const TSharedRef<SWidget>& InWidget, const FText& InLabel, const FMenuEntryStyleParams& InStyleParams, const bool bInSearchable, const TAttribute<FText>& InToolTipText, const TAttribute<FSlateIcon>& InIcon, const TAttribute<EVisibility>& InVisibility)
 {
 	ApplySectionBeginning();
 
 	TSharedRef<FWidgetBlock> NewWidgetBlock = MakeShared<FWidgetBlock>(InWidget, InLabel, InToolTipText, InStyleParams, InIcon);
 	NewWidgetBlock->SetSearchable( bInSearchable );
 
+	if (InVisibility.IsSet())
+	{
+		NewWidgetBlock->SetVisibilityOverride(InVisibility);
+	}
+
 	MultiBox->AddMultiBlock( NewWidgetBlock );
 }
 
@@ -395,7 +452,14 @@ void FMenuBuilder::ApplySectionBeginning()
 	{
 		if (!CurrentSectionHeadingText.IsEmpty())
 		{
-			MultiBox->AddMultiBlock( MakeShareable( new FHeadingBlock(CurrentSectionExtensionHook, CurrentSectionHeadingText) ) );
+			TSharedRef<FHeadingBlock> HeadingBlock = MakeShareable( new FHeadingBlock(CurrentSectionExtensionHook, CurrentSectionHeadingText) ) ;
+
+			if (CurrentSectionVisibility.IsSet())
+			{
+				HeadingBlock->SetVisibilityOverride(CurrentSectionVisibility);
+			}
+
+			MultiBox->AddMultiBlock(HeadingBlock);
 		}
 		bSectionNeedsToBeApplied = false;
 		CurrentSectionHeadingText = FText::GetEmpty();
diff --git a/Engine/Source/Runtime/Slate/Private/Framework/MultiBox/SHeadingBlock.cpp b/Engine/Source/Runtime/Slate/Private/Framework/MultiBox/SHeadingBlock.cpp
--- a/Engine/Source/Runtime/Slate/Private/Framework/MultiBox/SHeadingBlock.cpp
+++ b/Engine/Source/Runtime/Slate/Private/Framework/MultiBox/SHeadingBlock.cpp
@@ -74,4 +74,6 @@ void SHeadingBlock::BuildMultiBlockWidget(const ISlateStyle* StyleSet, const FNa
 			.SeparatorImage(StyleSet->GetBrush(StyleName, ".Separator") )
 		]
 	];
+
+	SetVisibility(MultiBlock->GetVisibilityOverride());
 }
diff --git a/Engine/Source/Runtime/Slate/Private/Framework/MultiBox/SMenuSeparatorBlock.cpp b/Engine/Source/Runtime/Slate/Private/Framework/MultiBox/SMenuSeparatorBlock.cpp
--- a/Engine/Source/Runtime/Slate/Private/Framework/MultiBox/SMenuSeparatorBlock.cpp
+++ b/Engine/Source/Runtime/Slate/Private/Framework/MultiBox/SMenuSeparatorBlock.cpp
@@ -57,6 +57,8 @@ void SMenuSeparatorBlock::BuildMultiBlockWidget(const ISlateStyle* StyleSet, con
 		]
 	];
 
+	SetVisibility(MultiBlock->GetVisibilityOverride());
+
 	// Add this widget to the search list of the multibox and hide it
 	OwnerMultiBoxWidget.Pin()->AddElement(this->AsWidget(), FText::GetEmpty(), MultiBlock->GetSearchable());
 }
diff --git a/Engine/Source/Runtime/Slate/Public/Framework/MultiBox/MultiBoxBuilder.h b/Engine/Source/Runtime/Slate/Public/Framework/MultiBox/MultiBoxBuilder.h
--- a/Engine/Source/Runtime/Slate/Public/Framework/MultiBox/MultiBoxBuilder.h
+++ b/Engine/Source/Runtime/Slate/Public/Framework/MultiBox/MultiBoxBuilder.h
@@ -193,6 +193,9 @@ struct FMenuEntryParams : public FMultiBlock::FMultiBlockParams
 	/** Optional overridden input binding text for this menu entry.  If not set, then the UI action's binding will be used if available. */
 	TAttribute<FText> InputBindingOverride;
 
+	/** Set in case this entry needs to be dynamically hidden */
+	TAttribute<EVisibility> Visibility;
+
 	/** Optional overridden icon for this tool bar button.  IF not set, then the action's icon will be used instead. */
 	FSlateIcon IconOverride;
 
@@ -259,8 +262,9 @@ public:
 	 * @param	InToolTipOverride	Optional tool tip override.	 If omitted, then the action's label will be used instead.
 	 * @param	InIconOverride		Optional name of the slate brush to use for the tool bar image.  If omitted, then the command's icon will be used instead.
 	 * @param	InTutorialHighlightName	Optional name to identify this widget and highlight during tutorials
+	 * @param   InVisibility Optional Visibility Override. Can be used to show/hide this entry dynamically.
 	 */
-	SLATE_API void AddMenuEntry( const TSharedPtr< const FUICommandInfo > InCommand, FName InExtensionHook = NAME_None, const TAttribute<FText>& InLabelOverride = TAttribute<FText>(), const TAttribute<FText>& InToolTipOverride = TAttribute<FText>(), const FSlateIcon& InIconOverride = FSlateIcon(), FName InTutorialHighlightName = NAME_None );
+	SLATE_API void AddMenuEntry( const TSharedPtr< const FUICommandInfo > InCommand, FName InExtensionHook = NAME_None, const TAttribute<FText>& InLabelOverride = TAttribute<FText>(), const TAttribute<FText>& InToolTipOverride = TAttribute<FText>(), const FSlateIcon& InIconOverride = FSlateIcon(), FName InTutorialHighlightName = NAME_None, const TAttribute<EVisibility>& InVisibility = TAttribute<EVisibility>());
 
 	/**
 	 * Adds a menu entry without the use of a command
@@ -273,8 +277,9 @@ public:
 	 * @param	UserInterfaceActionType	Type of interface action
 	 * @param	InTutorialHighlightName	Optional name to identify this widget and highlight during tutorials
 	 * @param	InInputBindingOverride  Optional overridden input binding text for this menu entry.  If not set, then the UI action's binding will be used if available.
+	 * @param   InVisibility Optional Visibility Override. Can be used to show/hide this entry dynamically.
 	 */
-	SLATE_API void AddMenuEntry( const TAttribute<FText>& InLabel, const TAttribute<FText>& InToolTip, const FSlateIcon& InIcon, const FUIAction& UIAction, FName InExtensionHook = NAME_None, const EUserInterfaceActionType UserInterfaceActionType = EUserInterfaceActionType::Button, FName InTutorialHighlightName = NAME_None, const TAttribute<FText>& InInputBindingOverride = TAttribute<FText>());
+	SLATE_API void AddMenuEntry( const TAttribute<FText>& InLabel, const TAttribute<FText>& InToolTip, const FSlateIcon& InIcon, const FUIAction& UIAction, FName InExtensionHook = NAME_None, const EUserInterfaceActionType UserInterfaceActionType = EUserInterfaceActionType::Button, FName InTutorialHighlightName = NAME_None, const TAttribute<FText>& InInputBindingOverride = TAttribute<FText>(), const TAttribute<EVisibility>& InVisibility = TAttribute<EVisibility>());
 	
 	/**
 	 * Adds a menu entry with a custom widget
@@ -285,8 +290,9 @@ public:
 	 * @param	InToolTip				Tool tip used when hovering over the menu entry
 	 * @param	UserInterfaceActionType	Type of interface action
 	 * @param	InTutorialHighlightName	Optional name to identify this widget and highlight during tutorials
+	 * @param   InVisibility Optional Visibility Override. Can be used to show/hide this entry dynamically.
 	 */
-	SLATE_API void AddMenuEntry( const FUIAction& UIAction, const TSharedRef< SWidget > Contents, const FName& InExtensionHook = NAME_None, const TAttribute<FText>& InToolTip = TAttribute<FText>(), const EUserInterfaceActionType UserInterfaceActionType = EUserInterfaceActionType::Button, FName InTutorialHighlightName = NAME_None );
+	SLATE_API void AddMenuEntry( const FUIAction& UIAction, const TSharedRef< SWidget > Contents, const FName& InExtensionHook = NAME_None, const TAttribute<FText>& InToolTip = TAttribute<FText>(), const EUserInterfaceActionType UserInterfaceActionType = EUserInterfaceActionType::Button, FName InTutorialHighlightName = NAME_None, const TAttribute<EVisibility>& InVisibility = TAttribute<EVisibility>());
 
 	/** Adds a menu entry with given param struct */
 	SLATE_API void AddMenuEntry( const FMenuEntryParams& InMenuEntryParams );
@@ -341,16 +347,17 @@ public:
 	/**
 	 * Adds a separator
 	 */
-	SLATE_API void AddMenuSeparator(FName InExtensionHook = NAME_None);
-	SLATE_API void AddSeparator(FName InExtensionHook = NAME_None);
+	SLATE_API void AddMenuSeparator(FName InExtensionHook = NAME_None, const TAttribute<EVisibility>& InVisibility = TAttribute<EVisibility>());
+	SLATE_API void AddSeparator(FName InExtensionHook = NAME_None, const TAttribute<EVisibility>& InVisibility = TAttribute<EVisibility>());
 	
 	/**
 	 * Starts a section on to the extender section hook stack
 	 * 
 	 * @param InExtensionHook	The section hook. Can be NAME_None
 	 * @param InHeadingText		The heading text to use. If none, only a separator is used
+	 * @param InVisibility Optional Visibility Override. Can be used to show/hide this Section dynamically.
 	 */
-	SLATE_API void BeginSection( FName InExtensionHook, const TAttribute< FText >& InHeadingText = TAttribute<FText>() );
+	SLATE_API void BeginSection( FName InExtensionHook, const TAttribute< FText >& InHeadingText = TAttribute<FText>(), const TAttribute<EVisibility>& InVisibility = TAttribute<EVisibility>() );
 
 	/**
 	 * Ends the current section
@@ -370,14 +377,15 @@ public:
 	 * @param bInOpenSubMenuOnClick Sub-menu will open only if the sub-menu entry is clicked
 	 * @param	InIcon				The icon to use
 	 * @param	bInShouldCloseWindowAfterMenuSelection	Whether the submenu should close after an item is selected
+	 * @param   InVisibility Optional Visibility Override. Can be used to show/hide this entry dynamically.
 	 */
-	SLATE_API void AddSubMenu( const TAttribute<FText>& InMenuLabel, const TAttribute<FText>& InToolTip, const FNewMenuDelegate& InSubMenu, const FUIAction& InUIAction, FName InExtensionHook, const EUserInterfaceActionType InUserInterfaceActionType, const bool bInOpenSubMenuOnClick = false, const FSlateIcon& InIcon = FSlateIcon(), const bool bInShouldCloseWindowAfterMenuSelection = true );
+	SLATE_API void AddSubMenu( const TAttribute<FText>& InMenuLabel, const TAttribute<FText>& InToolTip, const FNewMenuDelegate& InSubMenu, const FUIAction& InUIAction, FName InExtensionHook, const EUserInterfaceActionType InUserInterfaceActionType, const bool bInOpenSubMenuOnClick = false, const FSlateIcon& InIcon = FSlateIcon(), const bool bInShouldCloseWindowAfterMenuSelection = true, const TAttribute<EVisibility>& InVisibility = TAttribute<EVisibility>() );
 
-	SLATE_API void AddSubMenu( const TAttribute<FText>& InMenuLabel, const TAttribute<FText>& InToolTip, const FNewMenuDelegate& InSubMenu, const bool bInOpenSubMenuOnClick = false, const FSlateIcon& InIcon = FSlateIcon(), const bool bInShouldCloseWindowAfterMenuSelection = true, FName InExtensionHook = NAME_None, FName InTutorialHighlightName = NAME_None );
+	SLATE_API void AddSubMenu( const TAttribute<FText>& InMenuLabel, const TAttribute<FText>& InToolTip, const FNewMenuDelegate& InSubMenu, const bool bInOpenSubMenuOnClick = false, const FSlateIcon& InIcon = FSlateIcon(), const bool bInShouldCloseWindowAfterMenuSelection = true, FName InExtensionHook = NAME_None, FName InTutorialHighlightName = NAME_None, const TAttribute<EVisibility>& InVisibility = TAttribute<EVisibility>() );
 
-	SLATE_API void AddSubMenu( const TSharedRef< SWidget > Contents, const FNewMenuDelegate& InSubMenu, const bool bInOpenSubMenuOnClick = false, const bool bInShouldCloseWindowAfterMenuSelection = true );
+	SLATE_API void AddSubMenu( const TSharedRef< SWidget > Contents, const FNewMenuDelegate& InSubMenu, const bool bInOpenSubMenuOnClick = false, const bool bInShouldCloseWindowAfterMenuSelection = true, const TAttribute<EVisibility>& InVisibility = TAttribute<EVisibility>() );
 
-	SLATE_API void AddSubMenu( const FUIAction& UIAction, const TSharedRef< SWidget > Contents, const FNewMenuDelegate& InSubMenu, const bool bInShouldCloseWindowAfterMenuSelection = true );
+	SLATE_API void AddSubMenu( const FUIAction& UIAction, const TSharedRef< SWidget > Contents, const FNewMenuDelegate& InSubMenu, const bool bInShouldCloseWindowAfterMenuSelection = true, const TAttribute<EVisibility>& InVisibility = TAttribute<EVisibility>() );
 
 	/**
 	 * Adds any widget to the menu
@@ -400,7 +408,7 @@ public:
 	 * @param	InToolTipText		Optional tooltip text to be added to the widget and label
 	 * @param	InIcon				Optional icon to be shown to the left of the label/content
 	 */
-	SLATE_API void AddWidget(const TSharedRef<SWidget>& InWidget, const FText& InLabel, const FMenuEntryStyleParams& InStyleParams, bool bInSearchable = true, const TAttribute<FText>& InToolTipText = FText(), const TAttribute<FSlateIcon>& InIcon = TAttribute<FSlateIcon>());
+	SLATE_API void AddWidget(const TSharedRef<SWidget>& InWidget, const FText& InLabel, const FMenuEntryStyleParams& InStyleParams, bool bInSearchable = true, const TAttribute<FText>& InToolTipText = FText(), const TAttribute<FSlateIcon>& InIcon = TAttribute<FSlateIcon>(), const TAttribute<EVisibility>& InVisibility = TAttribute<EVisibility>());
 
 	/**
 	* Adds the widget the multibox will use for searching
@@ -443,7 +451,10 @@ private:
 	
 	/** Any pending section's heading text */
 	FText CurrentSectionHeadingText;
-	
+
+	/** Keep track of current section visibility, in case the section needs to be dynamically hidden */
+	TAttribute<EVisibility> CurrentSectionVisibility;
+
 	/** True if there is a pending section that needs to be applied */
 	bool bSectionNeedsToBeApplied;
 
