[SDockTab] SDockTab can no longer be closed when `CanEverClose` is set to false

diff --git a/Engine/Source/Runtime/Slate/Private/Widgets/Docking/SDockTab.cpp b/Engine/Source/Runtime/Slate/Private/Widgets/Docking/SDockTab.cpp
--- a/Engine/Source/Runtime/Slate/Private/Widgets/Docking/SDockTab.cpp
+++ b/Engine/Source/Runtime/Slate/Private/Widgets/Docking/SDockTab.cpp
@@ -993,6 +993,11 @@ bool SDockTab::CanCloseTab() const
 	// Check the below conditions separately and early out to prevent the OnCanCloseTab delegate
 	// from being called if we already know we cannot close the tab.
 
+	if (!bCanEverClose)
+	{
+		return false;
+	}
+
 	if (TSharedPtr<FTabManager> Manager = MyTabManager.Pin();
 		Manager && !Manager->IsTabCloseable(SharedThis(this)))
 	{
