[SInlineEditableTextBlock] Add Slate argument 'DelayedLeftClickEntersEditMode' to enable/disable left click + delay to enter edit mode.

Default is set to true, to respect current behavior.

#rnx

diff --git a/Engine/Source/Runtime/Slate/Private/Widgets/Text/SInlineEditableTextBlock.cpp b/Engine/Source/Runtime/Slate/Private/Widgets/Text/SInlineEditableTextBlock.cpp
--- a/Engine/Source/Runtime/Slate/Private/Widgets/Text/SInlineEditableTextBlock.cpp
+++ b/Engine/Source/Runtime/Slate/Private/Widgets/Text/SInlineEditableTextBlock.cpp
@@ -22,6 +22,7 @@ void SInlineEditableTextBlock::Construct( const FArguments& InArgs )
 	MaximumLength = InArgs._MaximumLength;
 
 	bIsMultiLine = InArgs._MultiLine;
+	bDelayedLeftClickEntersEditMode = InArgs._DelayedLeftClickEntersEditMode;
 	DoubleSelectDelay = 0.0f;
 
 	OnEnterEditingMode = InArgs._OnEnterEditingMode;
@@ -221,7 +222,7 @@ void SInlineEditableTextBlock::SetOverflowPolicy(TOptional<ETextOverflowPolicy>
 
 FReply SInlineEditableTextBlock::OnMouseButtonDown( const FGeometry& MyGeometry, const FPointerEvent& MouseEvent )
 {
-	if( !MouseEvent.IsMouseButtonDown(EKeys::LeftMouseButton) || MouseEvent.IsControlDown() || MouseEvent.IsShiftDown())
+	if (!bDelayedLeftClickEntersEditMode || !MouseEvent.IsMouseButtonDown(EKeys::LeftMouseButton) || MouseEvent.IsControlDown() || MouseEvent.IsShiftDown())
 	{
 		return FReply::Unhandled();
 	}
diff --git a/Engine/Source/Runtime/Slate/Public/Widgets/Text/SInlineEditableTextBlock.h b/Engine/Source/Runtime/Slate/Public/Widgets/Text/SInlineEditableTextBlock.h
--- a/Engine/Source/Runtime/Slate/Public/Widgets/Text/SInlineEditableTextBlock.h
+++ b/Engine/Source/Runtime/Slate/Public/Widgets/Text/SInlineEditableTextBlock.h
@@ -46,6 +46,7 @@ class SInlineEditableTextBlock: public SCompoundWidget
 		, _LineBreakPolicy()
 		, _IsReadOnly(false)
 		, _MultiLine(false)
+		, _DelayedLeftClickEntersEditMode(true)
 		, _ModiferKeyForNewLine(EModifierKey::None)
 		, _OverflowPolicy()
 	{
@@ -96,6 +97,9 @@ class SInlineEditableTextBlock: public SCompoundWidget
 		/** True if the editable text block is multi-line */
 		SLATE_ARGUMENT( bool, MultiLine )
 
+		/** True if left-clicking will enter edit mode after a delay. */
+		SLATE_ARGUMENT( bool, DelayedLeftClickEntersEditMode )
+
 		/** The optional modifier key necessary to create a newline when typing into the editor. */
 		SLATE_ARGUMENT(EModifierKey::Type, ModiferKeyForNewLine)
 
@@ -261,4 +265,6 @@ protected:
 	/** Attribute to look up if the widget is multiline */
 	uint8 bIsMultiLine:1;
 
+	/** Enable left-clicking the text block to enter edit mode. */
+	uint8 bDelayedLeftClickEntersEditMode:1;
 };
