Added setters for SInlineEditableTextBlock delegates.

diff --git a/Engine/Source/Runtime/Slate/Private/Widgets/Text/SInlineEditableTextBlock.cpp b/Engine/Source/Runtime/Slate/Private/Widgets/Text/SInlineEditableTextBlock.cpp
--- a/Engine/Source/Runtime/Slate/Private/Widgets/Text/SInlineEditableTextBlock.cpp
+++ b/Engine/Source/Runtime/Slate/Private/Widgets/Text/SInlineEditableTextBlock.cpp
@@ -220,6 +220,36 @@ void SInlineEditableTextBlock::SetOverflowPolicy(TOptional<ETextOverflowPolicy>
 	TextBlock->SetOverflowPolicy(InOverflowPolicy);
 }
 
+void SInlineEditableTextBlock::SetOnBeginTextEdit(FOnBeginTextEdit InOnBeginTextEdit)
+{
+	OnBeginTextEditDelegate = MoveTemp(InOnBeginTextEdit);
+}
+
+void SInlineEditableTextBlock::SetOnTextCommitted(FOnTextCommitted InOnTextCommitted)
+{
+	OnTextCommittedDelegate = MoveTemp(InOnTextCommitted);
+}
+
+void SInlineEditableTextBlock::SetOnEnterEditingMode(FSimpleDelegate InOnEnterEditingMode)
+{
+	OnEnterEditingMode = MoveTemp(InOnEnterEditingMode);
+}
+
+void SInlineEditableTextBlock::SetOnExitEditingMode(FSimpleDelegate InOnExitEditingMode)
+{
+	OnExitEditingMode = MoveTemp(InOnExitEditingMode);
+}
+
+void SInlineEditableTextBlock::SetOnExitEditingMode(FIsSelected InIsSelected)
+{
+	IsSelected = MoveTemp(InIsSelected);
+}
+
+void SInlineEditableTextBlock::SetOnVerifyTextChanged(FOnVerifyTextChanged InOnVerifyTextChanged)
+{
+	OnVerifyTextChanged = MoveTemp(InOnVerifyTextChanged);
+}
+
 FReply SInlineEditableTextBlock::OnMouseButtonDown( const FGeometry& MyGeometry, const FPointerEvent& MouseEvent )
 {
 	if (!bDelayedLeftClickEntersEditMode || !MouseEvent.IsMouseButtonDown(EKeys::LeftMouseButton) || MouseEvent.IsControlDown() || MouseEvent.IsShiftDown())
diff --git a/Engine/Source/Runtime/Slate/Public/Widgets/Text/SInlineEditableTextBlock.h b/Engine/Source/Runtime/Slate/Public/Widgets/Text/SInlineEditableTextBlock.h
--- a/Engine/Source/Runtime/Slate/Public/Widgets/Text/SInlineEditableTextBlock.h
+++ b/Engine/Source/Runtime/Slate/Public/Widgets/Text/SInlineEditableTextBlock.h
@@ -186,6 +186,18 @@ class SInlineEditableTextBlock: public SCompoundWidget
 	/** Sets the overflow policy for this text block */
 	SLATE_API void SetOverflowPolicy(TOptional<ETextOverflowPolicy> InOverflowPolicy);
 
+	SLATE_API void SetOnBeginTextEdit(FOnBeginTextEdit InOnBeginTextEdit);
+
+	SLATE_API void SetOnTextCommitted(FOnTextCommitted InOnTextCommitted);
+
+	SLATE_API void SetOnEnterEditingMode(FSimpleDelegate InOnEnterEditingMode);
+
+	SLATE_API void SetOnExitEditingMode(FSimpleDelegate InOnExitEditingMode);
+
+	SLATE_API void SetOnExitEditingMode(FIsSelected InIsSelected);
+
+	SLATE_API void SetOnVerifyTextChanged(FOnVerifyTextChanged InOnVerifyTextChanged);
+
 protected:
 	/** Callback for the text box's OnTextChanged event */
 	SLATE_API void OnTextChanged(const FText& InText);
