Slate: Rename NavigateToNextWidgetByDirection to NavigateFromWidget.
CommonUI: Add CommonUI fallback on navigation failure.

#rnx

diff --git a/Engine/Plugins/Runtime/CommonUI/Source/CommonUI/Private/CommonButtonBase.cpp b/Engine/Plugins/Runtime/CommonUI/Source/CommonUI/Private/CommonButtonBase.cpp
--- a/Engine/Plugins/Runtime/CommonUI/Source/CommonUI/Private/CommonButtonBase.cpp
+++ b/Engine/Plugins/Runtime/CommonUI/Source/CommonUI/Private/CommonButtonBase.cpp
@@ -2088,7 +2088,9 @@ void UCommonButtonBase::HandleImplicitFocusLost()
 		return;
 	}
 
-	FSlateApplication::Get().ForEachUser([SafeWidget](FSlateUser& User)
+	const UCommonButtonBase& ThisRef = *this;
+
+	FSlateApplication::Get().ForEachUser([SafeWidget, &ThisRef](FSlateUser& User)
 	{
 		if (User.IsWidgetInFocusPath(SafeWidget.ToSharedRef()))
 		{
@@ -2096,7 +2098,14 @@ void UCommonButtonBase::HandleImplicitFocusLost()
 			// Using Navigation::Next will first descend the widget hierarchy and focus on the SCommonButton regardless of the outer SObjectWidget state. 
 			// To avoid this, use the top level directional navigation from the SObjectWidget level
 			const TArray<EUINavigation> Directions = { EUINavigation::Right, EUINavigation::Down,  EUINavigation::Left,  EUINavigation::Up };
-			FSlateApplication::Get().NavigateToNextWidgetByDirection(User.GetUserIndex(), SafeWidget, Directions);
+			if (FSlateApplication::Get().NavigateFromWidget(User.GetUserIndex(), SafeWidget, Directions) == EUINavigation::Invalid)
+			{
+				// Fallback if failed to navigate to any other widgets. Reset the focus state
+				if (UCommonUIActionRouterBase* ActionRouter = UCommonUIActionRouterBase::Get(ThisRef))
+				{
+					ActionRouter->RefreshActiveRootFocus();
+				}
+			}
 		}
 	});
 }
diff --git a/Engine/Source/Runtime/Slate/Private/Framework/Application/SlateApplication.cpp b/Engine/Source/Runtime/Slate/Private/Framework/Application/SlateApplication.cpp
--- a/Engine/Source/Runtime/Slate/Private/Framework/Application/SlateApplication.cpp
+++ b/Engine/Source/Runtime/Slate/Private/Framework/Application/SlateApplication.cpp
@@ -6223,7 +6223,7 @@ void FSlateApplication::NavigateToWidget(const uint32 UserIndex, const TSharedPt
 	}
 }
 
-void FSlateApplication::NavigateToNextWidgetByDirection(const uint32 InUserIndex, TSharedPtr<SWidget> InNavigationSource, const TArray<EUINavigation>& InDirections)
+EUINavigation FSlateApplication::NavigateFromWidget(const uint32 InUserIndex, TSharedPtr<SWidget> InNavigationSource, const TArray<EUINavigation>& InNavigationTypes)
 {
 	FWidgetPath NavigationSourcePath;
 	if (InNavigationSource.IsValid())
@@ -6237,8 +6237,8 @@ void FSlateApplication::NavigateToNextWidgetByDirection(const uint32 InUserIndex
 
 	if (!NavigationSourcePath.IsValid())
 	{
-		UE_LOG(LogSlate, Log, TEXT("Failed to navigate to next available widget. No widget currently in Focus"));
-		return;
+		UE_LOG(LogSlate, Log, TEXT("Failed to navigate from widget %s. No widget currently in Focus"), *(InNavigationSource->ToString()));
+		return EUINavigation::Invalid;
 	}
 
 	// Traverse up the widget tree and attempt to navigate to next available widget by searching neighbor space
@@ -6246,16 +6246,19 @@ void FSlateApplication::NavigateToNextWidgetByDirection(const uint32 InUserIndex
 	{
 		const FArrangedWidget& BoundaryWidget = NavigationSourcePath.Widgets[WidgetIndex];
 
-		for (EUINavigation Direction : InDirections)
+		for (EUINavigation Direction : InNavigationTypes)
 		{
 			// Try to advance to next widget
 			const FNavigationEvent NavigationEventNext(PlatformApplication->GetModifierKeys(), InUserIndex, Direction, ENavigationGenesis::User);
 			if (AttemptNavigation(NavigationSourcePath, NavigationEventNext, FNavigationReply::Escape(), BoundaryWidget))
 			{
-				return;
+				return Direction;
 			}
 		}
 	}
+
+	UE_LOG(LogSlate, Log, TEXT("Failed to navigate from widget %s. No neighboring widgets could be reached."), *(InNavigationSource->ToString()));
+	return EUINavigation::Invalid;
 }
 
 bool FSlateApplication::AttemptNavigation(const FWidgetPath& NavigationSource, const FNavigationEvent& NavigationEvent, const FNavigationReply& NavigationReply, const FArrangedWidget& BoundaryWidget)
diff --git a/Engine/Source/Runtime/Slate/Public/Framework/Application/SlateApplication.h b/Engine/Source/Runtime/Slate/Public/Framework/Application/SlateApplication.h
--- a/Engine/Source/Runtime/Slate/Public/Framework/Application/SlateApplication.h
+++ b/Engine/Source/Runtime/Slate/Public/Framework/Application/SlateApplication.h
@@ -1815,9 +1815,11 @@ public:
 	 *
 	 * @param InUserIndex The user that is doing the navigation
 	 * @param InNavigationSource The source widget to navigate from. If null, will use current user focus
-	 * @param InDirections The direction space to search in
+	 * @param InNavigationTypes The navigation requests to attempt, in order
+
+	 * @return If succesful, returns the direction navigated. Otherwise return EUINavigation::Invalid
 	 */
-	SLATE_API void NavigateToNextWidgetByDirection(const uint32 InUserIndex, TSharedPtr<SWidget> InNavigationSource, const TArray<EUINavigation>& InDirections);
+	SLATE_API EUINavigation NavigateFromWidget(const uint32 InUserIndex, TSharedPtr<SWidget> InNavigationSource, const TArray<EUINavigation>& InNavigationTypes);
 
 	/**
 	* Given an optional widget, try and get the most suitable parent window to use with dialogs (such as file and directory pickers).
