[Viewport Toolbar] Support display of hotkey on "checkbox submenus" that are bound to a command

#okfrogithub public

diff --git a/Engine/Source/Developer/ToolMenus/Private/ToolMenus.cpp b/Engine/Source/Developer/ToolMenus/Private/ToolMenus.cpp
--- a/Engine/Source/Developer/ToolMenus/Private/ToolMenus.cpp
+++ b/Engine/Source/Developer/ToolMenus/Private/ToolMenus.cpp
@@ -273,7 +273,8 @@ public:
 						Block.SubMenuData.bOpenSubMenuOnClick,
 						Block.Icon.Get(),
 						Block.bShouldCloseWindowAfterMenuSelection,
-						VisibilityOverride
+						VisibilityOverride,
+						Block.InputBindingLabel
 					);
 				}
 				else
diff --git a/Engine/Source/Runtime/Slate/Private/Framework/MultiBox/MultiBoxBuilder.cpp b/Engine/Source/Runtime/Slate/Private/Framework/MultiBox/MultiBoxBuilder.cpp
--- a/Engine/Source/Runtime/Slate/Private/Framework/MultiBox/MultiBoxBuilder.cpp
+++ b/Engine/Source/Runtime/Slate/Private/Framework/MultiBox/MultiBoxBuilder.cpp
@@ -303,12 +303,12 @@ void FMenuBuilder::AddSeparator(FName InExtensionHook /*= NAME_None*/, const TAt
 	ApplyHook(InExtensionHook, EExtensionHook::After);
 }
 
-void FMenuBuilder::AddSubMenu( const TAttribute<FText>& InMenuLabel, const TAttribute<FText>& InToolTip, const FNewMenuDelegate& InSubMenu, const FUIAction& InUIAction, FName InExtensionHook, const EUserInterfaceActionType InUserInterfaceActionType, const bool bInOpenSubMenuOnClick, const FSlateIcon& InIcon, const bool bInShouldCloseWindowAfterMenuSelection /*= true*/, const TAttribute<EVisibility>& InVisibility )
+void FMenuBuilder::AddSubMenu( const TAttribute<FText>& InMenuLabel, const TAttribute<FText>& InToolTip, const FNewMenuDelegate& InSubMenu, const FUIAction& InUIAction, FName InExtensionHook, const EUserInterfaceActionType InUserInterfaceActionType, const bool bInOpenSubMenuOnClick, const FSlateIcon& InIcon, const bool bInShouldCloseWindowAfterMenuSelection /*= true*/, const TAttribute<EVisibility>& InVisibility, const TAttribute<FText>& InInputBindingOverride )
 {
 	ApplySectionBeginning();
 
 	const bool bIsSubMenu = true;
-	TSharedRef< FMenuEntryBlock > NewMenuEntryBlock = MakeShared<FMenuEntryBlock>( InExtensionHook, InMenuLabel, InToolTip, InSubMenu, ExtenderStack.Top(), bIsSubMenu, bInOpenSubMenuOnClick, InIcon, InUIAction, InUserInterfaceActionType, bCloseSelfOnly, bInShouldCloseWindowAfterMenuSelection, CommandListStack.Last() );
+	TSharedRef< FMenuEntryBlock > NewMenuEntryBlock = MakeShared<FMenuEntryBlock>( InExtensionHook, InMenuLabel, InToolTip, InSubMenu, ExtenderStack.Top(), bIsSubMenu, bInOpenSubMenuOnClick, InIcon, InUIAction, InUserInterfaceActionType, bCloseSelfOnly, bInShouldCloseWindowAfterMenuSelection, CommandListStack.Last(), InInputBindingOverride);
 	NewMenuEntryBlock->SetRecursivelySearchable(bRecursivelySearchable);
 
 	if (InVisibility.IsSet())
diff --git a/Engine/Source/Runtime/Slate/Private/Framework/MultiBox/SMenuEntryBlock.cpp b/Engine/Source/Runtime/Slate/Private/Framework/MultiBox/SMenuEntryBlock.cpp
--- a/Engine/Source/Runtime/Slate/Private/Framework/MultiBox/SMenuEntryBlock.cpp
+++ b/Engine/Source/Runtime/Slate/Private/Framework/MultiBox/SMenuEntryBlock.cpp
@@ -31,10 +31,11 @@ FMenuEntryBlock::FMenuEntryBlock( const FName& InExtensionHook, const TSharedPtr
 }
 
 
-FMenuEntryBlock::FMenuEntryBlock( const FName& InExtensionHook, const TAttribute<FText>& InLabel, const TAttribute<FText>& InToolTip, const FNewMenuDelegate& InEntryBuilder, TSharedPtr<FExtender> InExtender, bool bInSubMenu, bool bInSubMenuOnClick, const FSlateIcon& InIcon, const FUIAction& InUIAction, const EUserInterfaceActionType InUserInterfaceActionType, bool bInCloseSelfOnly, bool bInShouldCloseWindowAfterMenuSelection, TSharedPtr< const FUICommandList > InCommandList)
+FMenuEntryBlock::FMenuEntryBlock( const FName& InExtensionHook, const TAttribute<FText>& InLabel, const TAttribute<FText>& InToolTip, const FNewMenuDelegate& InEntryBuilder, TSharedPtr<FExtender> InExtender, bool bInSubMenu, bool bInSubMenuOnClick, const FSlateIcon& InIcon, const FUIAction& InUIAction, const EUserInterfaceActionType InUserInterfaceActionType, bool bInCloseSelfOnly, bool bInShouldCloseWindowAfterMenuSelection, TSharedPtr< const FUICommandList > InCommandList, const TAttribute<FText>& InInputBindingOverride)
 	: FMultiBlock( InUIAction, InExtensionHook, EMultiBlockType::MenuEntry, /* bInIsPartOfHeading = */ false, InCommandList )
 	, LabelOverride( InLabel )
 	, ToolTipOverride( InToolTip )
+	, InputBindingOverride(InInputBindingOverride)
 	, IconOverride( InIcon )
 	, EntryBuilder( InEntryBuilder )
 	, bIsSubMenu( bInSubMenu )
@@ -770,6 +771,7 @@ TSharedRef< SWidget> SMenuEntryBlock::BuildSubMenuWidget( const FMenuEntryBuildP
 {
 	const TAttribute<FText>& Label = InBuildParams.Label;
 	const TAttribute<FText>& EntryToolTip = InBuildParams.ToolTip;
+	const TAttribute<FText>& InputBinding = InBuildParams.InputBinding;
 
 	const TSharedPtr< const FMenuEntryBlock > MenuEntryBlock = InBuildParams.MenuEntryBlock;
 	const TSharedPtr< const FMultiBox > MultiBox = InBuildParams.MultiBox;
@@ -927,6 +929,22 @@ TSharedRef< SWidget> SMenuEntryBlock::BuildSubMenuWidget( const FMenuEntryBuildP
 				]
 				+ SHorizontalBox::Slot()
 				.AutoWidth()
+				.Padding(FMargin(25,0,0,0))
+				.VAlign( VAlign_Center )
+				.HAlign( HAlign_Right )
+				[
+					SNew( SBox )
+					.Visibility(InputBinding.Get().IsEmpty() ? EVisibility::Collapsed : EVisibility::Visible)
+					.Padding(FMargin(16,0,4,0))
+					[
+						SNew( STextBlock )
+						.TextStyle( StyleSet, ISlateStyle::Join( StyleName, ".Keybinding" ) )
+						.ColorAndOpacity( FSlateColor::UseSubduedForeground() )
+						.Text( InputBinding )
+					]
+				]
+				+ SHorizontalBox::Slot()
+				.AutoWidth()
 				.VAlign( VAlign_Center )
 				.HAlign( HAlign_Right )
 				[
diff --git a/Engine/Source/Runtime/Slate/Private/Framework/MultiBox/SMenuEntryBlock.h b/Engine/Source/Runtime/Slate/Private/Framework/MultiBox/SMenuEntryBlock.h
--- a/Engine/Source/Runtime/Slate/Private/Framework/MultiBox/SMenuEntryBlock.h
+++ b/Engine/Source/Runtime/Slate/Private/Framework/MultiBox/SMenuEntryBlock.h
@@ -101,8 +101,9 @@ public:
 	 * @param	bInShouldCloseWindowAfterMenuSelection	In the case of a submenu, whether it should close after an item is selected
 	 * @param	bInInvertLabelOnHover	Whether to invert the label text's color on hover
 	 * @param	InCommandList		The list of commands bound to delegates that should be executed for menu entries
+	 * @param	InInputBindingOverride		Optional overridden input binding text for this menu entry
 	 */
-	FMenuEntryBlock( const FName& InExtensionHook, const TAttribute<FText>& InLabel, const TAttribute<FText>& InToolTip, const FNewMenuDelegate& InEntryBuilder, TSharedPtr<FExtender> InExtender, bool bInSubMenu, bool bInSubMenuOnClick, const FSlateIcon& InIcon, const FUIAction& InUIAction, const EUserInterfaceActionType InUserInterfaceActionType, bool bInCloseSelfOnly, bool bInShouldCloseWindowAfterMenuSelection = true, TSharedPtr< const FUICommandList > InCommandList = nullptr );
+	FMenuEntryBlock( const FName& InExtensionHook, const TAttribute<FText>& InLabel, const TAttribute<FText>& InToolTip, const FNewMenuDelegate& InEntryBuilder, TSharedPtr<FExtender> InExtender, bool bInSubMenu, bool bInSubMenuOnClick, const FSlateIcon& InIcon, const FUIAction& InUIAction, const EUserInterfaceActionType InUserInterfaceActionType, bool bInCloseSelfOnly, bool bInShouldCloseWindowAfterMenuSelection = true, TSharedPtr< const FUICommandList > InCommandList = nullptr, const TAttribute<FText>& InInputBindingOverride = TAttribute<FText>());
 
 	FMenuEntryBlock( const FName& InExtensionHook, const TAttribute<FText>& InLabel, const TAttribute<FText>& InToolTip, const FSlateIcon& InIcon, const FUIAction& InUIAction, const EUserInterfaceActionType InUserInterfaceActionType, bool bInCloseSelfOnly, bool bInShouldCloseWindowAfterMenuSelection = true, const TAttribute<FText>& InInputBindingOverride = TAttribute<FText>());
 
diff --git a/Engine/Source/Runtime/Slate/Public/Framework/MultiBox/MultiBoxBuilder.h b/Engine/Source/Runtime/Slate/Public/Framework/MultiBox/MultiBoxBuilder.h
--- a/Engine/Source/Runtime/Slate/Public/Framework/MultiBox/MultiBoxBuilder.h
+++ b/Engine/Source/Runtime/Slate/Public/Framework/MultiBox/MultiBoxBuilder.h
@@ -379,7 +379,7 @@ public:
 	 * @param	bInShouldCloseWindowAfterMenuSelection	Whether the submenu should close after an item is selected
 	 * @param   InVisibility Optional Visibility Override. Can be used to show/hide this entry dynamically.
 	 */
-	SLATE_API void AddSubMenu( const TAttribute<FText>& InMenuLabel, const TAttribute<FText>& InToolTip, const FNewMenuDelegate& InSubMenu, const FUIAction& InUIAction, FName InExtensionHook, const EUserInterfaceActionType InUserInterfaceActionType, const bool bInOpenSubMenuOnClick = false, const FSlateIcon& InIcon = FSlateIcon(), const bool bInShouldCloseWindowAfterMenuSelection = true, const TAttribute<EVisibility>& InVisibility = TAttribute<EVisibility>() );
+	SLATE_API void AddSubMenu( const TAttribute<FText>& InMenuLabel, const TAttribute<FText>& InToolTip, const FNewMenuDelegate& InSubMenu, const FUIAction& InUIAction, FName InExtensionHook, const EUserInterfaceActionType InUserInterfaceActionType, const bool bInOpenSubMenuOnClick = false, const FSlateIcon& InIcon = FSlateIcon(), const bool bInShouldCloseWindowAfterMenuSelection = true, const TAttribute<EVisibility>& InVisibility = TAttribute<EVisibility>(), const TAttribute<FText>& InInputBindingOverride = TAttribute<FText>() );
 
 	SLATE_API void AddSubMenu( const TAttribute<FText>& InMenuLabel, const TAttribute<FText>& InToolTip, const FNewMenuDelegate& InSubMenu, const bool bInOpenSubMenuOnClick = false, const FSlateIcon& InIcon = FSlateIcon(), const bool bInShouldCloseWindowAfterMenuSelection = true, FName InExtensionHook = NAME_None, FName InTutorialHighlightName = NAME_None, const TAttribute<EVisibility>& InVisibility = TAttribute<EVisibility>() );
 
