[Truncation Warnings] Graph Editor Module - update MarqueeOperation to Vector2f

diff --git a/Engine/Plugins/Experimental/WidgetEditorToolPalette/Source/WidgetEditorToolPalette/Private/DefaultTools/RectangleSelectTool.cpp b/Engine/Plugins/Experimental/WidgetEditorToolPalette/Source/WidgetEditorToolPalette/Private/DefaultTools/RectangleSelectTool.cpp
--- a/Engine/Plugins/Experimental/WidgetEditorToolPalette/Source/WidgetEditorToolPalette/Private/DefaultTools/RectangleSelectTool.cpp
+++ b/Engine/Plugins/Experimental/WidgetEditorToolPalette/Source/WidgetEditorToolPalette/Private/DefaultTools/RectangleSelectTool.cpp
@@ -210,12 +210,12 @@ void URectangleSelectTool::SelectWidgetsAffectedByMarquee()
 		// If our selection is shrinking, use previous coordinate pair value to ensure a collision when testing for deselection
 
 		FMarqueeRect RectDeltaX;
-		RectDeltaX.StartPoint = { OldEndPoint.X, StartPoint.Y };
-		RectDeltaX.EndPoint = { NewEndPoint.X, SelectionDeltaY > 0 ? NewEndPoint.Y : OldEndPoint.Y };
+		RectDeltaX.StartPoint = FVector2f(OldEndPoint.X, StartPoint.Y);
+		RectDeltaX.EndPoint = FVector2f(NewEndPoint.X, SelectionDeltaY > 0 ? NewEndPoint.Y : OldEndPoint.Y);
 
 		FMarqueeRect RectDeltaY;
-		RectDeltaY.StartPoint = { StartPoint.X, OldEndPoint.Y };
-		RectDeltaY.EndPoint = { SelectionDeltaX > 0 ? NewEndPoint.X : OldEndPoint.X, NewEndPoint.Y };
+		RectDeltaY.StartPoint = FVector2f(StartPoint.X, OldEndPoint.Y);
+		RectDeltaY.EndPoint = FVector2f(SelectionDeltaX > 0 ? NewEndPoint.X : OldEndPoint.X, NewEndPoint.Y);
 
 		FMarqueeRect PrevSelectionRect;
 		PrevSelectionRect.StartPoint = StartPoint;
diff --git a/Engine/Source/Editor/GraphEditor/Public/MarqueeOperation.h b/Engine/Source/Editor/GraphEditor/Public/MarqueeOperation.h
--- a/Engine/Source/Editor/GraphEditor/Public/MarqueeOperation.h
+++ b/Engine/Source/Editor/GraphEditor/Public/MarqueeOperation.h
@@ -34,7 +34,7 @@ struct FMarqueeOperation
 		return Rect.IsValid();
 	}
 
-	void Start(const FVector2D& InStartLocation, FMarqueeOperation::Type InOperationType)
+	void Start(const UE::Slate::FDeprecateVector2DParameter& InStartLocation, FMarqueeOperation::Type InOperationType)
 	{
 		Rect = FMarqueeRect(InStartLocation);
 		Operation = InOperationType;
diff --git a/Engine/Source/Runtime/Slate/Public/Framework/MarqueeRect.h b/Engine/Source/Runtime/Slate/Public/Framework/MarqueeRect.h
--- a/Engine/Source/Runtime/Slate/Public/Framework/MarqueeRect.h
+++ b/Engine/Source/Runtime/Slate/Public/Framework/MarqueeRect.h
@@ -11,12 +11,12 @@
 struct FMarqueeRect
 {
 	/** Where the user began the marquee selection */
-	FVector2D StartPoint;
+	FDeprecateSlateVector2D StartPoint;
 	/** Where the user has dragged to so far */
-	FVector2D EndPoint;
+	FDeprecateSlateVector2D EndPoint;
 
 	/** Make a default marquee selection */
-	FMarqueeRect( FVector2D InStartPoint = FVector2D::ZeroVector )
+	FMarqueeRect( UE::Slate::FDeprecateVector2DParameter InStartPoint = FVector2f::ZeroVector )
 		: StartPoint( InStartPoint )
 		, EndPoint( InStartPoint )
 	{
@@ -27,7 +27,7 @@ struct FMarqueeRect
 	 *
 	 * @param NewEndPoint   Where the user has dragged so far.
 	 */
-	void UpdateEndPoint( FVector2D NewEndPoint )
+	void UpdateEndPoint( const UE::Slate::FDeprecateVector2DParameter& NewEndPoint )
 	{
 		EndPoint = NewEndPoint;
 	}
@@ -39,27 +39,27 @@ struct FMarqueeRect
 	}
 
 	/** @return the upper left point of the selection */
-	FVector2D GetUpperLeft() const
+	UE::Slate::FDeprecateVector2DResult GetUpperLeft() const
 	{
-		return FVector2D( FMath::Min(StartPoint.X, EndPoint.X), FMath::Min( StartPoint.Y, EndPoint.Y ) );
+		return FVector2f( FMath::Min(StartPoint.X, EndPoint.X), FMath::Min( StartPoint.Y, EndPoint.Y ) );
 	}
 
 	/** @return the lower right of the selection */
-	FVector2D GetLowerRight() const
+	UE::Slate::FDeprecateVector2DResult GetLowerRight() const
 	{
-		return  FVector2D( FMath::Max(StartPoint.X, EndPoint.X), FMath::Max( StartPoint.Y, EndPoint.Y ) );
+		return  FVector2f( FMath::Max(StartPoint.X, EndPoint.X), FMath::Max( StartPoint.Y, EndPoint.Y ) );
 	}
 
 	/** The size of the selection */
-	FVector2D GetSize() const
+	UE::Slate::FDeprecateVector2DResult GetSize() const
 	{
-		FVector2D SignedSize = EndPoint - StartPoint;
-		return FVector2D( FMath::Abs(SignedSize.X), FMath::Abs(SignedSize.Y) );
+		FVector2f SignedSize = EndPoint - StartPoint;
+		return FVector2f( FMath::Abs(SignedSize.X), FMath::Abs(SignedSize.Y) );
 	}
 
 	/** @return This marquee rectangle as a well-formed SlateRect */
 	FSlateRect ToSlateRect() const
 	{
-		return FSlateRect( FVector2D(FMath::Min(StartPoint.X, EndPoint.X), FMath::Min(StartPoint.Y, EndPoint.Y)), FVector2D(FMath::Max(StartPoint.X, EndPoint.X), FMath::Max( StartPoint.Y, EndPoint.Y )) );
+		return FSlateRect( FVector2f(FMath::Min(StartPoint.X, EndPoint.X), FMath::Min(StartPoint.Y, EndPoint.Y)), FVector2f(FMath::Max(StartPoint.X, EndPoint.X), FMath::Max( StartPoint.Y, EndPoint.Y )) );
 	}
 };
