UMG: Add option to List/Tile Views to decouple navigation from item selection

#rnx
PR #10998

diff --git a/Engine/Plugins/Runtime/CommonUI/Source/CommonUI/Public/CommonListView.h b/Engine/Plugins/Runtime/CommonUI/Source/CommonUI/Public/CommonListView.h
--- a/Engine/Plugins/Runtime/CommonUI/Source/CommonUI/Public/CommonListView.h
+++ b/Engine/Plugins/Runtime/CommonUI/Source/CommonUI/Public/CommonListView.h
@@ -41,7 +41,16 @@ public:
 					TOptional<ItemType> FirstValidItem = this->Private_FindNextSelectableOrNavigable(SelectedItem);
 					if (FirstValidItem.IsSet())
 					{
-						this->SetSelection(FirstValidItem.GetValue(), ESelectInfo::OnNavigation);
+						// Only select the item if that's desired, otherwise only update SelectorItem
+						if (this->bSelectItemOnNavigation)
+						{
+							this->SetSelection(FirstValidItem.GetValue(), ESelectInfo::OnNavigation);
+						}
+						else
+						{
+							this->SelectorItem = FirstValidItem.GetValue();
+						}
+
 						this->RequestNavigateToItem(FirstValidItem.GetValue(), InFocusEvent.GetUser());
 					}
 				}
diff --git a/Engine/Source/Runtime/Slate/Private/Widgets/Views/STableViewBase.cpp b/Engine/Source/Runtime/Slate/Private/Widgets/Views/STableViewBase.cpp
--- a/Engine/Source/Runtime/Slate/Private/Widgets/Views/STableViewBase.cpp
+++ b/Engine/Source/Runtime/Slate/Private/Widgets/Views/STableViewBase.cpp
@@ -1011,6 +1011,11 @@ void STableViewBase::SetIsTouchScrollingEnabled(const bool bInEnableTouchScrolli
 	ensureMsgf(!bStartedTouchInteraction, TEXT("TouchScrollingEnabled flag should not be changed while scrolling."));
 }
 
+void STableViewBase::SetSelectItemOnNavigation(const bool bInSelectItemOnNavigation)
+{
+	bSelectItemOnNavigation = bInSelectItemOnNavigation;
+}
+
 void STableViewBase::SetWheelScrollMultiplier(float NewWheelScrollMultiplier)
 {
 	WheelScrollMultiplier = NewWheelScrollMultiplier;
diff --git a/Engine/Source/Runtime/Slate/Public/Widgets/Views/SListView.h b/Engine/Source/Runtime/Slate/Public/Widgets/Views/SListView.h
--- a/Engine/Source/Runtime/Slate/Public/Widgets/Views/SListView.h
+++ b/Engine/Source/Runtime/Slate/Public/Widgets/Views/SListView.h
@@ -2547,31 +2547,34 @@ protected:
 			// Always request scroll into view, otherwise partially visible items will be selected - also do this before signaling selection for similar stomp-allowing reasons
 			Private_RequestNavigateToItem(ItemToSelect.GetValue(), InInputEvent.GetUserIndex());
 
-			if (CurrentSelectionMode == ESelectionMode::Multi && (InInputEvent.IsShiftDown() || InInputEvent.IsControlDown()))
+			if (bSelectItemOnNavigation)
 			{
-				// Range select.
-				if (InInputEvent.IsShiftDown())
+				if (CurrentSelectionMode == ESelectionMode::Multi && (InInputEvent.IsShiftDown() || InInputEvent.IsControlDown()))
 				{
-					// Holding control makes the range select bidirectional, where as it is normally unidirectional.
-					if (!(InInputEvent.IsControlDown()))
+					// Range select.
+					if (InInputEvent.IsShiftDown())
 					{
-						this->Private_ClearSelection();
-					}
+						// Holding control makes the range select bidirectional, where as it is normally unidirectional.
+						if (!(InInputEvent.IsControlDown()))
+						{
+							this->Private_ClearSelection();
+						}
 
-					this->Private_SelectRangeFromCurrentTo(ItemToSelect.GetValue());
+						this->Private_SelectRangeFromCurrentTo(ItemToSelect.GetValue());
+					}
+					// For ctrl select, simply add the item to the list of currently selected items
+					else
+					{
+						this->Private_SetItemSelection(ItemToSelect.GetValue(), true, true);
+					}
+				
+					this->Private_SignalSelectionChanged(ESelectInfo::OnNavigation);
 				}
-				// For ctrl select, simply add the item to the list of currently selected items
 				else
 				{
-					this->Private_SetItemSelection(ItemToSelect.GetValue(), true, true);
+					// Single select.
+					this->Private_SetSelection(ItemToSelect.GetValue(), ESelectInfo::OnNavigation);
 				}
-				
-				this->Private_SignalSelectionChanged(ESelectInfo::OnNavigation);
-			}
-			else
-			{
-				// Single select.
-				this->Private_SetSelection(ItemToSelect.GetValue(), ESelectInfo::OnNavigation);
 			}
 		}
 	}
diff --git a/Engine/Source/Runtime/Slate/Public/Widgets/Views/STableViewBase.h b/Engine/Source/Runtime/Slate/Public/Widgets/Views/STableViewBase.h
--- a/Engine/Source/Runtime/Slate/Public/Widgets/Views/STableViewBase.h
+++ b/Engine/Source/Runtime/Slate/Public/Widgets/Views/STableViewBase.h
@@ -209,6 +209,9 @@ public:
 	/** Enables/disables being able to scroll using touch input. */
 	SLATE_API void SetIsTouchScrollingEnabled(const bool bInEnableTouchScrolling);
 
+	/** Sets whether items will be selected and focused or only focused when navigating to them. */
+	SLATE_API void SetSelectItemOnNavigation(const bool bInSelectItemOnNavigation);
+
 	/** Sets the multiplier applied when wheel scrolling. Higher numbers will cover more distance per click of the wheel. */
 	SLATE_API void SetWheelScrollMultiplier(float NewWheelScrollMultiplier);
 
@@ -473,6 +476,9 @@ protected:
 	/** True to allow scrolling by using touch input. */
 	bool bEnableTouchScrolling = true;
 
+	/** If true, items will be "selected" (in addition to focused) when navigating to them. If false, they will only be focused. */
+	bool bSelectItemOnNavigation = true;
+
 	/** The currently displayed scroll offset from the beginning of the list in items. */
 	double CurrentScrollOffset = 0.;
 
diff --git a/Engine/Source/Runtime/UMG/Private/Components/ListViewBase.cpp b/Engine/Source/Runtime/UMG/Private/Components/ListViewBase.cpp
--- a/Engine/Source/Runtime/UMG/Private/Components/ListViewBase.cpp
+++ b/Engine/Source/Runtime/UMG/Private/Components/ListViewBase.cpp
@@ -197,6 +197,7 @@ void UListViewBase::SynchronizeProperties()
 		MyTableViewBase->SetAllowOverscroll(AllowOverscroll ? EAllowOverscroll::Yes : EAllowOverscroll::No);
 		MyTableViewBase->SetFixedLineScrollOffset(bEnableFixedLineOffset ? TOptional<double>(FixedLineScrollOffset) : TOptional<double>());
 		MyTableViewBase->SetWheelScrollMultiplier(GetGlobalScrollAmount() * WheelScrollMultiplier);
+		MyTableViewBase->SetSelectItemOnNavigation(bSelectItemOnNavigation);
 	}
 
 #if WITH_EDITORONLY_DATA
diff --git a/Engine/Source/Runtime/UMG/Public/Components/ListViewBase.h b/Engine/Source/Runtime/UMG/Public/Components/ListViewBase.h
--- a/Engine/Source/Runtime/UMG/Public/Components/ListViewBase.h
+++ b/Engine/Source/Runtime/UMG/Public/Components/ListViewBase.h
@@ -745,6 +745,10 @@ protected:
 	UPROPERTY(EditAnywhere, BlueprintReadOnly, Category = Input)
 	bool bAllowDragging = true;
 
+	/** If true, items will be "selected" (in addition to focused) when navigating to them. If false, they will only be focused. */
+	UPROPERTY(EditAnywhere, BlueprintReadOnly, Category = ListView)
+	bool bSelectItemOnNavigation = true;
+
 	/** Called when a row widget is released by the list (i.e. when it no longer represents a list item) */
 	UPROPERTY(BlueprintAssignable, Category = Events, meta = (DisplayName = "On Entry Released"))
 	FOnListEntryReleasedDynamic BP_OnEntryReleased;
