Fixed ApplyLineHeightToBottomLine(false) being incorrectly applied to all lines during layout, rather than only the last line in the current view

This was broken since the original commit (30412006), as FTextLayout::FlowLineLayout is called in a loop and uses the current TextLayoutSize.Height to calculate the offset for the current line. The old code was removing OverHeight from TextLayoutSize.Height for every line in the view (negating the LineHeightPercentage), and the new code only removes it when processing the last line in the view.

diff --git a/Engine/Source/Runtime/Slate/Private/Framework/Text/TextLayout.cpp b/Engine/Source/Runtime/Slate/Private/Framework/Text/TextLayout.cpp
--- a/Engine/Source/Runtime/Slate/Private/Framework/Text/TextLayout.cpp
+++ b/Engine/Source/Runtime/Slate/Private/Framework/Text/TextLayout.cpp
@@ -791,7 +791,7 @@ void FTextLayout::FlowLineLayout(const int32 LineModelIndex, const float Wrappin
 			}
 		}
 	}
-	if (!ApplyLineHeightToBottomLine)
+	if (!ApplyLineHeightToBottomLine && LineModelIndex == LineModels.Num() - 1)
 	{
 		TextLayoutSize.Height -= OverHeight;
 	}
@@ -2176,6 +2176,12 @@ bool FTextLayout::RemoveLine(int32 LineIndex)
 
 void FTextLayout::AddLine( const FNewLineData& NewLine )
 {
+	// Adjust the layout height to account for the last line OverHeight now that we're adding more lines to the layout
+	if (!LineModels.IsEmpty() && !ApplyLineHeightToBottomLine)
+	{
+		TextLayoutSize.Height += OverHeight;
+	}
+
 	{
 		FLineModel LineModel( NewLine.Text );
 		TransformLineText(LineModel);
@@ -2234,6 +2240,12 @@ void FTextLayout::AddLine( const FNewLineData& NewLine )
 
 void FTextLayout::AddLines( const TArray<FNewLineData>& NewLines )
 {
+	// Adjust the layout height to account for the last line OverHeight now that we're adding more lines to the layout
+	if (!LineModels.IsEmpty() && !ApplyLineHeightToBottomLine)
+	{
+		TextLayoutSize.Height += OverHeight;
+	}
+
 	LineModels.Reserve(NewLines.Num());
 	for (const auto& NewLine : NewLines)
 	{
