Improvements to enhanced tooltip positioning and animation, particularly at the bottom right of the screen.
When the tooltip is interactive, and it is pushed back from the edges of the screen, it is now permitted to overlap the cursor location instead of flipping to the other side of the cursor. This avoids jerky reflowing as it transitions between states.

diff --git a/Engine/Source/Editor/UnrealEd/Private/AssetThumbnail.cpp b/Engine/Source/Editor/UnrealEd/Private/AssetThumbnail.cpp
--- a/Engine/Source/Editor/UnrealEd/Private/AssetThumbnail.cpp
+++ b/Engine/Source/Editor/UnrealEd/Private/AssetThumbnail.cpp
@@ -679,7 +679,7 @@ public:
 		}
 	}
 
-	TSharedRef<SWidget> GetDefaultTooltip() const
+	TSharedRef<SDocumentationToolTip> GetDefaultToolTip() const
 	{
 		const FAssetData& AssetData = AssetThumbnail->GetAssetData();
 
@@ -1528,7 +1528,9 @@ void SAssetThumbnailToolTip::OnOpening()
 		TSharedPtr<SAssetThumbnail> AssetViewItemPin = AssetThumbnail.Pin();
 		if (AssetViewItemPin.IsValid())
 		{
-			SetContentWidget(AssetViewItemPin->GetDefaultTooltip());
+			TSharedRef<SDocumentationToolTip> AssetToolTipRef = AssetViewItemPin->GetDefaultToolTip();
+			SetContentWidget(AssetToolTipRef);
+			AssetToolTip = AssetToolTipRef.ToSharedPtr();
 		}
 	}
 }
@@ -1538,6 +1540,20 @@ void SAssetThumbnailToolTip::OnClosed()
 	ResetContentWidget();
 }
 
+bool SAssetThumbnailToolTip::IsInteractive() const
+{
+	if (UE::Editor::ContentBrowser::IsNewStyleEnabled())
+	{
+		TSharedPtr<SAssetThumbnail> AssetViewItemPin = AssetThumbnail.Pin();
+		if (!IsEmpty() && AssetViewItemPin.IsValid() && AssetToolTip.IsValid())
+		{
+			return AssetToolTip->IsInteractive();
+		}
+		return false;
+	}
+	return false;
+}
+
 FAssetThumbnail::FAssetThumbnail( UObject* InAsset, uint32 InWidth, uint32 InHeight, const TSharedPtr<class FAssetThumbnailPool>& InThumbnailPool )
 	: ThumbnailPool(InThumbnailPool)
 	, AssetData(InAsset ? FAssetData(InAsset) : FAssetData())
diff --git a/Engine/Source/Editor/UnrealEd/Private/AssetThumbnailToolTip.h b/Engine/Source/Editor/UnrealEd/Private/AssetThumbnailToolTip.h
--- a/Engine/Source/Editor/UnrealEd/Private/AssetThumbnailToolTip.h
+++ b/Engine/Source/Editor/UnrealEd/Private/AssetThumbnailToolTip.h
@@ -1,9 +1,10 @@
-﻿// Copyright Epic Games, Inc. All Rights Reserved.
+// Copyright Epic Games, Inc. All Rights Reserved.
 
 #pragma once
 #include "Widgets/SToolTip.h"
 
 class SAssetThumbnail;
+class SDocumentationToolTip;
 
 // AssetThumbnail ToolTip, Implementation is inside AssetThumbnail.cpp
 class SAssetThumbnailToolTip : public SToolTip
@@ -26,6 +27,8 @@ public:
 
 	virtual void OnClosed() override;
 
+	virtual bool IsInteractive() const override;
 private:
 	TWeakPtr<SAssetThumbnail> AssetThumbnail;
+	TSharedPtr<SDocumentationToolTip> AssetToolTip;
 };
diff --git a/Engine/Source/Runtime/Slate/Private/Framework/Application/SlateApplication.cpp b/Engine/Source/Runtime/Slate/Private/Framework/Application/SlateApplication.cpp
--- a/Engine/Source/Runtime/Slate/Private/Framework/Application/SlateApplication.cpp
+++ b/Engine/Source/Runtime/Slate/Private/Framework/Application/SlateApplication.cpp
@@ -3954,13 +3954,13 @@ bool FSlateApplication::GetAllowTooltips() const
 	return bEnableTooltips;
 }
 
-UE::Slate::FDeprecateVector2DResult FSlateApplication::CalculateTooltipWindowPosition( const FSlateRect& InAnchorRect, const UE::Slate::FDeprecateVector2DParameter& InSize, bool bAutoAdjustForDPIScale) const
+UE::Slate::FDeprecateVector2DResult FSlateApplication::CalculateTooltipWindowPosition( const FSlateRect& InAnchorRect, const UE::Slate::FDeprecateVector2DParameter& InSize, bool bAutoAdjustForDPIScale, EPopupCursorOverlapMode CursorOverlapMode) const
 {
 	// first use the CalculatePopupWindowPosition and if cursor is not inside it, proceed with it to avoid behavior change.
-	FVector2f PopupPosition = CalculatePopupWindowPosition(InAnchorRect, InSize, bAutoAdjustForDPIScale);
+	FVector2f PopupPosition = CalculatePopupWindowPosition(InAnchorRect, InSize, bAutoAdjustForDPIScale, UE::Math::TVector2<float>::ZeroVector, Orient_Vertical, EPopupLayoutMode::ToolTip);
 	FVector2f Cursor = GetCursorPos();
-	if (PopupPosition.X > Cursor.X || PopupPosition.X + InSize.X < Cursor.X ||
-		PopupPosition.Y > Cursor.Y || PopupPosition.Y + InSize.Y < Cursor.Y)
+	if (CursorOverlapMode == EPopupCursorOverlapMode::AllowOverlap || (PopupPosition.X > Cursor.X || PopupPosition.X + InSize.X < Cursor.X ||
+		PopupPosition.Y > Cursor.Y || PopupPosition.Y + InSize.Y < Cursor.Y))
 	{
 		return PopupPosition;
 	}
@@ -4004,7 +4004,7 @@ UE::Slate::FDeprecateVector2DResult FSlateApplication::CalculateTooltipWindowPos
 	return ToolTipLocation;
 }
 
-UE::Slate::FDeprecateVector2DResult FSlateApplication::CalculatePopupWindowPosition(const FSlateRect& InAnchor, const UE::Slate::FDeprecateVector2DParameter& InSize, bool bAutoAdjustForDPIScale, const UE::Slate::FDeprecateVector2DParameter& InProposedPlacement, const EOrientation Orientation) const
+UE::Slate::FDeprecateVector2DResult FSlateApplication::CalculatePopupWindowPosition(const FSlateRect& InAnchor, const UE::Slate::FDeprecateVector2DParameter& InSize, bool bAutoAdjustForDPIScale, const UE::Slate::FDeprecateVector2DParameter& InProposedPlacement, const EOrientation Orientation, const EPopupLayoutMode LayoutMode) const
 {
 	FVector2D CalculatedPopUpWindowPosition(0.f, 0.f);
 
@@ -4060,7 +4060,7 @@ UE::Slate::FDeprecateVector2DResult FSlateApplication::CalculatePopupWindowPosit
 				Orientation == Orient_Horizontal ? AnchorRect.Top : AnchorRect.Bottom);
 		}
 
-		return ComputePopupFitInRect(InAnchor, FSlateRect(ProposedPlacement, ProposedPlacement + AdjustedSize), Orientation, WorkAreaRect) / DPIScale;
+		return ComputePopupFitInRect(InAnchor, FSlateRect(ProposedPlacement, ProposedPlacement + AdjustedSize), Orientation, WorkAreaRect, /*bAllowFlip*/LayoutMode == EPopupLayoutMode::Menu) / DPIScale;
 	}
 }
 
diff --git a/Engine/Source/Runtime/Slate/Private/Framework/Application/SlateUser.cpp b/Engine/Source/Runtime/Slate/Private/Framework/Application/SlateUser.cpp
--- a/Engine/Source/Runtime/Slate/Private/Framework/Application/SlateUser.cpp
+++ b/Engine/Source/Runtime/Slate/Private/Framework/Application/SlateUser.cpp
@@ -561,7 +561,7 @@ void FSlateUser::ShowTooltip(const TSharedRef<IToolTip>& InTooltip, const UE::Sl
 
 	// Place the window as close to the given location as possible (MoveWindowTo will adjust the window's position to stay onscreen, if needed)
 	const FSlateRect Anchor(InLocation.X, InLocation.Y, InLocation.X, InLocation.Y);
-	ActiveTooltipInfo.DesiredLocation = SlateApp.CalculateTooltipWindowPosition(Anchor, TooltipWindow->GetDesiredSizeDesktopPixels(), /*bAutoAdjustForDPIScale =*/false);
+	ActiveTooltipInfo.DesiredLocation = SlateApp.CalculateTooltipWindowPosition(Anchor, TooltipWindow->GetDesiredSizeDesktopPixels(), /*bAutoAdjustForDPIScale =*/false, (InTooltip->IsInteractive()) ? EPopupCursorOverlapMode::AllowOverlap : EPopupCursorOverlapMode::PreventOverlap);
 	TooltipWindow->MoveWindowTo(ActiveTooltipInfo.DesiredLocation);
 
 	TooltipWindow->SetOpacity(0.0f);
@@ -1311,6 +1311,8 @@ void FSlateUser::UpdateTooltip(const FMenuStack& MenuStack, bool bCanSpawnNewToo
 				}
 			}
 		}
+
+		ActiveTooltipInfo.HasBeenPositioned = false;
 	}
 	
 	// If a widget under the cursor has a tool-tip forcefield active, then go through any menus
@@ -1330,27 +1332,45 @@ void FSlateUser::UpdateTooltip(const FMenuStack& MenuStack, bool bCanSpawnNewToo
 		}
 	}
 
+	// Start by targeting the position calculated last frame.
 	FVector2f DesiredLocation = ActiveTooltipInfo.DesiredLocation;
-	if ((ActiveTooltip && !ActiveTooltip->IsInteractive()) || (NewTooltip && NewTooltip != ActiveTooltip))
-	{
-		// New tooltips and non-interactive tooltips appear offset from the cursor position, and they follow the cursor as it moves.	
-		DesiredLocation = GetPreviousCursorPosition() + SlateDefs::TooltipOffsetFromMouse;
 
-		// Allow interactive tooltips to adjust the window location
-		if (NewTooltip && NewTooltip->IsInteractive() && !NewTooltipVisualizer.IsValid())
+	// When a tooltip changes from interactive back to regular, there are a few frames during which IsInteractive()
+	// returns false, but the tooltip window is still rendering the interactive content. Computing a new desired location
+	// during this state can lead to the tooltip jumping around visibly. This variable indicates when we're in this situation,
+	// so that we don't look for a new desired location yet.
+	bool bStillAtInteractiveSize = !bTooltipChanged && ActiveTooltip && ActiveTooltipInfo.WasInteractive && TooltipWindowPtr.IsValid() && ActiveTooltipInfo.DesiredSize == TooltipWindowPtr.Pin()->GetDesiredSizeDesktopPixels();
+
+	// Calculate a new desired location for the tooltip if:
+	// - it's a non-interactive tooltip, do this every frame so that it moves with the cursor,
+	//   unless the tooltip window is still rendering interactive content.
+	// - if it's interactive, only do this if the tooltip hasn't been positioned yet.
+	if (ActiveTooltip)
+	{
+		if ((!ActiveTooltip->IsInteractive() && !bStillAtInteractiveSize) || !ActiveTooltipInfo.HasBeenPositioned)
 		{
-			FVector2D DesiredLocation2d(DesiredLocation);
-			NewTooltip->OnSetInteractiveWindowLocation(DesiredLocation2d);
-			DesiredLocation = UE::Slate::CastToVector2f(DesiredLocation2d);
+			// New tooltips and non-interactive tooltips appear offset from the cursor position, and they follow the cursor as it moves.	
+			DesiredLocation = GetPreviousCursorPosition() + SlateDefs::TooltipOffsetFromMouse;
+			ActiveTooltipInfo.WasInteractive = false;
+			ActiveTooltipInfo.HasBeenPositioned = true;
+
+			// Allow interactive tooltips to adjust the window location
+			if (ActiveTooltip->IsInteractive() && !NewTooltipVisualizer.IsValid())
+			{
+				FVector2D DesiredLocation2d(DesiredLocation);
+				ActiveTooltip->OnSetInteractiveWindowLocation(DesiredLocation2d);
+				DesiredLocation = UE::Slate::CastToVector2f(DesiredLocation2d);
+			}
 		}
 	}
 
+	// Move the desired position back from the edges of the screen if needed.
 	FVector2f TooltipSize(ForceInitToZero);
 	if (TooltipWindowPtr.IsValid())
 	{
 		TooltipSize = TooltipWindowPtr.Pin()->GetDesiredSizeDesktopPixels();
 		FSlateRect Anchor(DesiredLocation.X, DesiredLocation.Y, DesiredLocation.X, DesiredLocation.Y);
-		DesiredLocation = SlateApp.CalculatePopupWindowPosition(Anchor, TooltipSize, /*bAutoAdjustForDPIScale =*/false);
+		DesiredLocation = SlateApp.CalculatePopupWindowPosition(Anchor, TooltipSize, /*bAutoAdjustForDPIScale =*/false, FVector2f::ZeroVector, EOrientation::Orient_Vertical, EPopupLayoutMode::ToolTip);
 	}
 
 	// Repel tooltip from a force field, if necessary
@@ -1371,7 +1391,7 @@ void FSlateUser::UpdateTooltip(const FMenuStack& MenuStack, bool bCanSpawnNewToo
 				{
 					FVector2f TentativeDesiredLocation = DesiredLocation + FVector2f(TooltipShift.X, 0);
 					FSlateRect Anchor(TentativeDesiredLocation.X, TentativeDesiredLocation.Y, TentativeDesiredLocation.X, TentativeDesiredLocation.Y);
-					FVector2f NewDesiredLocation = SlateApp.CalculatePopupWindowPosition(Anchor, TooltipSize, /*bAutoAdjustForDPIScale =*/false);
+					FVector2f NewDesiredLocation = SlateApp.CalculatePopupWindowPosition(Anchor, TooltipSize, /*bAutoAdjustForDPIScale =*/false, FVector2f::ZeroVector, EOrientation::Orient_Vertical, EPopupLayoutMode::ToolTip);
 					// If after adjustment the tooltip still overlaps with the force field, try the other direction
 					if (FSlateRect::DoRectanglesIntersect(ForceFieldRect.GetValue(), FSlateRect(NewDesiredLocation, NewDesiredLocation + TooltipSize)))
 					{
@@ -1383,7 +1403,7 @@ void FSlateUser::UpdateTooltip(const FMenuStack& MenuStack, bool bCanSpawnNewToo
 				{
 					FVector2f TentativeDesiredLocation = DesiredLocation + FVector2f(0, TooltipShift.Y);
 					FSlateRect Anchor(TentativeDesiredLocation.X, TentativeDesiredLocation.Y, TentativeDesiredLocation.X, TentativeDesiredLocation.Y);
-					FVector2f NewDesiredLocation = SlateApp.CalculatePopupWindowPosition(Anchor, TooltipSize, /*bAutoAdjustForDPIScale =*/false);
+					FVector2f NewDesiredLocation = SlateApp.CalculatePopupWindowPosition(Anchor, TooltipSize, /*bAutoAdjustForDPIScale =*/false, FVector2f::ZeroVector, EOrientation::Orient_Vertical, EPopupLayoutMode::ToolTip);
 					// If after adjustment the tooltip still overlaps with the force field, try the other direction
 					if (FSlateRect::DoRectanglesIntersect(ForceFieldRect.GetValue(), FSlateRect(NewDesiredLocation, NewDesiredLocation + TooltipSize)))
 					{
@@ -1432,8 +1452,9 @@ void FSlateUser::UpdateTooltip(const FMenuStack& MenuStack, bool bCanSpawnNewToo
 				{
 					ShowTooltip(NewTooltip.ToSharedRef(), DesiredLocation);
 					ActiveTooltipInfo.SourceWidget = WidgetProvidingNewTooltip;
-				}	
+				}
 			}
+			ActiveTooltip = NewTooltip;
 		}
 	}
 
@@ -1457,7 +1478,7 @@ void FSlateUser::UpdateTooltip(const FMenuStack& MenuStack, bool bCanSpawnNewToo
 		// Apply steep inbound curve to the movement, so it looks like it quickly decelerating
 		const float SlideProgress = bAllowAnimations ? FMath::Pow(1.0f - TooltipOpacity, 3.0f) : 0.0f;
 
-		FVector2f WindowLocation = DesiredLocation + SlideProgress * SlideDistance;
+		FVector2f WindowLocation = DesiredLocation + (SlideProgress * SlideDistance);
 		if (WindowLocation != TooltipWindow->GetPositionInScreen()
 			|| TooltipWindow->GetDesiredSize() != ActiveTooltipInfo.DesiredSize)
 		{
@@ -1466,7 +1487,7 @@ void FSlateUser::UpdateTooltip(const FMenuStack& MenuStack, bool bCanSpawnNewToo
 
 			// Avoid the edges of the desktop
 			FSlateRect Anchor(WindowLocation.X, WindowLocation.Y, WindowLocation.X, WindowLocation.Y);
-			WindowLocation = SlateApp.CalculateTooltipWindowPosition(Anchor, TooltipWindow->GetDesiredSizeDesktopPixels(), bAutoAdjustForDPIScale);
+			WindowLocation = SlateApp.CalculateTooltipWindowPosition(Anchor, TooltipWindow->GetDesiredSizeDesktopPixels(), bAutoAdjustForDPIScale, (ActiveTooltip && ActiveTooltip->IsInteractive()) ? EPopupCursorOverlapMode::AllowOverlap : EPopupCursorOverlapMode::PreventOverlap);
 
 			// Cache the size to compare against in future frames
 			ActiveTooltipInfo.DesiredSize = TooltipWindow->GetDesiredSize();
@@ -1478,6 +1499,12 @@ void FSlateUser::UpdateTooltip(const FMenuStack& MenuStack, bool bCanSpawnNewToo
 			TooltipWindow->SetCachedScreenPosition(WindowLocation);
 			TooltipWindow->MoveWindowTo(WindowLocation);
 		}
+
+		// Cache whether the tooltip is in interactive mode.
+		if (ActiveTooltip.IsValid() && ActiveTooltip->IsInteractive())
+		{
+			ActiveTooltipInfo.WasInteractive = true;
+		}
 	}
 }
 
diff --git a/Engine/Source/Runtime/Slate/Public/Framework/Application/SlateApplication.h b/Engine/Source/Runtime/Slate/Public/Framework/Application/SlateApplication.h
--- a/Engine/Source/Runtime/Slate/Public/Framework/Application/SlateApplication.h
+++ b/Engine/Source/Runtime/Slate/Public/Framework/Application/SlateApplication.h
@@ -207,6 +207,17 @@ enum class EInputPreProcessorType : uint8
 	Count
 };
 
+/** Determines how popup windows are laid out by CalculatePopupWindowPosition
+  * when there is no space on the screen for the popup window to fit in its preferred orientation. */
+UENUM()
+enum class EPopupLayoutMode : uint8
+{
+	/** The popup window is allowed to flip to the other side of the anchor that opens it. */
+	Menu = 0,
+	/** The popup window does not flip. Tooltips have their own logic to handle placement for this case. */
+	ToolTip
+};
+
 USTRUCT()
 struct FInputPreprocessorRegistrationKey
 {
@@ -917,24 +928,28 @@ public:
 	 * Calculates the popup window position from the passed in window position and size. 
 	 * Adjusts position for popup windows which are outside of the work area of the monitor where they reside
 	 *
-	 * @param InAnchor				The current(suggested) window position and size of an area which may not be covered by the popup.
-	 * @param InSize				The size of the window
-	 * @param InProposedPlacement	The location on screen where the popup should go if allowed. If zero this will be determined from Orientation and Anchor
-	 * @param Orientation			The direction of the popup.
-	 *								If vertical it will attempt to open below the anchor but will open above if there is no room.
-	 *								If horizontal it will attempt to open below the anchor but will open above if there is no room.
+	 * @param InAnchor				 The current(suggested) window position and size of an area which may not be covered by the popup.
+	 * @param InSize				 The size of the window
+	 * @param bAutoAdjustForDPIScale Whether the method should attempt to correct for the DPI scale of the monitor.
+	 * @param InProposedPlacement	 The location on screen where the popup should go if allowed. If zero this will be determined from Orientation and Anchor
+	 * @param Orientation			 The direction of the popup.
+	 *								 If vertical it will attempt to open below the anchor but will open above if there is no room.
+	 *								 If horizontal it will attempt to open below the anchor but will open above if there is no room.
+	 * @param LayoutMode			 Determines whether the popup is allowed to flip to the other side of the anchor when there is no room on the screen.
 	 * @return The adjusted position
 	 */
-	SLATE_API virtual UE::Slate::FDeprecateVector2DResult CalculatePopupWindowPosition( const FSlateRect& InAnchor, const UE::Slate::FDeprecateVector2DParameter& InSize, bool bAutoAdjustForDPIScale = true, const UE::Slate::FDeprecateVector2DParameter& InProposedPlacement = FVector2f::ZeroVector, const EOrientation Orientation = Orient_Vertical) const;
+	SLATE_API virtual UE::Slate::FDeprecateVector2DResult CalculatePopupWindowPosition( const FSlateRect& InAnchor, const UE::Slate::FDeprecateVector2DParameter& InSize, bool bAutoAdjustForDPIScale = true, const UE::Slate::FDeprecateVector2DParameter& InProposedPlacement = FVector2f::ZeroVector, const EOrientation Orientation = Orient_Vertical, const EPopupLayoutMode = EPopupLayoutMode::Menu) const;
 
 	/**
 	 * Calculates the tooltip window position.
-	 * 
+	 *
 	 * @param InAnchorRect The current(suggested) window position and size of an area which may not be covered by the popup.
 	 * @param InSize The size of the tooltip window.
+	 * @param bAutoAdjustForDPIScale Whether the method should attempt to correct for the DPI scale of the monitor.
+	 * @param CursorOverlapMode Whether the tooltip window should be permitted to open under the position of the cursor.
 	 * @return The suggested position.
 	 */
-	SLATE_API virtual UE::Slate::FDeprecateVector2DResult CalculateTooltipWindowPosition( const FSlateRect& InAnchorRect, const UE::Slate::FDeprecateVector2DParameter& InSize, bool bAutoAdjustForDPIScale) const;
+	SLATE_API virtual UE::Slate::FDeprecateVector2DResult CalculateTooltipWindowPosition( const FSlateRect& InAnchorRect, const UE::Slate::FDeprecateVector2DParameter& InSize, bool bAutoAdjustForDPIScale, EPopupCursorOverlapMode CursorOverlapMode = EPopupCursorOverlapMode::PreventOverlap) const;
 
 	/**
 	 * Is the window in the app's destroy queue? If so it will be destroyed next tick.
diff --git a/Engine/Source/Runtime/Slate/Public/Framework/Application/SlateUser.h b/Engine/Source/Runtime/Slate/Public/Framework/Application/SlateUser.h
--- a/Engine/Source/Runtime/Slate/Public/Framework/Application/SlateUser.h
+++ b/Engine/Source/Runtime/Slate/Public/Framework/Application/SlateUser.h
@@ -328,6 +328,12 @@ private:
 		 * Cached to prevent the direction of offset from changing as the user moves the mouse cursor (otherwise the tooltip teleports around)
 		 */
 		ETooltipOffsetDirection OffsetDirection = ETooltipOffsetDirection::Undetermined;
+
+		/** Caches whether the tooltip's IsInteractive() method returned true last frame. */
+		bool WasInteractive = false;
+
+		/** Caches whether the tooltip has been positioned yet. */
+		bool HasBeenPositioned = false;
 	};
 	FActiveTooltipInfo ActiveTooltipInfo;
 
diff --git a/Engine/Source/Runtime/SlateCore/Private/Layout/LayoutUtils.cpp b/Engine/Source/Runtime/SlateCore/Private/Layout/LayoutUtils.cpp
--- a/Engine/Source/Runtime/SlateCore/Private/Layout/LayoutUtils.cpp
+++ b/Engine/Source/Runtime/SlateCore/Private/Layout/LayoutUtils.cpp
@@ -2,7 +2,7 @@
 
 #include "Layout/LayoutUtils.h"
 
-UE::Slate::FDeprecateVector2DResult ComputePopupFitInRect(const FSlateRect& InAnchor, const FSlateRect& PopupRect, const EOrientation& Orientation, const FSlateRect& RectToFit)
+UE::Slate::FDeprecateVector2DResult ComputePopupFitInRect(const FSlateRect& InAnchor, const FSlateRect& PopupRect, const EOrientation& Orientation, const FSlateRect& RectToFit, bool bAllowFlip)
 {
 	const bool bAdjustmentNeeded = PopupRect.IntersectionWith(RectToFit) != PopupRect;
 	if (bAdjustmentNeeded)
@@ -12,44 +12,51 @@ UE::Slate::FDeprecateVector2DResult ComputePopupFitInRect(const FSlateRect& InAn
 
 		// In the direction we are opening, see if there is enough room. If there is not, flip the opening direction along the same axis.
 		FVector2f NewPosition = FVector2f::ZeroVector;
-		if (Orientation == Orient_Horizontal)
+		if (bAllowFlip)
 		{
-			const bool bFitsRight = InAnchor.Right + PopupSize.X < RectToFit.Right;
-			const bool bFitsLeft = InAnchor.Left - PopupSize.X >= RectToFit.Left;
-
-			if (bFitsRight || !bFitsLeft)
+			if (Orientation == Orient_Horizontal)
 			{
-				// The menu fits to the right of the anchor or it does not fit to the left, display to the right
-				NewPosition = FVector2f(InAnchor.Right, InAnchor.Top);
+				const bool bFitsRight = InAnchor.Right + PopupSize.X < RectToFit.Right;
+				const bool bFitsLeft = InAnchor.Left - PopupSize.X >= RectToFit.Left;
+
+				if (bFitsRight || !bFitsLeft)
+				{
+					// The menu fits to the right of the anchor or it does not fit to the left, display to the right
+					NewPosition = FVector2f(InAnchor.Right, InAnchor.Top);
+				}
+				else
+				{
+					// The menu does not fit to the right of the anchor but it does fit to the left, display to the left
+					NewPosition = FVector2f(InAnchor.Left - PopupSize.X, InAnchor.Top);
+				}
 			}
 			else
 			{
-				// The menu does not fit to the right of the anchor but it does fit to the left, display to the left
-				NewPosition = FVector2f(InAnchor.Left - PopupSize.X, InAnchor.Top);
+				const bool bFitsDown = InAnchor.Bottom + PopupSize.Y < RectToFit.Bottom;
+				const bool bFitsUp = InAnchor.Top - PopupSize.Y >= RectToFit.Top;
+
+				if (bFitsDown || !bFitsUp)
+				{
+					// The menu fits below the anchor or it does not fit above, display below
+					NewPosition = FVector2f(InAnchor.Left, InAnchor.Bottom);
+				}
+				else
+				{
+					// The menu does not fit below the anchor but it does fit above, display above
+					NewPosition = FVector2f(InAnchor.Left, InAnchor.Top - PopupSize.Y);
+				}
+
+				if (!bFitsDown && !bFitsUp)
+				{
+					return ComputePopupFitInRect(InAnchor, PopupRect, EOrientation::Orient_Horizontal, RectToFit);
+				}
 			}
 		}
 		else
 		{
-			const bool bFitsDown = InAnchor.Bottom + PopupSize.Y < RectToFit.Bottom;
-			const bool bFitsUp = InAnchor.Top - PopupSize.Y >= RectToFit.Top;
-
-			if (bFitsDown || !bFitsUp)
-			{
-				// The menu fits below the anchor or it does not fit above, display below
-				NewPosition = FVector2f(InAnchor.Left, InAnchor.Bottom);
-			}
-			else
-			{
-				// The menu does not fit below the anchor but it does fit above, display above
-				NewPosition = FVector2f(InAnchor.Left, InAnchor.Top - PopupSize.Y);
-			}
-
-			if (!bFitsDown && !bFitsUp)
-			{
-				return ComputePopupFitInRect(InAnchor, PopupRect, EOrientation::Orient_Horizontal, RectToFit);
-			}
+			NewPosition = FVector2f(InAnchor.Left, InAnchor.Bottom);
 		}
-
+		
 		// Adjust the position of popup windows so they do not go out of the visible area of the monitor(s)
 		// This can happen along the opposite axis that we are opening with
 		// Assumes this window has a valid size
diff --git a/Engine/Source/Runtime/SlateCore/Public/Application/SlateApplicationBase.h b/Engine/Source/Runtime/SlateCore/Public/Application/SlateApplicationBase.h
--- a/Engine/Source/Runtime/SlateCore/Public/Application/SlateApplicationBase.h
+++ b/Engine/Source/Runtime/SlateCore/Public/Application/SlateApplicationBase.h
@@ -94,6 +94,15 @@ namespace SlateApplicationDefs
 	static const int32 MaxHardwareUsers = 8;
 }
 
+/** Determines how popup window placement should treat the cursor. */
+enum class EPopupCursorOverlapMode : uint8
+{
+	/** The popup window is not allowed to overlap the position of the cursor. */
+	PreventOverlap = 0,
+	/** The popup window may overlap the position of the cursor. */
+	AllowOverlap
+};
+
 /**
  * Base class for Slate applications.
  *
@@ -376,9 +385,11 @@ public:
 	 *
 	 * @param InAnchorRect The current(suggested) window position and size of an area which may not be covered by the popup.
 	 * @param InSize The size of the tooltip window.
+	 * @param bAutoAdjustForDPIScale Whether the method should attempt to correct for the DPI scale of the monitor.
+	 * @param CursorOverlapMode Whether the tooltip window should be permitted to open under the position of the cursor. 
 	 * @return The suggested position.
 	 */
-	virtual UE::Slate::FDeprecateVector2DResult CalculateTooltipWindowPosition(const FSlateRect& InAnchorRect, const UE::Slate::FDeprecateVector2DParameter& InSize, bool bAutoAdjustForDPIScale) const = 0;
+	virtual UE::Slate::FDeprecateVector2DResult CalculateTooltipWindowPosition(const FSlateRect& InAnchorRect, const UE::Slate::FDeprecateVector2DParameter& InSize, bool bAutoAdjustForDPIScale, EPopupCursorOverlapMode CursorOverlapMode = EPopupCursorOverlapMode::PreventOverlap) const = 0;
 
 	/** @return true if 'WindowToTest' is being used to display the current tooltip and the tooltip is interactive. */
 	virtual bool IsWindowHousingInteractiveTooltip(const TSharedRef<const SWindow>& WindowToTest) const = 0;
diff --git a/Engine/Source/Runtime/SlateCore/Public/Layout/LayoutUtils.h b/Engine/Source/Runtime/SlateCore/Public/Layout/LayoutUtils.h
--- a/Engine/Source/Runtime/SlateCore/Public/Layout/LayoutUtils.h
+++ b/Engine/Source/Runtime/SlateCore/Public/Layout/LayoutUtils.h
@@ -598,7 +598,9 @@ static FMargin LayoutPaddingWithFlow(EFlowDirection InLayoutFlow, const FMargin&
 * @param PopupRect         Proposed placement of popup; position may require adjustment.
 * @param Orientation       Are we trying to show the popup above/below or left/right relative to the anchor?
 * @param RectToFit         The space available for showing this popup; we want to fit entirely within it without clipping.
+* @param bAllowFlip        Determines whether the popup can be flipped to the other side of the anchor when there
+*                          is no space on the screen for it to fit at its preferred position.
 *
 * @return A best position within the RectToFit such that none of the popup clips outside of the RectToFit.
 */
-SLATECORE_API UE::Slate::FDeprecateVector2DResult ComputePopupFitInRect(const FSlateRect& InAnchor, const FSlateRect& PopupRect, const EOrientation& Orientation, const FSlateRect& RectToFit);
+SLATECORE_API UE::Slate::FDeprecateVector2DResult ComputePopupFitInRect(const FSlateRect& InAnchor, const FSlateRect& PopupRect, const EOrientation& Orientation, const FSlateRect& RectToFit, bool bAllowFlip = true);
