Fix issue where toolbar icons didn't show the right icon color on hover.
- Add new LevelEditorStyle to handle Level Editor Viewport Toolbar customizations.
- Add optional style customization for split combo buttons. Viewport toolbars remove the inner rounded corners.

diff --git a/Engine/Source/Editor/EditorStyle/Private/StarshipStyle.cpp b/Engine/Source/Editor/EditorStyle/Private/StarshipStyle.cpp
--- a/Engine/Source/Editor/EditorStyle/Private/StarshipStyle.cpp
+++ b/Engine/Source/Editor/EditorStyle/Private/StarshipStyle.cpp
@@ -2672,6 +2672,12 @@ void FStarshipEditorStyle::FStyle::SetupViewportStyles()
 			Set("ViewportToolbar.ToggleButton.Start", ToggleButtonStyleStart);
 			Set("ViewportToolbar.ToggleButton.Middle", ToggleButtonStyleMiddle);
 			Set("ViewportToolbar.ToggleButton.End", ToggleButtonStyleEnd);
+			
+			// Ensure split buttons don't have rounded corners on the inside
+			Set("ViewportToolbar.SplitToggleButton", ToggleButtonStyleStart);
+			FComboButtonStyle SplitComboButton = ViewportToolbarStyle.ComboButtonStyle;
+			SplitComboButton.ButtonStyle = ButtonStyleEnd;
+			Set("ViewportToolbar.SplitComboButton", SplitComboButton);
 		}
 
 		// Special styling for top-level raised buttons.
@@ -2747,18 +2753,14 @@ void FStarshipEditorStyle::FStyle::SetupViewportStyles()
 			Set("ViewportToolbar.Raised.ToggleButton.Start", ToggleButtonStyleRaisedStart);
 			Set("ViewportToolbar.Raised.ToggleButton.Middle", ToggleButtonStyleRaisedMiddle);
 			Set("ViewportToolbar.Raised.ToggleButton.End", ToggleButtonStyleRaisedEnd);
+			
+			// Ensure split buttons don't have rounded corners on the inside
+			Set("ViewportToolbar.Raised.SplitToggleButton", ToggleButtonStyleRaisedStart);
+			FComboButtonStyle SplitComboButton = ViewportToolbarRaisedStyle.ComboButtonStyle;
+			SplitComboButton.ButtonStyle = ButtonStyleRaisedEnd;
+			Set("ViewportToolbar.Raised.SplitComboButton", SplitComboButton);
 		}
 
-		// Create a special viewport toolbar style for the level editor. To accomodate the persistent top-right min/max
-		// button, we put the overflow wrap button at index -3. This adds the wrap button to the left of the min/max
-		// button, keeping it in the far-right position even as toolbar size is constrained.
-		{
-			FToolBarStyle LevelEditorViewportToolbarStyle = FToolBarStyle(ViewportToolbarStyle).SetWrapButtonIndex(-3);
-
-			Set("LevelEditorViewportToolbar", LevelEditorViewportToolbarStyle);
-			Set("LevelEditorViewportToolbar.Raised", ViewportToolbarRaisedStyle);
-		}
-		
 		// Special style for raised transform buttons
 		{
 			FToolBarStyle RaisedTransformTools = FToolBarStyle(GetWidgetStyle<FToolBarStyle>("ViewportToolbar.Raised"));
diff --git a/Engine/Source/Editor/LevelEditor/Private/LevelEditorStyle.h b/Engine/Source/Editor/LevelEditor/Private/LevelEditorStyle.h
new file mode 100644
--- /dev/null
+++ b/Engine/Source/Editor/LevelEditor/Private/LevelEditorStyle.h
@@ -0,0 +1,13 @@
+﻿// Copyright Epic Games, Inc. All Rights Reserved.
+
+#pragma once
+#include "Styling/SlateStyle.h"
+
+class FLevelEditorStyle final : public FSlateStyleSet
+{
+public:
+	FLevelEditorStyle();
+	virtual ~FLevelEditorStyle() override;
+	
+	static FLevelEditorStyle& Get();
+};
diff --git a/Engine/Source/Editor/LevelEditor/Private/SLevelViewport.cpp b/Engine/Source/Editor/LevelEditor/Private/SLevelViewport.cpp
--- a/Engine/Source/Editor/LevelEditor/Private/SLevelViewport.cpp
+++ b/Engine/Source/Editor/LevelEditor/Private/SLevelViewport.cpp
@@ -105,6 +105,7 @@
 #include "EditorViewportTabContent.h"
 #include "EditorViewportCommands.h"
 #include "FunctionalUIScreenshotTest.h"
+#include "LevelEditorStyle.h"
 #include "SActionableMessageViewportWidget.h"
 #include "SkeletalRenderPublic.h"
 #include "ViewportToolbar/LevelEditorSubmenus.h"
@@ -2060,7 +2061,8 @@ TSharedPtr<SWidget> SLevelViewport::BuildViewportToolbar()
 				ViewportToolbarName, NAME_None /* parent */, EMultiBoxType::SlimHorizontalToolBar
 			);
 
-			ViewportToolbarMenu->StyleName = "LevelEditorViewportToolbar";
+			ViewportToolbarMenu->SetStyleSet(&FLevelEditorStyle::Get());
+			ViewportToolbarMenu->StyleName = "ViewportToolbar";
 			ViewportToolbarMenu->bSeparateSections = false;
 
 			// Add the left-aligned part of the viewport toolbar.
diff --git a/Engine/Source/Editor/LevelEditor/Private/ViewportToolbar/LevelEditorStyle.cpp b/Engine/Source/Editor/LevelEditor/Private/ViewportToolbar/LevelEditorStyle.cpp
new file mode 100644
--- /dev/null
+++ b/Engine/Source/Editor/LevelEditor/Private/ViewportToolbar/LevelEditorStyle.cpp
@@ -0,0 +1,37 @@
+﻿// Copyright Epic Games, Inc. All Rights Reserved.
+#include "LevelEditorStyle.h"
+
+#include "Styling/AppStyle.h"
+#include "Styling/SlateStyleRegistry.h"
+#include "Styling/ToolBarStyle.h"
+
+FLevelEditorStyle::FLevelEditorStyle()
+	: FSlateStyleSet("LevelEditorStyle")
+{
+	SetParentStyleName(FAppStyle::GetAppStyleSetName());
+		
+	// Toolbar styles involve many "nested" styles. For example, "ViewportToolbar.Raised" or "ViewportToolbar.Button.Start".
+	// If the level viewport toolbar just used a new style name, all of these sub-styles would need to be redefined as well.
+	// Using a style set with a parent allows for customization of some styles while falling back to the parent for others.
+	
+	// Create a special viewport toolbar style for the level editor. To accomodate the persistent top-right min/max
+	// button, we put the overflow wrap button at index -3. This adds the wrap button to the left of the min/max
+	// button, keeping it in the far-right position even as toolbar size is constrained.
+	FToolBarStyle ToolBarStyle = GetParentStyle()->GetWidgetStyle<FToolBarStyle>("ViewportToolbar");
+	ToolBarStyle.SetWrapButtonIndex(-3);
+	
+	Set("ViewportToolbar", ToolBarStyle);
+	
+	FSlateStyleRegistry::RegisterSlateStyle(*this);
+}
+
+FLevelEditorStyle::~FLevelEditorStyle()
+{
+	FSlateStyleRegistry::UnRegisterSlateStyle(*this);
+}
+
+FLevelEditorStyle& FLevelEditorStyle::Get()
+{
+	static FLevelEditorStyle Instance;
+	return Instance;
+}
\ No newline at end of file
diff --git a/Engine/Source/Runtime/Slate/Private/Framework/MultiBox/SToolBarComboButtonBlock.cpp b/Engine/Source/Runtime/Slate/Private/Framework/MultiBox/SToolBarComboButtonBlock.cpp
--- a/Engine/Source/Runtime/Slate/Private/Framework/MultiBox/SToolBarComboButtonBlock.cpp
+++ b/Engine/Source/Runtime/Slate/Private/Framework/MultiBox/SToolBarComboButtonBlock.cpp
@@ -257,8 +257,7 @@ void SToolBarComboButtonBlock::BuildMultiBlockWidget(const ISlateStyle* StyleSet
 		UserInterfaceType = UICommand->GetUserInterfaceType();
 	}
 	
-	// Cache the checkbox style for icon coloration
-	CheckBoxStyle = GetCheckBoxStyle(StyleSet, StyleName);
+	const FCheckBoxStyle* CheckStyle = GetCheckBoxStyle(StyleSet, StyleName, bIsSplitButton);
 	
 	LeftHandSideWidget = SNullWidget::NullWidget;
 
@@ -271,10 +270,16 @@ void SToolBarComboButtonBlock::BuildMultiBlockWidget(const ISlateStyle* StyleSet
 			ComboStyle = &ToolBarStyle.SettingsComboButton;
 			ButtonStyle = &ComboStyle->ButtonStyle;
 		}
+		else
+		{
+			// Allow for optional style customization of a split combo button 
+			ComboStyle = &StyleSet->GetWidgetStyle<FComboButtonStyle>(StyleName, ".SplitComboButton", ComboStyle);
+			ButtonStyle = &ComboStyle->ButtonStyle;
+		}
 	
 		// clang-format off
 		SAssignNew(LeftHandSideWidget, SCheckBox)
-			.Style(CheckBoxStyle)
+			.Style(CheckStyle)
 			.CheckBoxContentUsesAutoWidth(false)
 			.OnCheckStateChanged(this, &SToolBarComboButtonBlock::OnCheckStateChanged)
 			.OnGetMenuContent(this, &SToolBarComboButtonBlock::OnGetMenuContent)
@@ -290,6 +295,12 @@ void SToolBarComboButtonBlock::BuildMultiBlockWidget(const ISlateStyle* StyleSet
 			];
 		// clang-format on
 	}
+	else if (HasCheckedState())
+	{
+		// Only cache the checkbox style if we need to perform coloration
+		// on icons in non-checkbox widgets
+		CheckBoxStyle = CheckStyle;
+	}
 	
 	// clang-format off
 	ChildSlot
@@ -385,31 +396,37 @@ void SToolBarComboButtonBlock::OnCheckStateChanged(const ECheckBoxState NewCheck
 	OnClicked();
 }
 
-const FCheckBoxStyle* SToolBarComboButtonBlock::GetCheckBoxStyle(const ISlateStyle* StyleSet, const FName& StyleName) const
+const FCheckBoxStyle* SToolBarComboButtonBlock::GetCheckBoxStyle(const ISlateStyle* StyleSet, const FName& StyleName, bool bIsSplitButton) const
 {
 	const FToolBarStyle& ToolBarStyle = StyleSet->GetWidgetStyle<FToolBarStyle>(StyleName);
 	EMultiBlockLocation::Type BlockLocation = GetMultiBlockLocation();
 
 	const FCheckBoxStyle* CheckStyle;
-	
-	if (!MultiBlock->HasIcon())
+
+	if (OptionsBlockWidget.IsValid())
+	{
+		CheckStyle = &ToolBarStyle.SettingsToggleButton;
+	}	
+	else if (!MultiBlock->HasIcon())
 	{
 		CheckStyle = &FCoreStyle::Get().GetWidgetStyle<FCheckBoxStyle>("Checkbox");
 	}
 	else if (BlockLocation == EMultiBlockLocation::None)
 	{
-		CheckStyle = &ToolBarStyle.ToggleButton;
+		if (bIsSplitButton)
+		{
+			CheckStyle = &StyleSet->GetWidgetStyle<FCheckBoxStyle>(StyleName, ".SplitToggleButton", &ToolBarStyle.ToggleButton);
+		}
+		else
+		{
+			CheckStyle = &ToolBarStyle.ToggleButton;	
+		}
 	}
 	else
 	{
 		CheckStyle = &StyleSet->GetWidgetStyle<FCheckBoxStyle>(EMultiBlockLocation::ToName(ISlateStyle::Join(StyleName, ".ToggleButton"), BlockLocation));
 	}
 
-	if (OptionsBlockWidget.IsValid())
-	{
-		CheckStyle = &ToolBarStyle.SettingsToggleButton;
-	}
-	
 	return CheckStyle;
 }
 
@@ -450,6 +467,18 @@ bool SToolBarComboButtonBlock::HasAction() const
 		MultiBlock->GetDirectActions().IsBound();
 }
 
+bool SToolBarComboButtonBlock::HasCheckedState() const
+{
+	if (MultiBlock->GetActionList() && MultiBlock->GetAction())
+	{
+		if (const FUIAction* Action = MultiBlock->GetActionList()->GetActionForCommand(MultiBlock->GetAction()))
+		{
+			return Action->GetActionCheckState.IsBound();
+		}
+	}
+	return MultiBlock->GetDirectActions().GetActionCheckState.IsBound();
+}
+
 EVisibility SToolBarComboButtonBlock::GetVisibility() const
 {
 	// Let the visibility override take prescedence here.
@@ -510,14 +539,32 @@ FSlateColor SToolBarComboButtonBlock::GetIconForegroundColor() const
 	
 	if (CheckBoxStyle)
 	{
-		switch (GetCheckState())
+		const ECheckBoxState CheckState = GetCheckState();
+		const bool bIsHovered = LeftHandSideWidget != SNullWidget::NullWidget ? LeftHandSideWidget->IsHovered() : ComboButtonWidget->IsHovered();
+		
+		if (bIsHovered)
+		{
+			switch (CheckState)
+			{
+			case ECheckBoxState::Unchecked:
+				return CheckBoxStyle->HoveredForeground;
+			case ECheckBoxState::Checked:
+				return CheckBoxStyle->CheckedHoveredForeground;
+			case ECheckBoxState::Undetermined:
+				return CheckBoxStyle->HoveredForeground;
+			}
+		}
+		else
 		{
-		case ECheckBoxState::Unchecked:
-			return CheckBoxStyle->ForegroundColor;
-		case ECheckBoxState::Checked:
-			return CheckBoxStyle->CheckedForeground;
-		case ECheckBoxState::Undetermined:
-			return CheckBoxStyle->UndeterminedForeground;
+			switch (CheckState)
+			{
+			case ECheckBoxState::Unchecked:
+				return CheckBoxStyle->ForegroundColor;
+			case ECheckBoxState::Checked:
+				return CheckBoxStyle->CheckedForeground;
+			case ECheckBoxState::Undetermined:
+				return CheckBoxStyle->UndeterminedForeground;
+			}
 		}
 	}
 	
diff --git a/Engine/Source/Runtime/Slate/Public/Framework/MultiBox/SToolBarComboButtonBlock.h b/Engine/Source/Runtime/Slate/Public/Framework/MultiBox/SToolBarComboButtonBlock.h
--- a/Engine/Source/Runtime/Slate/Public/Framework/MultiBox/SToolBarComboButtonBlock.h
+++ b/Engine/Source/Runtime/Slate/Public/Framework/MultiBox/SToolBarComboButtonBlock.h
@@ -185,11 +185,14 @@ protected:
 
 	/** True if we have an action bound */
 	bool HasAction() const;
+	
+	/** True if we have IsChecked bound */ 
+	bool HasCheckedState() const;
 
 	/**
 	 * @return Gets the effective checkbox style for this widget. 
 	 */
-	const FCheckBoxStyle* GetCheckBoxStyle(const ISlateStyle* StyleSet, const FName& StyleName) const;
+	const FCheckBoxStyle* GetCheckBoxStyle(const ISlateStyle* StyleSet, const FName& StyleName, bool bIsSplitButton) const;
 	
 	/**
 	 * Called by Slate to determine if this button is visible
