[Crash Fix] Checks for a valid SharedPtr before accessing it's functions. There are no repro steps but this should be the valid fix for this crash

diff --git a/Engine/Source/Runtime/Slate/Private/Framework/Docking/SDockingArea.cpp b/Engine/Source/Runtime/Slate/Private/Framework/Docking/SDockingArea.cpp
--- a/Engine/Source/Runtime/Slate/Private/Framework/Docking/SDockingArea.cpp
+++ b/Engine/Source/Runtime/Slate/Private/Framework/Docking/SDockingArea.cpp
@@ -235,8 +235,8 @@ void SDockingArea::OnTabFoundNewHome( const TSharedRef<SDockTab>& RelocatedTab,
 
 	// The last tab has been successfully relocated elsewhere.
 	// Destroy this window.
-	TSharedPtr<SWindow> ParentWindow = ParentWindowPtr.Pin();
-	if ( bManageParentWindow && bCleanUpUponTabRelocation && ParentWindow != NewOwnerWindow )
+	TSharedPtr<SWindow> ParentWindow = GetParentWindow();
+	if (ParentWindow.IsValid() && bManageParentWindow && bCleanUpUponTabRelocation && ParentWindow != NewOwnerWindow)
 	{
 		ParentWindow->SetRequestDestroyWindowOverride( FRequestDestroyWindowOverride() );
 		ParentWindow->RequestDestroyWindow();
