Slate: Add ScrollBox support for the gamepad equivalent to Mousewheel

[RN]

diff --git a/Engine/Source/Runtime/Slate/Private/Widgets/Layout/SScrollBox.cpp b/Engine/Source/Runtime/Slate/Private/Widgets/Layout/SScrollBox.cpp
--- a/Engine/Source/Runtime/Slate/Private/Widgets/Layout/SScrollBox.cpp
+++ b/Engine/Source/Runtime/Slate/Private/Widgets/Layout/SScrollBox.cpp
@@ -121,7 +121,10 @@ void SScrollBox::Construct( const FArguments& InArgs )
 	NavigationScrollPadding = InArgs._NavigationScrollPadding;
 	NavigationDestination = InArgs._NavigationDestination;
 	ScrollWhenFocusChanges = InArgs._ScrollWhenFocusChanges;
+	OnScrollBoxFocusReceived = InArgs._OnFocusReceived;
+	OnScrollBoxFocusLost = InArgs._OnFocusLost;
 	bTouchPanningCapture = false;
+	bIsFocusable = false;
 	bVolatilityAlwaysInvalidatesPrepass = true;
 
 	if (InArgs._ExternalScrollbar.IsValid())
@@ -577,6 +580,17 @@ void SScrollBox::SetConsumeMouseWheel(EConsumeMouseWheel NewConsumeMouseWheel)
 	ConsumeMouseWheel = NewConsumeMouseWheel;
 }
 
+void SScrollBox::SetAnalogMouseWheelKey(FKey InAnalogKey)
+{
+	ensure(!InAnalogKey.IsValid() || InAnalogKey.IsAnalog());
+	AnalogMouseWheelKey = InAnalogKey;
+}
+
+void SScrollBox::SetIsFocusable(bool bInIsFocusable)
+{
+	bIsFocusable = bInIsFocusable;
+}
+
 void SScrollBox::SetOrientation(EOrientation InOrientation)
 {
 	if (Orientation != InOrientation)
@@ -975,6 +989,38 @@ FReply SScrollBox::OnMouseWheel( const FGeometry& MyGeometry, const FPointerEven
 	}
 }
 
+FReply SScrollBox::OnAnalogValueChanged(const FGeometry& MyGeometry, const FAnalogInputEvent& InAnalogInputEvent)
+{
+	const float AnalogInputThreshold = 0.2f;
+	if ((ScrollBar->IsNeeded() && ConsumeMouseWheel != EConsumeMouseWheel::Never) || ConsumeMouseWheel == EConsumeMouseWheel::Always)
+	{
+		const FKey& AnalogKey   = InAnalogInputEvent.GetKey();
+		const float AnalogValue = InAnalogInputEvent.GetAnalogValue();
+
+		// If not the correct input, or not within input thresholds, ignore.
+		if (AnalogKey != AnalogMouseWheelKey || (AnalogValue < AnalogInputThreshold && AnalogValue > -AnalogInputThreshold))
+		{
+			return FReply::Unhandled();
+		}
+		
+		// Make sure scroll velocity is cleared so it doesn't fight with the mouse wheel input
+		InertialScrollManager.ClearScrollVelocity();
+
+		const float ScrollAmount = GetGlobalScrollAmount() * GetGlobalGamepadScrollMultiplier() * WheelScrollMultiplier;
+		const bool bScrollWasHandled = ScrollBy(MyGeometry, -AnalogValue * ScrollAmount, EAllowOverscroll::No, bAnimateWheelScrolling);
+
+		if ( bScrollWasHandled && !bIsScrollingActiveTimerRegistered )
+		{
+			// Register the active timer to handle the inertial scrolling
+			CachedGeometry = MyGeometry;
+			BeginInertialScrolling();
+		}
+
+		return bScrollWasHandled ? FReply::Handled() : FReply::Unhandled();
+	}
+	return FReply::Unhandled();
+}
+
 bool SScrollBox::ScrollBy(const FGeometry& AllottedGeometry, float LocalScrollAmount, EAllowOverscroll Overscrolling, bool InAnimateScroll)
 {
 	Invalidate(EInvalidateWidget::LayoutAndVolatility);
@@ -1129,6 +1175,21 @@ FNavigationReply SScrollBox::OnNavigation(const FGeometry& MyGeometry, const FNa
 	return SCompoundWidget::OnNavigation(MyGeometry, InNavigationEvent);
 }
 
+FReply SScrollBox::OnFocusReceived(const FGeometry& MyGeometry, const FFocusEvent& InFocusEvent)
+{
+	if (OnScrollBoxFocusReceived.IsBound())
+	{
+		OnScrollBoxFocusReceived.ExecuteIfBound();
+		return FReply::Handled();
+	}
+	return FReply::Unhandled();
+}
+
+void SScrollBox::OnFocusLost(const FFocusEvent& InFocusEvent)
+{
+	OnScrollBoxFocusLost.ExecuteIfBound();
+}
+
 void SScrollBox::OnFocusChanging(const FWeakWidgetPath& PreviousFocusPath, const FWidgetPath& NewWidgetPath, const FFocusEvent& InFocusEvent)
 {
 	if (ScrollWhenFocusChanges != EScrollWhenFocusChanges::NoScroll)
diff --git a/Engine/Source/Runtime/Slate/Public/Widgets/Layout/SScrollBox.h b/Engine/Source/Runtime/Slate/Public/Widgets/Layout/SScrollBox.h
--- a/Engine/Source/Runtime/Slate/Public/Widgets/Layout/SScrollBox.h
+++ b/Engine/Source/Runtime/Slate/Public/Widgets/Layout/SScrollBox.h
@@ -28,6 +28,10 @@ class FPaintArgs;
 class FSlateWindowElementList;
 class SScrollPanel;
 
+DECLARE_DELEGATE(FOnScrollBoxFocusReceived);
+DECLARE_DELEGATE(FOnScrollBoxFocusLost);
+
+
 /** Where to scroll the descendant to */
 UENUM(BlueprintType)
 enum class EDescendantScrollDestination : uint8
@@ -164,6 +168,8 @@ public:
 		, _ScrollWhenFocusChanges(EScrollWhenFocusChanges::NoScroll)
 		, _OnUserScrolled()
 	    , _OnScrollBarVisibilityChanged()
+		, _OnFocusReceived()
+		, _OnFocusLost()
 		, _ConsumeMouseWheel(EConsumeMouseWheel::WhenScrollingPossible)
 		{
 			_Clipping = EWidgetClipping::ClipToBounds;
@@ -226,6 +232,12 @@ public:
 		/** Fired when scroll bar visibility changed */
 		SLATE_EVENT(FOnScrollBarVisibilityChanged, OnScrollBarVisibilityChanged)
 
+		/** If bIsFocusable is enabled, Called when the scrollbox focus is received */
+		SLATE_EVENT(FOnScrollBoxFocusReceived, OnFocusReceived)
+		
+		/** If bIsFocusable is enabled, Called when the scrollbox focus is lost */
+		SLATE_EVENT(FOnScrollBoxFocusLost, OnFocusLost)
+
 		SLATE_ARGUMENT(EConsumeMouseWheel, ConsumeMouseWheel);
 
 	SLATE_END_ARGS()
@@ -318,6 +330,10 @@ public:
 
 	SLATE_API void SetConsumeMouseWheel(EConsumeMouseWheel NewConsumeMouseWheel);
 
+	SLATE_API void SetAnalogMouseWheelKey(FKey InAnalogKey);
+
+	SLATE_API void SetIsFocusable(bool bInIsFocusable);
+
 	/** Sets the current orientation of the scrollbox and updates the layout */
 	SLATE_API void SetOrientation(EOrientation InOrientation);
 
@@ -353,12 +369,16 @@ public:
 	SLATE_API virtual void OnMouseEnter( const FGeometry& MyGeometry, const FPointerEvent& MouseEvent ) override;
 	SLATE_API virtual void OnMouseLeave( const FPointerEvent& MouseEvent ) override;
 	SLATE_API virtual FReply OnMouseWheel( const FGeometry& MyGeometry, const FPointerEvent& MouseEvent ) override;
+	SLATE_API virtual FReply OnAnalogValueChanged( const FGeometry& MyGeometry, const FAnalogInputEvent& InAnalogInputEvent);
 	SLATE_API virtual FCursorReply OnCursorQuery( const FGeometry& MyGeometry, const FPointerEvent& CursorEvent ) const override;
 	SLATE_API virtual int32 OnPaint( const FPaintArgs& Args, const FGeometry& AllottedGeometry, const FSlateRect& MyCullingRect, FSlateWindowElementList& OutDrawElements, int32 LayerId, const FWidgetStyle& InWidgetStyle, bool bParentEnabled ) const override;
 	SLATE_API virtual FReply OnTouchEnded(const FGeometry& MyGeometry, const FPointerEvent& InTouchEvent) override;
 	SLATE_API virtual void OnMouseCaptureLost(const FCaptureLostEvent& CaptureLostEvent) override;
 	SLATE_API virtual FNavigationReply OnNavigation(const FGeometry& MyGeometry, const FNavigationEvent& InNavigationEvent) override;
 	SLATE_API virtual void OnFocusChanging(const FWeakWidgetPath& PreviousFocusPath, const FWidgetPath& NewWidgetPath, const FFocusEvent& InFocusEvent) override;
+	SLATE_API virtual FReply OnFocusReceived(const FGeometry& MyGeometry, const FFocusEvent& InFocusEvent) override;
+	SLATE_API virtual void OnFocusLost(const FFocusEvent& InFocusEvent) override;
+	SLATE_API virtual bool SupportsKeyboardFocus() const override { return bIsFocusable; }
 	// End of SWidget interface
 
 protected:
@@ -497,6 +517,12 @@ protected:
 	/** Fired when scroll bar visibility changed */
 	FOnScrollBarVisibilityChanged OnScrollBarVisibilityChanged;
 
+	/** If bIsFocusable is enabled, Called when the scrollbox focus is received */
+	FOnScrollBoxFocusReceived OnScrollBoxFocusReceived;
+		
+	/** If bIsFocusable is enabled, Called when the scrollbox focus is lost */
+	FOnScrollBoxFocusLost OnScrollBoxFocusLost;
+
 	/** The scrolling and stacking orientation. */
 	EOrientation Orientation;
 
@@ -509,6 +535,9 @@ protected:
 	/** How we should handle scrolling with the mouse wheel */
 	EConsumeMouseWheel ConsumeMouseWheel;
 
+	/** Gamepad key used for scrolling. Will adhere to ConsumeMouseWheel */
+	FKey AnalogMouseWheelKey;
+
 	/** Cached geometry for use with the active timer */
 	FGeometry CachedGeometry;
 
@@ -555,6 +584,8 @@ protected:
 	bool bAllowsRightClickDragScrolling : 1;
 	
 	bool bTouchPanningCapture : 1;
+
+	bool bIsFocusable : 1;
 };
 
 class SScrollPanel : public SPanel
diff --git a/Engine/Source/Runtime/SlateCore/Private/SlateCoreClasses.cpp b/Engine/Source/Runtime/SlateCore/Private/SlateCoreClasses.cpp
--- a/Engine/Source/Runtime/SlateCore/Private/SlateCoreClasses.cpp
+++ b/Engine/Source/Runtime/SlateCore/Private/SlateCoreClasses.cpp
@@ -18,6 +18,12 @@ TAutoConsoleVariable<float> GlobalScrollAmount(
 	32.0f,
 	TEXT("How much to scroll for each click of the mouse wheel (in Slate Screen Units)."));
 
+/** How much to scroll for each click of the mouse wheel (in Slate Screen Units). */
+TAutoConsoleVariable<float> GlobalGamepadScrollMultiplier(
+	TEXT("Slate.GlobalGamepadScrollMultiplier"),
+	1.0f,
+	TEXT("How much to multiply scroll amount when using gamepad."));
+
 
 float GSlateContrast = 1;
 
diff --git a/Engine/Source/Runtime/SlateCore/Public/Types/SlateConstants.h b/Engine/Source/Runtime/SlateCore/Public/Types/SlateConstants.h
--- a/Engine/Source/Runtime/SlateCore/Public/Types/SlateConstants.h
+++ b/Engine/Source/Runtime/SlateCore/Public/Types/SlateConstants.h
@@ -9,5 +9,9 @@
 extern SLATECORE_API TAutoConsoleVariable<float> GlobalScrollAmount;
 inline float GetGlobalScrollAmount() { return GlobalScrollAmount.GetValueOnAnyThread(); }
 
+/** How much to multiply scroll amount when using gamepad. */
+extern SLATECORE_API TAutoConsoleVariable<float> GlobalGamepadScrollMultiplier;
+inline float GetGlobalGamepadScrollMultiplier() { return GlobalGamepadScrollMultiplier.GetValueOnAnyThread(); }
+
 /**  */
 extern SLATECORE_API float GSlateContrast;
diff --git a/Engine/Source/Runtime/UMG/Private/Components/ScrollBox.cpp b/Engine/Source/Runtime/UMG/Private/Components/ScrollBox.cpp
--- a/Engine/Source/Runtime/UMG/Private/Components/ScrollBox.cpp
+++ b/Engine/Source/Runtime/UMG/Private/Components/ScrollBox.cpp
@@ -107,7 +107,9 @@ PRAGMA_DISABLE_DEPRECATION_WARNINGS
 		.EnableTouchScrolling(bEnableTouchScrolling)
 		.ConsumePointerInput(bConsumePointerInput)
 		.OnUserScrolled(BIND_UOBJECT_DELEGATE(FOnUserScrolled, SlateHandleUserScrolled))
-		.OnScrollBarVisibilityChanged(BIND_UOBJECT_DELEGATE(FOnScrollBarVisibilityChanged, SlateHandleScrollBarVisibilityChanged));
+		.OnScrollBarVisibilityChanged(BIND_UOBJECT_DELEGATE(FOnScrollBarVisibilityChanged, SlateHandleScrollBarVisibilityChanged))
+		.OnFocusReceived(BIND_UOBJECT_DELEGATE(FOnScrollBoxFocusReceived, SlateHandleFocusReceived))
+		.OnFocusLost(BIND_UOBJECT_DELEGATE(FOnScrollBoxFocusLost, SlateHandleFocusLost));
 PRAGMA_ENABLE_DEPRECATION_WARNINGS
 	for ( UPanelSlot* PanelSlot : Slots )
 	{
@@ -141,6 +143,8 @@ PRAGMA_DISABLE_DEPRECATION_WARNINGS
 	MyScrollBox->SetAllowOverscroll(AllowOverscroll ? EAllowOverscroll::Yes : EAllowOverscroll::No);
 	MyScrollBox->SetScrollBarRightClickDragAllowed(bAllowRightClickDragScrolling);
 	MyScrollBox->SetConsumeMouseWheel(ConsumeMouseWheel);
+	MyScrollBox->SetAnalogMouseWheelKey(AnalogMouseWheelKey);
+	MyScrollBox->SetIsFocusable(bIsFocusable);
 	MyScrollBox->SetAnimateWheelScrolling(bAnimateWheelScrolling);
 	MyScrollBox->SetScrollingAnimationInterpolationSpeed(ScrollAnimationInterpolationSpeed);
 	MyScrollBox->SetWheelScrollMultiplier(WheelScrollMultiplier);
@@ -150,6 +154,7 @@ PRAGMA_DISABLE_DEPRECATION_WARNINGS
 	MyScrollBox->InvalidateStyle();
 	MyScrollBox->SetScrollBarStyle(&WidgetBarStyle);
 	MyScrollBox->InvalidateScrollBarStyle();
+
 PRAGMA_ENABLE_DEPRECATION_WARNINGS
 }
 
@@ -539,6 +544,34 @@ void UScrollBox::EndInertialScrolling()
 	}
 }
 
+FKey UScrollBox::GetAnalogMouseWheelKey() const
+{
+	return AnalogMouseWheelKey;
+}
+
+void UScrollBox::SetAnalogMouseWheelKey(FKey InMouseWheelKey)
+{
+	AnalogMouseWheelKey = InMouseWheelKey;
+	if (MyScrollBox)
+	{
+		MyScrollBox->SetAnalogMouseWheelKey(AnalogMouseWheelKey);
+	}
+}
+
+bool UScrollBox::GetIsFocusable() const
+{
+	return bIsFocusable;
+}
+
+void UScrollBox::SetIsFocusable(bool bInIsFocusable)
+{	
+	bIsFocusable = bInIsFocusable;
+	if (MyScrollBox)
+	{
+		MyScrollBox->SetIsFocusable(bIsFocusable);
+	}
+}
+
 void UScrollBox::SetAlwaysShowScrollbarTrack(bool NewAlwaysShowScrollbarTrack)
 {
 	AlwaysShowScrollbarTrack = NewAlwaysShowScrollbarTrack;
@@ -611,6 +644,16 @@ void UScrollBox::SlateHandleScrollBarVisibilityChanged(EVisibility NewVisibility
 	OnScrollBarVisibilityChanged.Broadcast(ConvertRuntimeToSerializedVisibility(NewVisibility));
 }
 
+void UScrollBox::SlateHandleFocusReceived()
+{
+	OnFocusReceived.Broadcast();
+}
+
+void UScrollBox::SlateHandleFocusLost()
+{
+	OnFocusLost.Broadcast();
+}
+
 #if WITH_EDITOR
 
 const FText UScrollBox::GetPaletteCategory()
diff --git a/Engine/Source/Runtime/UMG/Public/Components/ScrollBox.h b/Engine/Source/Runtime/UMG/Public/Components/ScrollBox.h
--- a/Engine/Source/Runtime/UMG/Public/Components/ScrollBox.h
+++ b/Engine/Source/Runtime/UMG/Public/Components/ScrollBox.h
@@ -11,6 +11,8 @@
 #include "Containers/Ticker.h"
 #include "ScrollBox.generated.h"
 
+DECLARE_DYNAMIC_MULTICAST_DELEGATE(FOnScrollBoxFocusReceivedEvent);
+DECLARE_DYNAMIC_MULTICAST_DELEGATE(FOnScrollBoxFocusLostEvent);
 DECLARE_DYNAMIC_MULTICAST_DELEGATE_OneParam(FOnUserScrolledEvent, float, CurrentOffset);
 DECLARE_DYNAMIC_MULTICAST_DELEGATE_OneParam(FOnScrollBarVisibilityChangedEvent, ESlateVisibility , NewVisibility);
 
@@ -34,6 +36,14 @@ private:
 	/** Option to pass touch input events to widgets behind the ScrollBox. */
 	UPROPERTY(EditAnywhere, BlueprintReadWrite, Getter = "GetConsumePointerInput", Setter = "SetConsumePointerInput", Category = "Scroll", meta = (AllowPrivateAccess = "true", DisplayAfter = "bEnableTouchScrolling"))
 	bool bConsumePointerInput = true;
+		
+	/** The virtual mouse wheel key for gamepad input */
+	UPROPERTY(EditAnywhere, BlueprintReadWrite, Getter = "GetAnalogMouseWheelKey", Setter = "SetAnalogMouseWheelKey", Category = "Scroll", meta = (AllowPrivateAccess = "true", DisplayAfter = "bConsumePointerInput"))
+	FKey AnalogMouseWheelKey;
+
+	/** If true, scrollbox can receive focus. On gamepad, ensure AnalogMouseWheelKey is set to allow user interaction */
+	UPROPERTY(EditAnywhere, BlueprintReadWrite, Getter = "GetIsFocusable", Setter = "SetIsFocusable", Category = "Scroll", meta = (AllowPrivateAccess = "true", DisplayAfter = "AnalogMouseWheelKey"))
+	bool bIsFocusable;
 
 public:
 
@@ -227,6 +237,18 @@ public:
 	UFUNCTION(BlueprintCallable, Category = "Scroll")
 	UMG_API void EndInertialScrolling();
 
+	UFUNCTION(BlueprintCallable, Category = "Scroll")
+	UMG_API FKey GetAnalogMouseWheelKey() const;
+
+	UFUNCTION(BlueprintCallable, Category = "Scroll")
+	UMG_API void SetAnalogMouseWheelKey(FKey InMouseWheelKey);
+
+	UFUNCTION(BlueprintCallable, Category = "Scroll")
+	UMG_API bool GetIsFocusable() const;
+
+	UFUNCTION(BlueprintCallable, Category = "Scroll")
+	UMG_API void SetIsFocusable(bool bInIsFocusable);
+
 public:
 
 	/** Called when the scroll has changed */
@@ -237,6 +259,14 @@ public:
 	UPROPERTY(BlueprintAssignable, Category = "Button|Event")
 	FOnScrollBarVisibilityChangedEvent OnScrollBarVisibilityChanged;
 
+	/** Called when the scroll has changed */
+	UPROPERTY(BlueprintAssignable, Category = "Scroll|Event")
+	FOnScrollBoxFocusReceivedEvent OnFocusReceived;
+
+	/** Called when the scroll has changed */
+	UPROPERTY(BlueprintAssignable, Category = "Scroll|Event")
+	FOnScrollBoxFocusLostEvent OnFocusLost;
+
 	/**
 	 * Updates the scroll offset of the scrollbox.
 	 * @param NewScrollOffset is in Slate Units.
@@ -322,6 +352,10 @@ protected:
 
 	UMG_API void SlateHandleScrollBarVisibilityChanged(EVisibility NewVisibility);
 
+	UMG_API void SlateHandleFocusReceived();
+	
+	UMG_API void SlateHandleFocusLost();
+
 	// Initialize IsFocusable in the constructor before the SWidget is constructed.
 	UMG_API void InitBackPadScrolling(bool InBackPadScrolling);
 	// Initialize IsFocusable in the constructor before the SWidget is constructed.
