[Slate] Added NumSlots, RemoveSlot, ClearChildren, and IsValidSlotIndex to SSplitter. This is consistent with other SPanel widgets.

diff --git a/Engine/Source/Runtime/Slate/Private/Widgets/Layout/SSplitter.cpp b/Engine/Source/Runtime/Slate/Private/Widgets/Layout/SSplitter.cpp
--- a/Engine/Source/Runtime/Slate/Private/Widgets/Layout/SSplitter.cpp
+++ b/Engine/Source/Runtime/Slate/Private/Widgets/Layout/SSplitter.cpp
@@ -57,11 +57,31 @@ SSplitter::FSlot& SSplitter::SlotAt( int32 SlotIndex )
 	return Children[SlotIndex];
 }
 
+int32 SSplitter::RemoveSlot(const TSharedRef<SWidget>& SlotWidget)
+{
+	return Children.Remove(SlotWidget);
+}
+
 void SSplitter::RemoveAt( int32 IndexToRemove )
 {
 	Children.RemoveAt( IndexToRemove );
 }
 
+void SSplitter::ClearChildren()
+{
+	Children.Empty();
+}
+
+int32 SSplitter::NumSlots() const
+{
+	return Children.Num();
+}
+
+bool SSplitter::IsValidSlotIndex(int32 Index) const
+{
+	return Children.IsValidIndex(Index);
+}
+
 /**
 * Construct this widget
 *
diff --git a/Engine/Source/Runtime/Slate/Public/Widgets/Layout/SSplitter.h b/Engine/Source/Runtime/Slate/Public/Widgets/Layout/SSplitter.h
--- a/Engine/Source/Runtime/Slate/Public/Widgets/Layout/SSplitter.h
+++ b/Engine/Source/Runtime/Slate/Public/Widgets/Layout/SSplitter.h
@@ -228,6 +228,12 @@ public:
 	 */
 	SLATE_API SSplitter::FSlot& SlotAt( int32 SlotIndex );
 
+	/** Removes a slot from this splitter panel which contains the specified SWidget
+	 *
+	 * @param SlotWidget The widget to match when searching through the slots
+	 * @returns The index in the children array where the slot was removed and -1 if no slot was found matching the widget
+	 */
+	SLATE_API int32 RemoveSlot( const TSharedRef<SWidget>& SlotWidget );
 
 	/**
 	 * Remove the child at IndexToRemove
@@ -236,6 +242,15 @@ public:
 	 */
 	SLATE_API void RemoveAt( int32 IndexToRemove );
 
+	/** Removes all children from the splitter. */
+	SLATE_API void ClearChildren();
+
+	/** @return the number of slots. */
+	SLATE_API int32 NumSlots() const;
+
+	/** @return if it's a valid index slot index. */
+	SLATE_API bool IsValidSlotIndex(int32 Index) const;
+
 public:
 
 	SLATE_API virtual void OnArrangeChildren( const FGeometry& AllottedGeometry, FArrangedChildren& ArrangedChildren ) const override;
