Fix a bug where icons in SToolBarComboButtonBlock would never show up if a bound icon attribute didn't return a valid icon at the moment of construction.

diff --git a/Engine/Source/Runtime/Slate/Private/Framework/MultiBox/SToolBarComboButtonBlock.cpp b/Engine/Source/Runtime/Slate/Private/Framework/MultiBox/SToolBarComboButtonBlock.cpp
--- a/Engine/Source/Runtime/Slate/Private/Framework/MultiBox/SToolBarComboButtonBlock.cpp
+++ b/Engine/Source/Runtime/Slate/Private/Framework/MultiBox/SToolBarComboButtonBlock.cpp
@@ -104,7 +104,7 @@ void SToolBarComboButtonBlock::BuildMultiBlockWidget(const ISlateStyle* StyleSet
 	TSharedRef<SWidget> IconWidget = SNullWidget::NullWidget;
 	if (!ToolBarComboButtonBlock->bSimpleComboBox)
 	{
-		if (MultiBlock->HasIcon())
+		if (Icon.IsSet())
 		{
 			TSharedRef<SLayeredImage> ActualIconWidget =
 				SNew(SLayeredImage)
@@ -407,7 +407,7 @@ const FCheckBoxStyle* SToolBarComboButtonBlock::GetCheckBoxStyle(const ISlateSty
 	{
 		CheckStyle = &ToolBarStyle.SettingsToggleButton;
 	}	
-	else if (!MultiBlock->HasIcon())
+	else if (!Icon.IsSet())
 	{
 		CheckStyle = &FCoreStyle::Get().GetWidgetStyle<FCheckBoxStyle>("Checkbox");
 	}
