Protect against crash on mac where the optional pointer is actually not set.
UE-260467

#tests No repro steps, mostly code analysis
#rnx

diff --git a/Engine/Source/Runtime/Slate/Private/Widgets/Text/SlateEditableTextLayout.cpp b/Engine/Source/Runtime/Slate/Private/Widgets/Text/SlateEditableTextLayout.cpp
--- a/Engine/Source/Runtime/Slate/Private/Widgets/Text/SlateEditableTextLayout.cpp
+++ b/Engine/Source/Runtime/Slate/Private/Widgets/Text/SlateEditableTextLayout.cpp
@@ -3184,6 +3184,14 @@ void FSlateEditableTextLayout::EndEditTransaction()
 		return;
 	}
 
+	if (!StateBeforeChangingText.IsSet())
+	{
+		// if OwnerWidget->IsTextReadOnly happens on the first transaction in BeginEditTransaction, the count is increased but the State is never being set so we should skip on the work
+		// if EndEditTransaction is getting called after the UndoState has been reset, probably by getting deleted
+		// before the call happens on another thread, there's no point in trying to close the transaction.
+		return;
+	}
+
 	// We're no longer changing text
 	const FText EditedText = GetEditableText();
 
