Added the option for row headers to indicate that they're busy sorting.

A new option "Sorting" was added to SHeaderRow's EColumnSortMode list. When set to this option, the up/down arrows are replaced with a spinner to indicate that sorting is in progress. This was added because in some cases in the editor sorting can take a couple of seconds when the sorting is complex or done on a large list of rows. With this option it's possible to show a spinner for the duration of the sorting so users aren't left wondering why nothing is happening.

diff --git a/Engine/Source/Runtime/Slate/Private/Widgets/Views/SHeaderRow.cpp b/Engine/Source/Runtime/Slate/Private/Widgets/Views/SHeaderRow.cpp
--- a/Engine/Source/Runtime/Slate/Private/Widgets/Views/SHeaderRow.cpp
+++ b/Engine/Source/Runtime/Slate/Private/Widgets/Views/SHeaderRow.cpp
@@ -11,6 +11,7 @@
 #include "Framework/MultiBox/MultiBoxBuilder.h"
 #include "Widgets/Layout/SSpacer.h"
 #include "Widgets/Images/SImage.h"
+#include "Widgets/Images/SThrobber.h"
 #include "Widgets/Text/STextBlock.h"
 #include "Widgets/Layout/SBox.h"
 #include "Widgets/Layout/SScrollBar.h"
@@ -35,6 +36,7 @@ public:
 		: InitialSortMode( EColumnSortMode::Ascending )
 		, SortMode( EColumnSortMode::None )
 		, SortPriority( EColumnSortPriority::Primary )
+		, IsSorting( false )
 		, OnSortModeChanged()
 		, ContextMenuContent( SNullWidget::NullWidget )
 		, ColumnId( NAME_None )
@@ -57,6 +59,7 @@ public:
 		InitialSortMode = Column.InitialSortMode;
 		SortMode = Column.SortMode;
 		SortPriority = Column.SortPriority;
+		IsSorting = Column.IsSorting;
 
 		OnSortModeChanged = Column.OnSortModeChanged;
 		ContextMenuContent = OverrideHeaderMenuContent;
@@ -126,10 +129,21 @@ public:
 						.VAlign( VAlign_Center)
 						.Padding( FMargin( 4.0f, 0.0f, 0.0f, 0.0f) )
 						[
-							SNew( SImage )
-							.ColorAndOpacity( FSlateColor::UseForeground() )
-							.Image( this, &STableColumnHeader::GetSortingBrush )
-							.Visibility( this, &STableColumnHeader::GetSortModeVisibility )
+							SNew( SOverlay )
+							+SOverlay::Slot()
+							[
+								SNew( SImage )
+								.ColorAndOpacity( FSlateColor::UseForeground() )
+								.Image( this, &STableColumnHeader::GetSortingBrush )
+								.Visibility( this, &STableColumnHeader::GetSortModeVisibility )
+							]
+							+SOverlay::Slot()
+							[
+								SNew(SCircularThrobber)
+								.Radius(10.0f)
+								.PieceImage(FAppStyle::Get().GetBrush("Throbber.CircleChunk.Small"))
+								.Visibility( this, &STableColumnHeader::GetSortProcessingVisibility )
+							]
 						]
 					]
 				]
@@ -376,14 +390,22 @@ private:
 		EColumnSortPriority::Type ColumnSortPriority = SortPriority.Get();
 
 		return ( SortMode.Get() == EColumnSortMode::Ascending ? 
-			(SortPriority.Get() == EColumnSortPriority::Secondary ? &Style->SortSecondaryAscendingImage : &Style->SortPrimaryAscendingImage) :
-			(SortPriority.Get() == EColumnSortPriority::Secondary ? &Style->SortSecondaryDescendingImage : &Style->SortPrimaryDescendingImage) );
+			(ColumnSortPriority == EColumnSortPriority::Secondary ? &Style->SortSecondaryAscendingImage : &Style->SortPrimaryAscendingImage) :
+			(ColumnSortPriority == EColumnSortPriority::Secondary ? &Style->SortSecondaryDescendingImage : &Style->SortPrimaryDescendingImage) );
 	}
 
 	/** Checks if sorting mode has been selected */
 	EVisibility GetSortModeVisibility() const
 	{
-		return (SortMode.Get() != EColumnSortMode::None) ? EVisibility::HitTestInvisible : EVisibility::Collapsed;
+		EColumnSortMode::Type ColumnSortMode = SortMode.Get();
+		const bool bShow = ColumnSortMode != EColumnSortMode::None && !IsSorting.Get();
+		return bShow ? EVisibility::HitTestInvisible : EVisibility::Collapsed;
+	}
+
+	/** Checks if sorting is in progress */
+	EVisibility GetSortProcessingVisibility() const
+	{
+		return IsSorting.Get() ? EVisibility::HitTestInvisible : EVisibility::Collapsed;
 	}
 	
 	/** Called when the column title has been clicked to change sorting mode */
@@ -455,7 +477,10 @@ private:
 
 	/** Current sorting order */
 	TAttribute< EColumnSortPriority::Type > SortPriority;
-
+	
+	/** Whether or not the list under the header is being sorted. */
+	TAttribute< bool > IsSorting;
+	
 	/** Callback triggered when sorting mode changes */
 	FOnSortModeChanged OnSortModeChanged;
 
diff --git a/Engine/Source/Runtime/Slate/Public/Widgets/Views/SHeaderRow.h b/Engine/Source/Runtime/Slate/Public/Widgets/Views/SHeaderRow.h
--- a/Engine/Source/Runtime/Slate/Public/Widgets/Views/SHeaderRow.h
+++ b/Engine/Source/Runtime/Slate/Public/Widgets/Views/SHeaderRow.h
@@ -41,7 +41,7 @@ namespace EColumnSortMode
 		Ascending = 1,
 
 		/** Descending */
-		Descending = 2,
+		Descending = 2
 	};
 };
 
@@ -120,6 +120,7 @@ public:
 			, _VAlignCell( VAlign_Fill )
 			, _InitialSortMode( EColumnSortMode::Ascending )
 			, _SortMode( EColumnSortMode::None )
+			, _IsSorting( false )
 			, _OnSort()
 			, _OverflowPolicy( ETextOverflowPolicy::Clip)
 			{}
@@ -156,6 +157,7 @@ public:
 			SLATE_ATTRIBUTE( EColumnSortMode::Type, InitialSortMode )
 			SLATE_ATTRIBUTE( EColumnSortMode::Type, SortMode )
 			SLATE_ATTRIBUTE( EColumnSortPriority::Type, SortPriority )
+			SLATE_ATTRIBUTE( bool, IsSorting )
 			SLATE_EVENT( FOnSortModeChanged, OnSort )
 
 			SLATE_ATTRIBUTE(bool, ShouldGenerateWidget)
@@ -185,6 +187,7 @@ public:
 			, InitialSortMode( InArgs._InitialSortMode )
 			, SortMode( InArgs._SortMode )
 			, SortPriority( InArgs._SortPriority )
+			, IsSorting( InArgs._IsSorting)
 			, OnSortModeChanged( InArgs._OnSort )
 			, ShouldGenerateWidget(InArgs._ShouldGenerateWidget)
 			, bIsVisible(true)
@@ -273,6 +276,7 @@ public:
 		TAttribute< EColumnSortMode::Type > InitialSortMode;
 		TAttribute< EColumnSortMode::Type > SortMode;
 		TAttribute< EColumnSortPriority::Type > SortPriority;
+		TAttribute< bool > IsSorting;
 		FOnSortModeChanged OnSortModeChanged;
 
 		TAttribute<bool> ShouldGenerateWidget;
diff --git a/Engine/Source/Runtime/SlateCore/Private/Styling/StarshipCoreStyle.cpp b/Engine/Source/Runtime/SlateCore/Private/Styling/StarshipCoreStyle.cpp
--- a/Engine/Source/Runtime/SlateCore/Private/Styling/StarshipCoreStyle.cpp
+++ b/Engine/Source/Runtime/SlateCore/Private/Styling/StarshipCoreStyle.cpp
@@ -564,6 +564,7 @@ TSharedRef<ISlateStyle> FStarshipCoreStyle::Create()
 	{
 		Style->Set("Throbber.Chunk", new IMAGE_BRUSH("Common/Throbber_Piece", Icon16x16));
 		Style->Set("Throbber.CircleChunk", new IMAGE_BRUSH("Common/Throbber_Piece", Icon8x8));
+		Style->Set("Throbber.CircleChunk.Small", new IMAGE_BRUSH("Common/Throbber_Piece", Icon4x4));
 	}
 
 	// SExpandableArea defaults...
diff --git a/Engine/Source/Runtime/SlateCore/Public/Styling/CoreStyle.h b/Engine/Source/Runtime/SlateCore/Public/Styling/CoreStyle.h
--- a/Engine/Source/Runtime/SlateCore/Public/Styling/CoreStyle.h
+++ b/Engine/Source/Runtime/SlateCore/Public/Styling/CoreStyle.h
@@ -94,6 +94,7 @@ namespace CoreStyleConstants
 	inline const UE::Slate::FDeprecateVector2DResult Icon8x4(8.0f, 4.0f);
 	inline const UE::Slate::FDeprecateVector2DResult Icon16x4(16.0f, 4.0f);
 	inline const UE::Slate::FDeprecateVector2DResult Icon8x8(8.0f, 8.0f);
+	inline const UE::Slate::FDeprecateVector2DResult Icon4x4(4.0f, 4.0f);
 	inline const UE::Slate::FDeprecateVector2DResult Icon10x10(10.0f, 10.0f);
 	inline const UE::Slate::FDeprecateVector2DResult Icon12x12(12.0f, 12.0f);
 	inline const UE::Slate::FDeprecateVector2DResult Icon12x16(12.0f, 16.0f);
