Correct enhanced tooltip placement when no longer hovering on the original widget.
- The positioning flag should only be reset when there is actually a new tooltip.
- The forcefield offset should not happen for interactive tooltips that have already been positioned. Once positioned, they should remain in place.

diff --git a/Engine/Source/Runtime/Slate/Private/Framework/Application/SlateUser.cpp b/Engine/Source/Runtime/Slate/Private/Framework/Application/SlateUser.cpp
--- a/Engine/Source/Runtime/Slate/Private/Framework/Application/SlateUser.cpp
+++ b/Engine/Source/Runtime/Slate/Private/Framework/Application/SlateUser.cpp
@@ -1294,6 +1294,7 @@ void FSlateUser::UpdateTooltip(const FMenuStack& MenuStack, bool bCanSpawnNewToo
 		if (NewTooltip && bCanSpawnNewTooltip)
 		{
 			NewTooltip->OnOpening();
+			ActiveTooltipInfo.HasBeenPositioned = false;
 		}
 
 		// Some widgets might want to provide an alternative Tooltip Handler.
@@ -1311,8 +1312,6 @@ void FSlateUser::UpdateTooltip(const FMenuStack& MenuStack, bool bCanSpawnNewToo
 				}
 			}
 		}
-
-		ActiveTooltipInfo.HasBeenPositioned = false;
 	}
 	
 	// If a widget under the cursor has a tool-tip forcefield active, then go through any menus
@@ -1374,7 +1373,7 @@ void FSlateUser::UpdateTooltip(const FMenuStack& MenuStack, bool bCanSpawnNewToo
 	}
 
 	// Repel tooltip from a force field, if necessary
-	if (ForceFieldRect.IsSet())
+	if (ForceFieldRect.IsSet() && !ActiveTooltipInfo.WasInteractive)
 	{
 		FVector2f TooltipShift;
 		TooltipShift.X = (ForceFieldRect->Right + SlateDefs::TooltipOffsetFromForceField.X) - DesiredLocation.X;
@@ -1453,8 +1452,8 @@ void FSlateUser::UpdateTooltip(const FMenuStack& MenuStack, bool bCanSpawnNewToo
 					ShowTooltip(NewTooltip.ToSharedRef(), DesiredLocation);
 					ActiveTooltipInfo.SourceWidget = WidgetProvidingNewTooltip;
 				}
+				ActiveTooltip = NewTooltip;
 			}
-			ActiveTooltip = NewTooltip;
 		}
 	}
 
