When a SActionButton or SPrimaryButton have an icon binding, dynamically change the button normal/pressed padding depending on whether the binding resolves to a non-null icon.  This ensures that a SPrimaryButton with an icon binding that returns null appears the same as any other SPrimaryButton without an icon.

diff --git a/Engine/Source/Developer/ToolWidgets/Private/SActionButton.cpp b/Engine/Source/Developer/ToolWidgets/Private/SActionButton.cpp
--- a/Engine/Source/Developer/ToolWidgets/Private/SActionButton.cpp
+++ b/Engine/Source/Developer/ToolWidgets/Private/SActionButton.cpp
@@ -75,7 +75,9 @@ void SActionButton::Construct(const FArguments& InArgs)
 		? &ActionButtonStyle->IconBrush.GetValue()
 		: TAttribute<const FSlateBrush*>();
 
-	const bool bHasIcon = Icon.Get() || Icon.IsBound();
+	const bool bCanEverHaveIcon = Icon.IsBound() || Icon.Get(nullptr) != nullptr;
+
+	const FButtonStyle* SelectedButtonStyle = bCanEverHaveIcon ? IconButtonStyle : ButtonStyle;
 
 	// Check for widget level override, then style override, otherwise get from ActionButtonType
 	const TAttribute<FSlateColor> IconColorAndOpacity =
@@ -109,7 +111,7 @@ void SActionButton::Construct(const FArguments& InArgs)
 			SAssignNew(ComboButton, SComboButton)
 			.HasDownArrow(bHasDownArrow)
 			.ContentPadding(ComboButtonContentPadding)
-			.ButtonStyle(bHasIcon ? IconButtonStyle : ButtonStyle)
+			.ButtonStyle(SelectedButtonStyle)
 			.ComboButtonStyle(ComboButtonStyle)
 			.IsEnabled(InArgs._IsEnabled)
 			.ToolTipText(InArgs._ToolTipText)
@@ -139,11 +141,13 @@ void SActionButton::Construct(const FArguments& InArgs)
 		[
 			SAssignNew(Button, SButton)
 			.ContentPadding(ButtonContentPadding)
-			.ButtonStyle(bHasIcon ? IconButtonStyle : ButtonStyle)
+			.ButtonStyle(SelectedButtonStyle)
 			.IsEnabled(InArgs._IsEnabled)
 			.ToolTipText(InArgs._ToolTipText)
 			.HAlign(HorizontalContentAlignment)
 			.VAlign(VAlign_Center)
+			.NormalPaddingOverride(UE::ToolWidgets::Private::ActionButton::MakeIconPaddingOverride(Icon, SelectedButtonStyle->NormalPadding, ActionButtonStyle->IconNormalPadding))
+			.PressedPaddingOverride(UE::ToolWidgets::Private::ActionButton::MakeIconPaddingOverride(Icon, SelectedButtonStyle->PressedPadding, ActionButtonStyle->IconPressedPadding))
 			.OnClicked(InArgs._OnClicked)
 			[
 				ButtonContent
diff --git a/Engine/Source/Developer/ToolWidgets/Private/SPrimaryButton.cpp b/Engine/Source/Developer/ToolWidgets/Private/SPrimaryButton.cpp
--- a/Engine/Source/Developer/ToolWidgets/Private/SPrimaryButton.cpp
+++ b/Engine/Source/Developer/ToolWidgets/Private/SPrimaryButton.cpp
@@ -20,7 +20,9 @@ void SPrimaryButton::Construct(const FArguments& InArgs)
 		? &ActionButtonStyle->IconBrush.GetValue()
 		: nullptr;
 
-	const bool bHasIcon = Icon.Get() || Icon.IsBound();
+	const bool bCanEverHaveIcon = Icon.IsBound() || Icon.Get(nullptr) != nullptr;
+
+	const FButtonStyle& SelectedButtonStyle = bCanEverHaveIcon ? ActionButtonStyle->GetIconButtonStyle() : ActionButtonStyle->ButtonStyle;
 
 	// Empty/default args will resolve from the ActionButtonStyle
 	const TSharedRef<SWidget> ButtonContent =
@@ -34,11 +36,13 @@ void SPrimaryButton::Construct(const FArguments& InArgs)
 	SButton::Construct(
 		SButton::FArguments()
 		.ContentPadding(ActionButtonStyle->GetButtonContentPadding())
-		.ButtonStyle(bHasIcon ? &ActionButtonStyle->GetIconButtonStyle() : &ActionButtonStyle->ButtonStyle)
+		.ButtonStyle(&SelectedButtonStyle)
 		.IsEnabled(InArgs._IsEnabled)
 		.ToolTipText(InArgs._ToolTipText)
 		.HAlign(static_cast<EHorizontalAlignment>(ActionButtonStyle->HorizontalContentAlignment))
 		.VAlign(VAlign_Center)
+		.NormalPaddingOverride(UE::ToolWidgets::Private::ActionButton::MakeIconPaddingOverride(Icon, SelectedButtonStyle.NormalPadding, ActionButtonStyle->IconNormalPadding))
+		.PressedPaddingOverride(UE::ToolWidgets::Private::ActionButton::MakeIconPaddingOverride(Icon, SelectedButtonStyle.PressedPadding, ActionButtonStyle->IconPressedPadding))
 		.OnClicked(InArgs._OnClicked)
 		[
 			ButtonContent
diff --git a/Engine/Source/Developer/ToolWidgets/Private/ToolWidgetsSlateTypes.cpp b/Engine/Source/Developer/ToolWidgets/Private/ToolWidgetsSlateTypes.cpp
--- a/Engine/Source/Developer/ToolWidgets/Private/ToolWidgetsSlateTypes.cpp
+++ b/Engine/Source/Developer/ToolWidgets/Private/ToolWidgetsSlateTypes.cpp
@@ -171,6 +171,18 @@ FActionButtonStyle& FActionButtonStyle::SetIconColorAndOpacity(const FSlateColor
 	return *this;
 }
 
+FActionButtonStyle& FActionButtonStyle::SetIconNormalPadding(const FMargin& InIconNormalPadding)
+{
+	IconNormalPadding = InIconNormalPadding;
+	return *this;
+}
+
+FActionButtonStyle& FActionButtonStyle::SetIconPressedPadding(const FMargin& InIconPressedPadding)
+{
+	IconPressedPadding = InIconPressedPadding;
+	return *this;
+}
+
 void FActionButtonStyle::UnlinkColors()
 {
 	ButtonStyle.UnlinkColors();
diff --git a/Engine/Source/Developer/ToolWidgets/Private/ToolWidgetsStyle.cpp b/Engine/Source/Developer/ToolWidgets/Private/ToolWidgetsStyle.cpp
--- a/Engine/Source/Developer/ToolWidgets/Private/ToolWidgetsStyle.cpp
+++ b/Engine/Source/Developer/ToolWidgets/Private/ToolWidgetsStyle.cpp
@@ -137,7 +137,8 @@ namespace UE::ToolWidgets
 					.SetButtonStyle(ButtonStyle)
 					.SetButtonContentPadding(ModifyContentPadding(ButtonStyle, DefaultButtonContentPadding))
 					.SetComboButtonContentPadding(ModifyContentPadding(ButtonStyle, DefaultButtonContentPadding))
-					.SetIconButtonStyle(GetWidgetStyle<FButtonStyle>("PrimaryButtonLabelAndIcon"))
+					.SetIconNormalPadding(GetMargin("PrimaryButtonLabelAndIconNormalPadding"))
+					.SetIconPressedPadding(GetMargin("PrimaryButtonLabelAndIconPressedPadding"))
 					.SetTextBlockStyle(GetWidgetStyle<FTextBlockStyle>("PrimaryButtonText"))
 					.SetIconColorAndOpacity(FSlateColor::UseForeground());
 
diff --git a/Engine/Source/Developer/ToolWidgets/Private/ToolWidgetsUtilitiesPrivate.cpp b/Engine/Source/Developer/ToolWidgets/Private/ToolWidgetsUtilitiesPrivate.cpp
--- a/Engine/Source/Developer/ToolWidgets/Private/ToolWidgetsUtilitiesPrivate.cpp
+++ b/Engine/Source/Developer/ToolWidgets/Private/ToolWidgetsUtilitiesPrivate.cpp
@@ -26,7 +26,7 @@ namespace UE::ToolWidgets::Private
 			check(InIconColorAndOpacity.IsSet() || InIconColorAndOpacity.IsBound());
 			check(InTextBlockStyle);
 
-			const bool bCanEverHaveIcon = InIcon.Get() || InIcon.IsBound();
+			const bool bCanEverHaveIcon = InIcon.IsBound() || InIcon.Get(nullptr) != nullptr;
 
 			static constexpr float DefaultIconHeight = FToolWidgetsStylePrivate::FActionButton::DefaultIconHeight;
 			static constexpr float IconTextPadding = FToolWidgetsStylePrivate::FActionButton::DefaultIconLabelSpacing;
@@ -114,5 +114,29 @@ namespace UE::ToolWidgets::Private
 
 			return MakeButtonContent(Icon, IconColorAndOpacity, InText, TextBlockStyle);
 		}
+
+		TAttribute<FMargin> MakeIconPaddingOverride(const TAttribute<const FSlateBrush*>& Icon, const FMargin& ButtonStylePadding, const TOptional<FMargin>& IconPadding)
+		{
+			if (IconPadding.IsSet())
+			{
+				if (Icon.IsBound())
+				{
+					return TAttribute<FMargin>::CreateLambda([Icon, ButtonStylePadding, IconPadding = IconPadding.GetValue()]()
+					{
+						if (Icon.Get(nullptr) != nullptr)
+						{
+							return IconPadding;
+						}
+
+						return ButtonStylePadding;
+					});
+				}
+				else if (Icon.Get(nullptr) != nullptr)
+				{
+					return IconPadding.GetValue();
+				}
+			}
+			return {};
+		}
 	}
 }
diff --git a/Engine/Source/Developer/ToolWidgets/Private/ToolWidgetsUtilitiesPrivate.h b/Engine/Source/Developer/ToolWidgets/Private/ToolWidgetsUtilitiesPrivate.h
--- a/Engine/Source/Developer/ToolWidgets/Private/ToolWidgetsUtilitiesPrivate.h
+++ b/Engine/Source/Developer/ToolWidgets/Private/ToolWidgetsUtilitiesPrivate.h
@@ -23,5 +23,7 @@ namespace UE::ToolWidgets::Private
 			const TAttribute<FSlateColor>& InIconColorAndOpacity,
 			const TAttribute<FText>& InText,
 			const FTextBlockStyle* InTextBlockStyle);
+
+		TAttribute<FMargin> MakeIconPaddingOverride(const TAttribute<const FSlateBrush*>& Icon, const FMargin& ButtonStylePadding, const TOptional<FMargin>& IconPadding);
 	}
 }
diff --git a/Engine/Source/Developer/ToolWidgets/Public/ToolWidgetsSlateTypes.h b/Engine/Source/Developer/ToolWidgets/Public/ToolWidgetsSlateTypes.h
--- a/Engine/Source/Developer/ToolWidgets/Public/ToolWidgetsSlateTypes.h
+++ b/Engine/Source/Developer/ToolWidgets/Public/ToolWidgetsSlateTypes.h
@@ -116,6 +116,20 @@ struct FActionButtonStyle : public FSlateWidgetStyle
 	TOptional<FSlateColor> IconColorAndOpacity;
 	FActionButtonStyle& SetIconColorAndOpacity(const FSlateColor& InIconColorAndOpacity);
 
+	/**
+	 * If set and the button's icon is non-null, overrides the button style's additional spacing between the button's border and the content when not pressed.
+	 */
+	UPROPERTY(EditAnywhere, BlueprintReadWrite, Category = "Appearance")
+	TOptional<FMargin> IconNormalPadding;
+	FActionButtonStyle& SetIconNormalPadding(const FMargin& InIconNormalPadding);
+
+	/**
+	 * If set and the button's icon is non-null, overrides the button style's additional spacing between the button's border and the content when pressed.
+	 */
+	UPROPERTY(EditAnywhere, BlueprintReadWrite, Category = "Appearance")
+	TOptional<FMargin> IconPressedPadding;
+	FActionButtonStyle& SetIconPressedPadding(const FMargin& InIconPressedPadding);
+
 	/**
 	* Unlinks all colors in this style.
 	* @see FSlateColor::Unlink
diff --git a/Engine/Source/Runtime/Slate/Private/Widgets/Input/SButton.cpp b/Engine/Source/Runtime/Slate/Private/Widgets/Input/SButton.cpp
--- a/Engine/Source/Runtime/Slate/Private/Widgets/Input/SButton.cpp
+++ b/Engine/Source/Runtime/Slate/Private/Widgets/Input/SButton.cpp
@@ -25,6 +25,18 @@ void SButton::PrivateRegisterAttributes(FSlateAttributeInitializer& AttributeIni
 		{
 			static_cast<SButton&>(Widget).UpdatePadding();
 		}));
+
+	SLATE_ADD_MEMBER_ATTRIBUTE_DEFINITION_WITH_NAME(AttributeInitializer, "NormalPadding", NormalPaddingAttribute, EInvalidateWidgetReason::Layout)
+		.OnValueChanged(FSlateAttributeDescriptor::FAttributeValueChangedDelegate::CreateLambda([](SWidget& Widget)
+		{
+			static_cast<SButton&>(Widget).UpdatePadding();
+		}));
+
+	SLATE_ADD_MEMBER_ATTRIBUTE_DEFINITION_WITH_NAME(AttributeInitializer, "PressedPadding", PressedPaddingAttribute, EInvalidateWidgetReason::Layout)
+		.OnValueChanged(FSlateAttributeDescriptor::FAttributeValueChangedDelegate::CreateLambda([](SWidget& Widget)
+		{
+			static_cast<SButton&>(Widget).UpdatePadding();
+		}));
 	
 	SLATE_ADD_MEMBER_ATTRIBUTE_DEFINITION_WITH_NAME(AttributeInitializer, "AppearPressed", AppearPressedAttribute, EInvalidateWidgetReason::Paint)
 		.OnValueChanged(FSlateAttributeDescriptor::FAttributeValueChangedDelegate::CreateLambda([](SWidget& Widget)
@@ -54,6 +66,8 @@ void SButton::PrivateRegisterAttributes(FSlateAttributeInitializer& AttributeIni
 SButton::SButton()
 	: BorderForegroundColorAttribute(*this)
 	, ContentPaddingAttribute(*this)
+	, NormalPaddingAttribute(*this)
+	, PressedPaddingAttribute(*this)
 	, AppearPressedAttribute(*this)
 {
 #if WITH_ACCESSIBILITY
@@ -76,6 +90,18 @@ void SButton::Construct( const FArguments& InArgs )
 
 	BorderForegroundColorAttribute.Assign(*this, InArgs._ForegroundColor);
 
+	bIsStyleNormalPaddingOverridden = InArgs._NormalPaddingOverride.IsSet();
+	if (bIsStyleNormalPaddingOverridden)
+	{
+		NormalPaddingAttribute.Assign(*this, InArgs._NormalPaddingOverride);
+	}
+
+	bIsStylePressedPaddingOverridden = InArgs._PressedPaddingOverride.IsSet();
+	if (bIsStylePressedPaddingOverridden)
+	{
+		PressedPaddingAttribute.Assign(*this, InArgs._PressedPaddingOverride);
+	}
+
 	OnClicked = InArgs._OnClicked;
 	OnPressed = InArgs._OnPressed;
 	OnReleased = InArgs._OnReleased;
@@ -159,11 +185,11 @@ int32 SButton::OnPaint(const FPaintArgs& Args, const FGeometry& AllottedGeometry
 FMargin SButton::GetCombinedPadding() const
 {
 	return ( IsPressed() )
-		? ContentPaddingAttribute.Get() + Style->PressedPadding
-		: ContentPaddingAttribute.Get() + Style->NormalPadding;
+		? ContentPaddingAttribute.Get() + PressedPaddingAttribute.Get()
+		: ContentPaddingAttribute.Get() + NormalPaddingAttribute.Get();
 }
 
-//~ Update when { ContentPaddingAttribute, IsPressed, Style }
+//~ Update when { ContentPaddingAttribute, PressedPaddingAttribute, NormalPaddingAttribute, IsPressed, Style }
 void SButton::UpdatePadding()
 {
 	SetPadding(GetCombinedPadding());
@@ -614,6 +640,16 @@ void SButton::SetButtonStyle(const FButtonStyle* InButtonStyle)
 	/* Get pointer to the button Style */
 	Style = InButtonStyle;
 
+	if (!bIsStyleNormalPaddingOverridden)
+	{
+		NormalPaddingAttribute.Set(*this, Style->NormalPadding);
+	}
+
+	if (!bIsStylePressedPaddingOverridden)
+	{
+		PressedPaddingAttribute.Set(*this, Style->PressedPadding);
+	}
+
 	HoveredSound = Style->HoveredSlateSound;
 	PressedSound = Style->PressedSlateSound;
 
diff --git a/Engine/Source/Runtime/Slate/Public/Widgets/Input/SButton.h b/Engine/Source/Runtime/Slate/Public/Widgets/Input/SButton.h
--- a/Engine/Source/Runtime/Slate/Public/Widgets/Input/SButton.h
+++ b/Engine/Source/Runtime/Slate/Public/Widgets/Input/SButton.h
@@ -69,7 +69,13 @@ public:
 	
 		/** Spacing between button's border and the content. */
 		SLATE_ATTRIBUTE( FMargin, ContentPadding )
-	
+
+		/** If set, overrides the button style's additional spacing between the button's border and the content when not pressed. */
+		SLATE_ATTRIBUTE( FMargin, NormalPaddingOverride )
+
+		/** If set, overrides the button style's additional spacing between the button's border and the content when pressed. */
+		SLATE_ATTRIBUTE( FMargin, PressedPaddingOverride )
+
 		/** The text to display in this button, if no custom content is specified */
 		SLATE_ATTRIBUTE( FText, Text )
 	
@@ -308,11 +314,21 @@ private:
 	/** True if this button is currently in a pressed state */
 	uint8 bIsPressed:1;
 
+	/** True if NormalPaddingAttribute is overriding the button style's normal padding */
+	uint8 bIsStyleNormalPaddingOverridden:1;
+
+	/** True if PressedPaddingAttribute is overriding the button style's pressed padding */
+	uint8 bIsStylePressedPaddingOverridden:1;
+
 private:
 	/** Optional foreground color that will be inherited by all of this widget's contents */
 	TSlateAttribute<FSlateColor> BorderForegroundColorAttribute;
 	/** Padding specified by the user; it will be combind with the button's internal padding. */
 	TSlateAttribute<FMargin> ContentPaddingAttribute;
+	/** Normal padding override specified by the user, or the button style's normal padding if not being overridden. */
+	TSlateAttribute<FMargin> NormalPaddingAttribute;
+	/** Pressed padding override specified by the user, or the button style's pressed padding if not being overridden. */
+	TSlateAttribute<FMargin> PressedPaddingAttribute;
 	/** Optional foreground color that will be inherited by all of this widget's contents */
 	TSlateAttribute<bool> AppearPressedAttribute;
 };
diff --git a/Engine/Source/Runtime/SlateCore/Private/Styling/StarshipCoreStyle.cpp b/Engine/Source/Runtime/SlateCore/Private/Styling/StarshipCoreStyle.cpp
--- a/Engine/Source/Runtime/SlateCore/Private/Styling/StarshipCoreStyle.cpp
+++ b/Engine/Source/Runtime/SlateCore/Private/Styling/StarshipCoreStyle.cpp
@@ -1386,9 +1386,8 @@ void FStarshipCoreStyle::SetupButtonStyles(TSharedRef<FStyle>& Style)
 		.SetNormalPadding(ButtonMargins)
 		.SetPressedPadding(PressedButtonMargins);
 
-	const FButtonStyle PrimaryButtonLabelAndIcon = FButtonStyle(PrimaryButton)
-		.SetNormalPadding(FMargin(4.f, 1.5f, 4.f, 1.5f))
-		.SetPressedPadding(FMargin(4.f, 2.5f, 4.f, 0.5f));
+	const FMargin PrimaryButtonLabelAndIconNormalPadding(4.f, 1.5f, 4.f, 1.5f);
+	const FMargin PrimaryButtonLabelAndIconPressedPadding(4.f, 2.5f, 4.f, 0.5f);
 
 	const FButtonStyle Button = FButtonStyle()
 		.SetNormal(FSlateRoundedBoxBrush(FStyleColors::Secondary, 4.0f, FStyleColors::Input, InputFocusThickness))
@@ -1433,7 +1432,8 @@ void FStarshipCoreStyle::SetupButtonStyles(TSharedRef<FStyle>& Style)
 		Style->Set("ButtonText", NormalText);
 
 		Style->Set("PrimaryButton", PrimaryButton);
-		Style->Set("PrimaryButtonLabelAndIcon", PrimaryButtonLabelAndIcon);
+		Style->Set("PrimaryButtonLabelAndIconNormalPadding", PrimaryButtonLabelAndIconNormalPadding);
+		Style->Set("PrimaryButtonLabelAndIconPressedPadding", PrimaryButtonLabelAndIconPressedPadding);
 		Style->Set("Button", Button);
 		Style->Set("EditorUtilityButton", EditorUtilityButton);
 		Style->Set("SimpleButton", SimpleButtonIconOnly);
