[Slate] Added PrioritizedWrapPanel

diff --git a/Engine/Source/Developer/ToolMenus/Private/ToolMenuSection.cpp b/Engine/Source/Developer/ToolMenus/Private/ToolMenuSection.cpp
--- a/Engine/Source/Developer/ToolMenus/Private/ToolMenuSection.cpp
+++ b/Engine/Source/Developer/ToolMenus/Private/ToolMenuSection.cpp
@@ -36,6 +36,7 @@ void FToolMenuSection::InitGeneratedSectionCopy(const FToolMenuSection& Source,
 	Context = InContext;
 	Alignment = Source.Alignment;
 	Visibility = Source.Visibility;
+	ResizeParams = Source.ResizeParams;
 }
 
 bool FToolMenuSection::IsRegistering() const
diff --git a/Engine/Source/Developer/ToolMenus/Private/ToolMenus.cpp b/Engine/Source/Developer/ToolMenus/Private/ToolMenus.cpp
--- a/Engine/Source/Developer/ToolMenus/Private/ToolMenus.cpp
+++ b/Engine/Source/Developer/ToolMenus/Private/ToolMenus.cpp
@@ -470,6 +470,7 @@ public:
 					Widget.ToSharedRef(),
 					LabelToDisplay.Get(),
 					Block.WidgetData.StyleParams,
+					Block.WidgetData.ResizeParams,
 					Block.WidgetData.bSearchable,
 					Block.ToolTip.Get(),
 					Block.Icon,
@@ -1536,11 +1537,11 @@ void UToolMenus::PopulateMenuBuilder(FMenuBuilder& MenuBuilder, UToolMenu* MenuD
 			{
 				LabelText = FText::FromName(Section.Name);
 			}
-			MenuBuilder.BeginSection(Section.Name, LabelText, Section.Visibility);
+			MenuBuilder.BeginSection(Section.Name, LabelText, Section.Visibility, Section.ResizeParams);
 		}
 		else
 		{
-			MenuBuilder.BeginSection(Section.Name, Section.Label, Section.Visibility);
+			MenuBuilder.BeginSection(Section.Name, Section.Label, Section.Visibility, Section.ResizeParams);
 		}
 
 		for (FToolMenuEntry& Block : Section.Blocks)
@@ -1903,7 +1904,7 @@ void UToolMenus::PopulateToolBarBuilderWithEntry(
 				return UToolMenus::ConvertUIAction(Block, MenuData->Context);
 			}
 		}();
-	
+
 	TAttribute<FText> ToolbarLabelOverride;
 	if (Block.ToolBarData.LabelOverride.IsSet())
 	{
@@ -2192,7 +2193,7 @@ void UToolMenus::PopulateToolBarBuilder(FToolBarBuilder& ToolBarBuilder, UToolMe
 		const bool bFirstSectionInAlignmentGroup = Section.Alignment != LastAlignment && Section.Alignment != EToolMenuSectionAlign::Default;
 
 		const bool bSectionShouldHaveSeparator = MenuData->bSeparateSections && !bFirstSectionInAlignmentGroup;
-		ToolBarBuilder.BeginSection(Section.Name, bSectionShouldHaveSeparator);
+		ToolBarBuilder.BeginSection(Section.Name, bSectionShouldHaveSeparator, Section.ResizeParams);
 
 		FName PreviousBlockGroupName = NAME_None;
 		for (int BlockIndex = 0; BlockIndex < Section.Blocks.Num(); ++BlockIndex)
@@ -2854,7 +2855,7 @@ TSharedRef<SWidget> UToolMenus::GenerateWidget(UToolMenu* GeneratedMenu)
 	if (GeneratedMenu->IsEditing())
 	{
 		// Convert toolbar into menu during editing
-		if (GeneratedMenu->MenuType == EMultiBoxType::ToolBar || GeneratedMenu->MenuType == EMultiBoxType::VerticalToolBar || GeneratedMenu->MenuType == EMultiBoxType::UniformToolBar || GeneratedMenu->MenuType == EMultiBoxType::SlimHorizontalToolBar)
+		if (GeneratedMenu->MenuType == EMultiBoxType::ToolBar || GeneratedMenu->MenuType == EMultiBoxType::VerticalToolBar || GeneratedMenu->MenuType == EMultiBoxType::UniformToolBar || GeneratedMenu->MenuType == EMultiBoxType::SlimHorizontalToolBar || GeneratedMenu->MenuType == EMultiBoxType::SlimWrappingToolBar)
 		{
 			for (FToolMenuSection& Section : GeneratedMenu->Sections)
 			{
@@ -2910,7 +2911,7 @@ TSharedRef<SWidget> UToolMenus::GenerateWidget(UToolMenu* GeneratedMenu)
 		TSharedRef<SWidget> Result = MenuBarBuilder.MakeWidget();
 		GeneratedWidget = Result;
 	}
-	else if (GeneratedMenu->MenuType == EMultiBoxType::ToolBar || GeneratedMenu->MenuType == EMultiBoxType::VerticalToolBar || GeneratedMenu->MenuType == EMultiBoxType::UniformToolBar || GeneratedMenu->MenuType == EMultiBoxType::SlimHorizontalToolBar)
+	else if (GeneratedMenu->MenuType == EMultiBoxType::ToolBar || GeneratedMenu->MenuType == EMultiBoxType::VerticalToolBar || GeneratedMenu->MenuType == EMultiBoxType::UniformToolBar || GeneratedMenu->MenuType == EMultiBoxType::SlimHorizontalToolBar || GeneratedMenu->MenuType == EMultiBoxType::SlimWrappingToolBar)
 	{
 		FToolBarBuilder ToolbarBuilder(GeneratedMenu->MenuType, GeneratedMenu->Context.CommandList, GeneratedMenu->MenuName, GeneratedMenu->Context.GetAllExtenders(), GeneratedMenu->bToolBarForceSmallIcons);
 		ToolbarBuilder.SetExtendersEnabled(GeneratedMenu->bExtendersEnabled);
@@ -3152,7 +3153,7 @@ bool UToolMenus::RefreshMenuWidget(const FName InName, FGeneratedToolMenuWidget&
 		PopulateMenuBarBuilder(MenuBarBuilder, GeneratedMenu);
 		MultiBoxWidget->SetMultiBox(MenuBarBuilder.GetMultiBox());
 	}
-	else if (GeneratedMenu->MenuType == EMultiBoxType::ToolBar || GeneratedMenu->MenuType == EMultiBoxType::VerticalToolBar || GeneratedMenu->MenuType == EMultiBoxType::UniformToolBar || GeneratedMenu->MenuType == EMultiBoxType::SlimHorizontalToolBar)
+	else if (GeneratedMenu->MenuType == EMultiBoxType::ToolBar || GeneratedMenu->MenuType == EMultiBoxType::VerticalToolBar || GeneratedMenu->MenuType == EMultiBoxType::UniformToolBar || GeneratedMenu->MenuType == EMultiBoxType::SlimHorizontalToolBar || GeneratedMenu->MenuType == EMultiBoxType::SlimWrappingToolBar)
 	{
 		FToolBarBuilder ToolbarBuilder(GeneratedMenu->MenuType, GeneratedMenu->Context.CommandList, GeneratedMenu->MenuName, GeneratedMenu->Context.GetAllExtenders(), GeneratedMenu->bToolBarForceSmallIcons);
 		ToolbarBuilder.SetExtendersEnabled(GeneratedMenu->bExtendersEnabled);
diff --git a/Engine/Source/Developer/ToolMenus/Public/ToolMenuEntry.h b/Engine/Source/Developer/ToolMenus/Public/ToolMenuEntry.h
--- a/Engine/Source/Developer/ToolMenus/Public/ToolMenuEntry.h
+++ b/Engine/Source/Developer/ToolMenus/Public/ToolMenuEntry.h
@@ -85,6 +85,7 @@ public:
 	/** Optional override placement to customize menu placement once opened via e.g. a toolbar menu button */
 	TAttribute<EMenuPlacement> PlacementOverride;
 
+	/** Various Resize parameters and overrides, will take precedent over those stored in FToolMenuEntryWidgetData for ToolBar entries. */
 	FMenuEntryResizeParams ResizeParams;
 
 	/** Optional action override to use when the entry appears in a toolbar */
@@ -113,6 +114,9 @@ public:
 
 	/** Various Style parameters and overrides */
 	FMenuEntryStyleParams StyleParams;
+
+	/** Various Resize parameters and overrides */
+	FMenuEntryResizeParams ResizeParams;
 };
 
 struct FToolMenuCustomWidgetContext
diff --git a/Engine/Source/Developer/ToolMenus/Public/ToolMenuSection.h b/Engine/Source/Developer/ToolMenus/Public/ToolMenuSection.h
--- a/Engine/Source/Developer/ToolMenus/Public/ToolMenuSection.h
+++ b/Engine/Source/Developer/ToolMenus/Public/ToolMenuSection.h
@@ -139,6 +139,9 @@ public:
 
 	EToolMenuSectionAlign Alignment;
 
+	/** Section-level resize parameters. */
+	FMenuEntryResizeParams ResizeParams;
+
 private:
 
 	bool bIsRegistering;
diff --git a/Engine/Source/Runtime/Slate/Private/Framework/MultiBox/MultiBox.cpp b/Engine/Source/Runtime/Slate/Private/Framework/MultiBox/MultiBox.cpp
--- a/Engine/Source/Runtime/Slate/Private/Framework/MultiBox/MultiBox.cpp
+++ b/Engine/Source/Runtime/Slate/Private/Framework/MultiBox/MultiBox.cpp
@@ -33,6 +33,7 @@
 #include "Framework/MultiBox/SUniformToolbarPanel.h"
 #include "Styling/ToolBarStyle.h"
 #include "Misc/CoreMiscDefines.h"
+#include "SPrioritizedWrapBox.h"
 
 #define LOCTEXT_NAMESPACE "MultiBox"
 
@@ -722,12 +723,12 @@ EVisibility SMultiBoxWidget::GetCustomizationBorderDragVisibility(const FName In
 	return EVisibility::Collapsed;
 }
 
-void SMultiBoxWidget::AddBlockWidget(const FMultiBlock& Block, TSharedPtr<SHorizontalBox> HorizontalBox, TSharedPtr<SVerticalBox> VerticalBox, EMultiBlockLocation::Type InLocation, bool bSectionContainsIcons, TSharedPtr<const FToolBarComboButtonBlock> OptionsBlock)
+void SMultiBoxWidget::AddBlockWidget(const FMultiBlock& Block, TSharedPtr<SHorizontalBox> HorizontalBox, TSharedPtr<SVerticalBox> VerticalBox, TSharedPtr<SPanel> Panel, EMultiBlockLocation::Type InLocation, bool bSectionContainsIcons, TSharedPtr<const FToolBarComboButtonBlock> OptionsBlock)
 {
 	check( MultiBox.IsValid() );
 
 	// Skip Separators for Uniform Tool Bars
-	if ( Block.GetType() == EMultiBlockType::Separator  && (MultiBox->GetType() == EMultiBoxType::UniformToolBar || MultiBox->GetType() == EMultiBoxType::SlimHorizontalUniformToolBar))
+	if ( Block.IsSeparator()  && (MultiBox->GetType() == EMultiBoxType::UniformToolBar || MultiBox->GetType() == EMultiBoxType::SlimHorizontalUniformToolBar))
 	{
 		return;
 	}
@@ -881,6 +882,86 @@ void SMultiBoxWidget::AddBlockWidget(const FMultiBlock& Block, TSharedPtr<SHoriz
 			}
 		}
 		break;
+	case EMultiBoxType::SlimWrappingToolBar:
+		{
+			FMenuEntryStyleParams StyleParams;
+			const bool bOverride = Block.GetAlignmentOverrides(StyleParams);
+
+			{
+				static const FName PrioritizedWrapBoxTypeName("SPrioritizedWrapBox");
+			
+				// SPrioritizedWrapBox isn't public, so is stored as an SPanel in the header. We need to check it before casting
+				if (Panel->GetType() != PrioritizedWrapBoxTypeName)
+				{
+					break;
+				}
+
+				TSharedPtr<SPrioritizedWrapBox> PrioritizedWrapBox = StaticCastSharedPtr<SPrioritizedWrapBox>(Panel);
+				if (!PrioritizedWrapBox.IsValid())
+				{
+					break;
+				}
+
+				const FMenuEntryResizeParams ResizeParams = Block.GetResizeParams();
+
+				TAttribute<bool> AllowWrapping = ResizeParams.Wrapping.Allow;
+				if (!AllowWrapping.IsSet()
+					&& StyleSet->HasWidgetStyle<FToolBarStyle>(MultiBox->GetStyleName()))
+				{
+					const FToolBarStyle& Style = StyleSet->GetWidgetStyle<FToolBarStyle>(MultiBox->GetStyleName());
+					AllowWrapping = Style.bAllowWrappingDefault;
+				}
+
+				SPrioritizedWrapBox::FScopedWidgetSlotArguments NewSlot = PrioritizedWrapBox->AddSlot();
+				NewSlot.AllowWrapping(AllowWrapping);
+				NewSlot.WrapPriority(ResizeParams.Wrapping.Priority.IsSet() ? ResizeParams.Wrapping.Priority : FMenuEntryResizeParams::FWrappingParams::DefaultPriority);
+				NewSlot.WrapMode(ResizeParams.Wrapping.Mode.IsSet() ? ResizeParams.Wrapping.Mode : FMenuEntryResizeParams::FWrappingParams::DefaultMode);
+				NewSlot.VerticalOverflowBehavior(ResizeParams.Wrapping.VerticalOverflowBehavior.Get(FMenuEntryResizeParams::FWrappingParams::DefaultVerticalOverflowBehavior));
+				NewSlot.ExcludeIfFirstOrLast(Block.IsSeparator());
+
+				NewSlot.Padding(0.f)
+				[
+					FinalWidgetWithHook
+				];
+
+				FinalMultiBoxWidgets.Add(FinalWidgetWithHook.ToSharedPtr(), BlockRef.ToSharedPtr());
+
+				NewSlot.SizeParam(FAuto());
+
+				if (bOverride)
+				{
+					NewSlot
+					.HAlign(StyleParams.HorizontalAlignment)
+					.VAlign(StyleParams.VerticalAlignment.Get(VAlign_Fill));
+
+					if (StyleParams.SizeRule.IsSet())
+					{
+						FSizeParam::ESizeRule SizeRule = StyleParams.SizeRule.GetValue();
+						if (SizeRule == FSizeParam::SizeRule_Stretch)
+						{
+							NewSlot.SizeParam(FStretch(StyleParams.FillSize.Get(1.0f)));
+						}
+						else if (SizeRule == FSizeParam::SizeRule_StretchContent)
+						{
+							NewSlot.SizeParam(FStretchContent(
+								StyleParams.FillSize.Get(1.0f),
+								StyleParams.FillSizeMin.Get(StyleParams.FillSize.Get(1.0f))));
+						}
+					}
+
+					if (StyleParams.MinimumSize.IsSet() || StyleParams.MinimumSize.IsBound())
+					{
+						NewSlot.MinSize(StyleParams.MinimumSize);
+					}
+
+					if (StyleParams.MaximumSize.IsSet() || StyleParams.MaximumSize.IsBound())
+					{
+						NewSlot.MaxSize(StyleParams.MaximumSize);
+					}
+				}
+			}
+		}
+		break;
 	case EMultiBoxType::VerticalToolBar:
 		{
 			if (UniformToolbarPanel.IsValid())
@@ -1078,6 +1159,7 @@ void SMultiBoxWidget::BuildMultiBoxWidget()
 	TSharedPtr<SVerticalBox> VerticalBox;
 	TSharedPtr<SWidget> MainWidget;
 	TSharedPtr<SHorizontalBox> HorizontalBox;
+	TSharedPtr<SPanel> Panel;
 
 	/** The current row of buttons for if the multibox type is a button row */
 	TSharedPtr<SHorizontalBox> ButtonRow;
@@ -1115,6 +1197,16 @@ void SMultiBoxWidget::BuildMultiBoxWidget()
 					.StyleName(StyleName);
 		}
 		break;
+	case EMultiBoxType::SlimWrappingToolBar:
+		{
+			const FToolBarStyle& Style = StyleSet->GetWidgetStyle<FToolBarStyle>(MultiBox->GetStyleName());
+			const float LineHeight = Style.IconSize.Y + Style.ButtonStyle.NormalPadding.Top + Style.ButtonStyle.NormalPadding.Bottom; 
+			MainWidget = Panel = SNew(SPrioritizedWrapBox)
+				.MinLineHeight(LineHeight)
+				.LinePadding(Style.BackgroundPadding.Bottom + Style.BackgroundPadding.Top)
+				.GroupedWrapping(true);
+			break;
+		}
 	case EMultiBoxType::VerticalToolBar:
 		{
 			MainWidget = VerticalBox = ClippedVerticalBox = SNew(SClippingVerticalBox)
@@ -1332,13 +1424,13 @@ void SMultiBoxWidget::BuildMultiBoxWidget()
 		{
 			// Add the drag preview before if we have it. This block shows where the custom block will be 
 			// added if the user drops it
-			AddBlockWidget(*DragPreview.PreviewBlock, HorizontalBox, VerticalBox, EMultiBlockLocation::None, bSectionContainsIcons, OptionsBlock);
+			AddBlockWidget(*DragPreview.PreviewBlock, HorizontalBox, VerticalBox, Panel, EMultiBlockLocation::None, bSectionContainsIcons, OptionsBlock);
 		}
 		
 		// Do not add a block if it is being dragged
 		if( !IsBlockBeingDragged( Blocks[Index] ) )
 		{
-			AddBlockWidget(Block, HorizontalBox, VerticalBox, Location, bSectionContainsIcons, OptionsBlock);
+			AddBlockWidget(Block, HorizontalBox, VerticalBox, Panel, Location, bSectionContainsIcons, OptionsBlock);
 		}
 
 		const bool bIsSearchWidgetBlock = Index == 0 && MultiBox->bHasSearchWidget;
@@ -1355,7 +1447,7 @@ void SMultiBoxWidget::BuildMultiBoxWidget()
 				// Note that the search widget might be collapsed, so the visibility of this spacer needs to be dynamic (like the search box itself)
 				SpacerBlock->SetVisibilityOverride(TAttribute<EVisibility>(this, &SMultiBoxWidget::GetSearchBoxVisibility));
 
-				AddBlockWidget(*SpacerBlock, HorizontalBox, VerticalBox, EMultiBlockLocation::None, bSectionContainsIcons, OptionsBlock);
+				AddBlockWidget(*SpacerBlock, HorizontalBox, VerticalBox, Panel, EMultiBlockLocation::None, bSectionContainsIcons, OptionsBlock);
 			}
 		}
 	}
@@ -1367,7 +1459,7 @@ void SMultiBoxWidget::BuildMultiBoxWidget()
 		if (!IsBlockBeingDragged(Pair.Key))
 		{
 			TSharedPtr<const FToolBarComboButtonBlock> OptionsBlock;
-			AddBlockWidget(*Pair.Key, HorizontalBox, VerticalBox, EMultiBlockLocation::None, /*bSectionContainsIcons*/false, OptionsBlock);
+			AddBlockWidget(*Pair.Key, HorizontalBox, VerticalBox, Panel, EMultiBlockLocation::None, /*bSectionContainsIcons*/false, OptionsBlock);
 		}
 	}
 
@@ -1423,9 +1515,10 @@ void SMultiBoxWidget::BuildMultiBoxWidget()
 		RootBorder.ToSharedRef()
 	];
 
-	// Save pointers to horizontal/vertical box so that we can insert more blocks later if necessary (e.g. FlattenSubMenusRecursive)
+	// Save pointers to horizontal/vertical/panel box so that we can insert more blocks later if necessary (e.g. FlattenSubMenusRecursive)
 	MainHorizontalBox = HorizontalBox;
 	MainVerticalBox = VerticalBox;
+	MainPanelBox = Panel;
 }
 
 void SMultiBoxWidget::UpdateMultiBoxWidget(TSharedRef<FMultiBox> InMultiBox)
@@ -1922,7 +2015,7 @@ void SMultiBoxWidget::FlattenSubMenusRecursive(uint32 MaxRecursionLevels)
 		if (!IsBlockBeingDragged(Pair.Key))
 		{
 			TSharedPtr<const FToolBarComboButtonBlock> OptionsBlock;
-			AddBlockWidget(*Pair.Key, MainHorizontalBox, MainVerticalBox, EMultiBlockLocation::None, /*bSectionContainsIcons*/false, OptionsBlock);
+			AddBlockWidget(*Pair.Key, MainHorizontalBox, MainVerticalBox, MainPanelBox, EMultiBlockLocation::None, /*bSectionContainsIcons*/false, OptionsBlock);
 		}
 	}
 }
diff --git a/Engine/Source/Runtime/Slate/Private/Framework/MultiBox/MultiBoxBuilder.cpp b/Engine/Source/Runtime/Slate/Private/Framework/MultiBox/MultiBoxBuilder.cpp
--- a/Engine/Source/Runtime/Slate/Private/Framework/MultiBox/MultiBoxBuilder.cpp
+++ b/Engine/Source/Runtime/Slate/Private/Framework/MultiBox/MultiBoxBuilder.cpp
@@ -19,6 +19,69 @@
 #include "Framework/MultiBox/SGroupMarkerBlock.h"
 #include "Styling/ToolBarStyle.h"
 
+namespace UE::Slate::Private
+{
+	/** Utility to set a Target attribute to the given Source attribute, if the target is not already set. */
+	template <typename AttributeType>
+	static void ApplyUnsetAttributeOverride(TAttribute<AttributeType>* InTargetAttribute, const TAttribute<AttributeType>* InSourceAttribute)
+	{
+		if (InTargetAttribute->IsSet() && !InTargetAttribute->IsBound())
+		{
+			// If the target is already set, don't override it
+			return;
+		}
+
+		if (InSourceAttribute->IsSet())
+		{
+			// Target is not set, and source is, so override Target with Source
+			*InTargetAttribute = *InSourceAttribute;
+		}
+	}
+
+	/** Applies section-level overrides to the given block, if any. Currently only affects Resize parameters. */
+	static void ApplySectionOverridesToBlock(
+		const TSharedRef<FMultiBlock>& InMultiBlock,
+		const TAttribute<FMenuEntryResizeParams>& InSectionResizeParams)
+	{
+		// If per-entry resize params aren't set, and section params are, inject those of the section
+		if (InSectionResizeParams.IsSet())
+		{
+			TAttribute<FMenuEntryResizeParams> BlockResizeParamsAttribute = InMultiBlock->GetResizeParams();
+			if (!BlockResizeParamsAttribute.IsSet())
+			{
+				BlockResizeParamsAttribute = InSectionResizeParams;
+			}
+			else if (BlockResizeParamsAttribute.IsSet() && !BlockResizeParamsAttribute.IsBound())
+			{
+				const FMenuEntryResizeParams& SectionResizeParams = InSectionResizeParams.Get();
+
+				FMenuEntryResizeParams BlockResizeParams = BlockResizeParamsAttribute.Get();
+				ApplyUnsetAttributeOverride(&BlockResizeParams.Wrapping.Allow, &SectionResizeParams.Wrapping.Allow);
+				ApplyUnsetAttributeOverride(&BlockResizeParams.Wrapping.Priority, &SectionResizeParams.Wrapping.Priority);
+
+				BlockResizeParamsAttribute = BlockResizeParams;
+			}
+			else if (BlockResizeParamsAttribute.IsBound())
+			{
+				// Intercept and inject wrapping properties
+				TAttribute<FMenuEntryResizeParams> NewResizeParams(
+					TAttribute<FMenuEntryResizeParams>::Create(
+						TAttribute<FMenuEntryResizeParams>::FGetter::CreateLambda([BlockResizeParamsAttribute, SectionResizeParams = InSectionResizeParams.Get()]
+						{
+							FMenuEntryResizeParams ResizeParams = BlockResizeParamsAttribute.Get();
+							ApplyUnsetAttributeOverride(&ResizeParams.Wrapping.Allow, &SectionResizeParams.Wrapping.Allow);
+							ApplyUnsetAttributeOverride(&ResizeParams.Wrapping.Priority, &SectionResizeParams.Wrapping.Priority);
+
+							return ResizeParams;
+						})));
+
+				BlockResizeParamsAttribute = NewResizeParams;
+			}
+
+			InMultiBlock->SetResizeParams(InSectionResizeParams);
+		}
+	}
+}
 
 FMultiBoxBuilder::FMultiBoxBuilder( const EMultiBoxType InType, FMultiBoxCustomization InCustomization, const bool bInShouldCloseWindowAfterMenuSelection, const TSharedPtr< const FUICommandList >& InCommandList, TSharedPtr<FExtender> InExtender, FName InTutorialHighlightName, FName InMenuName )
 	: MultiBox( FMultiBox::Create( InType, (InMenuName != NAME_None) ? FMultiBoxCustomization::AllowCustomization(InMenuName) : InCustomization, bInShouldCloseWindowAfterMenuSelection ) )
@@ -243,7 +306,7 @@ TSharedRef< class SWidget > FMenuBuilder::MakeWidget( FMultiBox::FOnMakeMultiBox
 	return MultiBox->MakeWidget(bSearchable, InMakeMultiBoxBuilderOverride, MaxHeightAttribute);
 }
 
-void FMenuBuilder::BeginSection(FName InExtensionHook, const TAttribute< FText >& InHeadingText, const TAttribute<EVisibility>& InVisibility)
+void FMenuBuilder::BeginSection(FName InExtensionHook, const TAttribute< FText >& InHeadingText, const TAttribute<EVisibility>& InVisibility, const TAttribute<FMenuEntryResizeParams>& InResizeParams)
 {
 	check(CurrentSectionExtensionHook == NAME_None && !bSectionNeedsToBeApplied);
 
@@ -255,6 +318,7 @@ void FMenuBuilder::BeginSection(FName InExtensionHook, const TAttribute< FText >
 	CurrentSectionExtensionHook = InExtensionHook;
 	CurrentSectionHeadingText = InHeadingText.Get();
 	CurrentSectionVisibility = InVisibility;
+	CurrentSectionResizeParams = InResizeParams;
 
 	// Do apply the section beginning if we are in developer "show me all the hooks" mode
 	if (FMultiBoxSettings::DisplayMultiboxHooks.Get())
@@ -272,6 +336,7 @@ void FMenuBuilder::EndSection()
 	bSectionNeedsToBeApplied = false;
 	CurrentSectionHeadingText = FText::GetEmpty();
 	CurrentSectionVisibility = EVisibility::Visible;
+	CurrentSectionResizeParams = TAttribute<FMenuEntryResizeParams>();
 
 	ApplyHook(SectionExtensionHook, EExtensionHook::After);
 }
@@ -411,17 +476,25 @@ void FMenuBuilder::AddWidget( TSharedRef<SWidget> InWidget, const FText& Label,
 }
 
 void FMenuBuilder::AddWidget(const TSharedRef<SWidget>& InWidget, const FText& InLabel, const FMenuEntryStyleParams& InStyleParams, const bool bInSearchable, const TAttribute<FText>& InToolTipText, const TAttribute<FSlateIcon>& InIcon, const TAttribute<EVisibility>& InVisibility)
+{
+	AddWidget(InWidget, InLabel, InStyleParams, FMenuEntryResizeParams(), bInSearchable, InToolTipText, InIcon);
+}
+
+void FMenuBuilder::AddWidget(const TSharedRef<SWidget>& InWidget, const FText& InLabel, const FMenuEntryStyleParams& InStyleParams, const TAttribute<FMenuEntryResizeParams>& InResizeParams, const bool bInSearchable, const TAttribute<FText>& InToolTipText, const TAttribute<FSlateIcon>& InIcon, const TAttribute<EVisibility>& InVisibility)
 {
 	ApplySectionBeginning();
 
 	TSharedRef<FWidgetBlock> NewWidgetBlock = MakeShared<FWidgetBlock>(InWidget, InLabel, InToolTipText, InStyleParams, InIcon);
 	NewWidgetBlock->SetSearchable( bInSearchable );
+	NewWidgetBlock->SetResizeParams(InResizeParams);
 
 	if (InVisibility.IsSet())
 	{
 		NewWidgetBlock->SetVisibilityOverride(InVisibility);
 	}
 
+	UE::Slate::Private::ApplySectionOverridesToBlock(NewWidgetBlock, CurrentSectionResizeParams);
+
 	MultiBox->AddMultiBlock( NewWidgetBlock );
 }
 
@@ -452,7 +525,8 @@ void FMenuBuilder::ApplySectionBeginning()
 	{
 		if (!CurrentSectionHeadingText.IsEmpty())
 		{
-			TSharedRef<FHeadingBlock> HeadingBlock = MakeShareable( new FHeadingBlock(CurrentSectionExtensionHook, CurrentSectionHeadingText) ) ;
+			TSharedRef<FHeadingBlock> HeadingBlock = MakeShared<FHeadingBlock>(CurrentSectionExtensionHook, CurrentSectionHeadingText);
+			HeadingBlock->SetResizeParams(CurrentSectionResizeParams);
 
 			if (CurrentSectionVisibility.IsSet())
 			{
@@ -560,6 +634,8 @@ void FToolBarBuilder::AddToolBarButton(
 	NewToolBarButtonBlock->SetVisibilityOverride(InVisibilityOverride);
 	NewToolBarButtonBlock->SetResizeParams(InResizeParams);
 
+	UE::Slate::Private::ApplySectionOverridesToBlock(NewToolBarButtonBlock, CurrentSectionResizeParams);
+
 	MultiBox->AddMultiBlock( NewToolBarButtonBlock );
 
 	ApplyHook(InExtensionHook, EExtensionHook::After);
@@ -644,6 +720,8 @@ void FToolBarBuilder::AddToolBarButton(
 	NewToolBarButtonBlock->SetVisibilityOverride(InVisibilityOverride);
 	NewToolBarButtonBlock->SetResizeParams(InResizeParams);
 
+	UE::Slate::Private::ApplySectionOverridesToBlock(NewToolBarButtonBlock, CurrentSectionResizeParams);
+
 	MultiBox->AddMultiBlock( NewToolBarButtonBlock );
 
 	ApplyHook(InExtensionHook, EExtensionHook::After);
@@ -681,6 +759,8 @@ void FToolBarBuilder::AddComboButton(
 	NewToolBarComboButtonBlock->SetVisibilityOverride(InVisibilityOverride);
 	NewToolBarComboButtonBlock->SetResizeParams(InResizeParams);
 
+	UE::Slate::Private::ApplySectionOverridesToBlock(NewToolBarComboButtonBlock, CurrentSectionResizeParams);
+
 	MultiBox->AddMultiBlock( NewToolBarComboButtonBlock );
 }
 
@@ -821,6 +901,8 @@ void FToolBarBuilder::AddWidgetInternal(
 	NewWidgetBlock->SetCustomMenuDelegate(InCustomMenuDelegate);
 	NewWidgetBlock->SetVisibilityOverride(InVisibilityOverride);
 	NewWidgetBlock->SetResizeParams(InResizeParams);
+
+	UE::Slate::Private::ApplySectionOverridesToBlock(NewWidgetBlock, CurrentSectionResizeParams);
 }
 
 void FToolBarBuilder::AddSeparator(
@@ -836,12 +918,14 @@ void FToolBarBuilder::AddSeparator(
 	NewSeparatorBlock->SetVisibilityOverride(InVisibilityOverride);
 	NewSeparatorBlock->SetResizeParams(InResizeParams);
 
+	UE::Slate::Private::ApplySectionOverridesToBlock(NewSeparatorBlock, CurrentSectionResizeParams);
+
 	MultiBox->AddMultiBlock(NewSeparatorBlock);
 
 	ApplyHook(InExtensionHook, EExtensionHook::After);
 }
 
-void FToolBarBuilder::BeginSection(FName InExtensionHook, bool bInSectionShouldHaveSeparator)
+void FToolBarBuilder::BeginSection(FName InExtensionHook, bool bInSectionShouldHaveSeparator, const TAttribute<FMenuEntryResizeParams>& InResizeParams)
 {
 	checkf(CurrentSectionExtensionHook == NAME_None && !bSectionNeedsToBeApplied, TEXT("Did you forget to call EndSection()?"));
 
@@ -852,7 +936,8 @@ void FToolBarBuilder::BeginSection(FName InExtensionHook, bool bInSectionShouldH
 	bSectionNeedsToBeApplied = true;
 	bSectionShouldHaveSeparator = bInSectionShouldHaveSeparator;
 	CurrentSectionExtensionHook = InExtensionHook;
-	
+	CurrentSectionResizeParams = InResizeParams;
+
 	// Do apply the section beginning if we are in developer "show me all the hooks" mode
 	if (FMultiBoxSettings::DisplayMultiboxHooks.Get())
 	{
@@ -866,6 +951,7 @@ void FToolBarBuilder::EndSection()
 {
 	FName SectionExtensionHook = CurrentSectionExtensionHook;
 	CurrentSectionExtensionHook = NAME_None;
+	CurrentSectionResizeParams = TAttribute<FMenuEntryResizeParams>();
 	bSectionNeedsToBeApplied = false;
 
 	ApplyHook(SectionExtensionHook, EExtensionHook::After);
@@ -893,6 +979,7 @@ void FToolBarBuilder::ApplySectionBeginning()
 		{
 			TSharedRef<FToolBarSeparatorBlock> NewSeparatorBlock = MakeShared<FToolBarSeparatorBlock>(CurrentSectionExtensionHook);
 			NewSeparatorBlock->SetStyleNameOverride(CurrentStyleOverride);
+			NewSeparatorBlock->SetResizeParams(CurrentSectionResizeParams);
 
 			MultiBox->AddMultiBlock(NewSeparatorBlock);
 		}
diff --git a/Engine/Source/Runtime/Slate/Private/Framework/MultiBox/SPrioritizedWrapBox.cpp b/Engine/Source/Runtime/Slate/Private/Framework/MultiBox/SPrioritizedWrapBox.cpp
new file mode 100644
--- /dev/null
+++ b/Engine/Source/Runtime/Slate/Private/Framework/MultiBox/SPrioritizedWrapBox.cpp
@@ -0,0 +1,912 @@
+// Copyright Epic Games, Inc. All Rights Reserved.
+
+#include "Framework/MultiBox/SPrioritizedWrapBox.h"
+
+#include "Algo/AllOf.h"
+#include "Framework/Application/SlateApplication.h"
+#include "Layout/LayoutUtils.h"
+
+namespace UE::Slate::PrioritizedWrapBox
+{
+	uint32 GetTypeHash(const SPrioritizedWrapBox::FSlot& InSlot)
+	{
+		return HashCombineFast(
+			::GetTypeHash(InSlot.GetWrapPriority()),
+			::GetTypeHash(InSlot.GetWrapMode()),
+			::GetTypeHash(InSlot.GetWidget()));
+	}
+
+	class FChildArranger
+	{
+	private:
+		static constexpr EOrientation Orientation = EOrientation::Orient_Horizontal;
+		inline static const FVector2f MinimumLineSize = FVector2f(0.001f, 0.001f); // non-zero minimum
+
+		struct FSlotProxy : ::FSlotProxy
+		{
+			using SlotType = SPrioritizedWrapBox::FSlot;
+
+			FSlotProxy() = default;
+
+			FSlotProxy(const int32 InSlotIndex, const SlotType& InSlot)
+				: ::FSlotProxy(InSlotIndex, InSlot)
+			{
+				WrapMode = InSlot.GetWrapMode();
+				VerticalOverflowBehavior = InSlot.GetVerticalOverflowBehavior();
+				VerticalExpansionThreshold = InSlot.GetVerticalExpansionThreshold();
+				WrapPriority = InSlot.GetWrapPriority();
+				bExcludeIfFirstOrLast = InSlot.GetExcludeIfFirstOrLast();
+			}
+
+			/** Applies the given slot's values to this proxy. Will return true if any values have changed from those stored. */
+			bool UpdateFromSlot(const int32 InSlotIndex, const SlotType& InSlot)
+			{
+				bool bAnyValueChanged = ::FSlotProxy::UpdateFromSlot<SlotType, true>(InSlotIndex, InSlot);
+
+				bAnyValueChanged = bAnyValueChanged || (WrapMode != InSlot.GetWrapMode());
+				WrapMode = InSlot.GetWrapMode();
+
+				bAnyValueChanged = bAnyValueChanged || (VerticalOverflowBehavior != InSlot.GetVerticalOverflowBehavior());
+				VerticalOverflowBehavior = InSlot.GetVerticalOverflowBehavior();
+
+				bAnyValueChanged = bAnyValueChanged || (VerticalExpansionThreshold != InSlot.GetVerticalExpansionThreshold());
+				VerticalExpansionThreshold = InSlot.GetVerticalExpansionThreshold();
+
+				bAnyValueChanged = bAnyValueChanged || (WrapPriority != InSlot.GetWrapPriority());
+				WrapPriority = InSlot.GetWrapPriority();
+
+				bAnyValueChanged = bAnyValueChanged || (bExcludeIfFirstOrLast != InSlot.GetExcludeIfFirstOrLast());
+				bExcludeIfFirstOrLast = InSlot.GetExcludeIfFirstOrLast();
+
+				return bAnyValueChanged;
+			}
+
+			EWrapMode WrapMode = EWrapMode::Preferred;
+			EVerticalOverflowBehavior VerticalOverflowBehavior = EVerticalOverflowBehavior::Default;
+			TOptional<float> VerticalExpansionThreshold;
+			int32 WrapPriority = INDEX_NONE;
+			bool bExcludeIfFirstOrLast = false;
+		};
+
+		/** Represents a contiguous block of layout elements, which can be either a single slot or sequence of slots. */
+		struct FBlock
+		{
+			/** The original, left-to-right index. */
+			int32 SequentialIndex = INDEX_NONE;
+
+			int32 WrapPriority = INDEX_NONE;
+
+			/** Whether the block can wrap at all. */
+			bool bCanWrap = true;
+
+			/** Based on it's desired size, it was determined this doesn't vertically expand. */
+			bool bHasVerticalExpansion = false;
+
+			/** Based on one or more member slots' VerticalBehavior, this could *possibly* expand - the ReserveLength indicates a "probe length". */
+			bool bCanVerticallyExpand = false;
+
+			/** The last cached or calculated desired size. */
+			FVector2f DesiredSize = FVector2f::Zero();
+
+			/** Size adjusted by the layout algorithm. */
+			FVector2f AdjustedSize = FVector2f::Zero();
+
+			/** The minimum acceptable length of this block, which may be the same or less than the ReserveLength. */
+			FVector2f::FReal MinLength = 0.0f;
+
+			/** The desired length to reserve, used to test for wrapping. */
+			FVector2f::FReal ReserveLength = 0.0f;
+
+			/** The 2D area, calculated from the DesiredSize and maintained when calculating the ReserveLength. */
+			FVector2f::FReal Area = 0.0f;
+
+			/** The actual slots that make up this block, either singular or multiple when using grouping. */
+			TArray<int32, TInlineAllocator<16>> SlotIndices;
+
+			void ResetToDefault()
+			{
+				SequentialIndex = INDEX_NONE;
+				WrapPriority = INDEX_NONE;
+				bCanWrap = true;
+				bHasVerticalExpansion = false;
+				bCanVerticallyExpand = false;
+				DesiredSize = FVector2f::Zero();
+				AdjustedSize = FVector2f::Zero();
+				ReserveLength = 0.0f;
+				MinLength = 0.0f;
+				Area = 0.0f;
+				SlotIndices.Reset();
+			}
+
+			// Min length by default, Max length if expandable
+			FVector2f::FReal GetEffectiveLength() const
+			{
+				return bHasVerticalExpansion ? ReserveLength : MinLength;
+			}
+
+			friend bool operator==(const FBlock& InLeft, const FBlock& InRight) { return InLeft.SequentialIndex == InRight.SequentialIndex; }
+			friend bool operator==(const uint32& InLeft, const FBlock& InRight) { return InLeft == InRight.SequentialIndex; }
+			friend bool operator<(const FBlock& InLeft, const FBlock& InRight) { return InLeft.SequentialIndex < InRight.SequentialIndex; }
+		};
+
+		struct FBlockLine
+		{
+			int32 Index = INDEX_NONE;
+			FVector2f::FReal MinLength;
+			FVector2f::FReal MaxLength;
+			FVector2f::FReal EffectiveLength; // Min length by default, Max length if expandable
+			FVector2f::FReal Height;
+			TArray<FBlock> Blocks;
+		};
+
+	public:
+		FChildArranger() = default;
+
+		template <bool UseGroupedWrapping = false>
+		FVector2D GetDesiredSize(const SPrioritizedWrapBox& InWidget, const TPanelChildren<SPrioritizedWrapBox::FSlot>& InChildren)
+		{
+			// This effectively performs a pre-arrangement of the slots
+			// based only on the slot's individual desired size and the parent widgets wrapping length.
+			// The parents DesiredSize/Geometry isn't always valid, or is a frame behind.
+
+			int32 ParentLineLength = FMath::FloorToInt32(InWidget.GetPaintSpaceGeometry().Size.X);
+			const int32 PreferredLineLength = InWidget.GetPreferredSize();
+
+			// A slot that uses Preferred doesn't even attempt wrapping unless the parent size is the same or less than preferred
+			const bool bDoPreferredWrapping = NumPreferredWrappingChildren > 0 && ParentLineLength <= PreferredLineLength;
+
+			// A slot that uses Parent always attempts wrapping
+			const bool bDoParentWrapping = NumParentWrappingChildren > 0;
+
+			const float MinLineHeight = InWidget.GetMinLineHeight().Get(0.0f);
+			float MaxLineHeight = std::numeric_limits<float>::max();
+
+			if (ParentLineLength == 0)
+			{
+				LastParentLineLength = ParentLineLength = std::numeric_limits<uint16>::max();
+				LastParentSize = FVector2f::One();
+				MaxLineHeight = MinLineHeight;
+			}
+
+			const bool bSlotsHaveChanged = UpdateChildProxies(InChildren);
+
+			if (bSlotsHaveChanged)
+			{
+				UpdateBlocks(InChildren, MinLineHeight);
+			}
+
+			const FVector2f ParentSize = InWidget.GetTickSpaceGeometry().Size;
+			const bool bParentSizeChanged = (LastParentSize != ParentSize)
+				|| (bDoPreferredWrapping && (PreferredLineLength != LastPreferredLineLength))
+				|| (bDoParentWrapping && (ParentLineLength != LastParentLineLength));
+
+			FVector2f DesiredSize = FVector2f::ZeroVector;
+
+			if (bSlotsHaveChanged || bParentSizeChanged)
+			{
+				int32 NumArrangedBlocks = 0;
+
+				TArray<FBlockLine> Lines;
+				Lines.Reserve(Blocks.Num());
+
+				TArray<FBlock> CurrentLineBlocks = Blocks;
+
+				TArray<FBlock> NextLineBlocks;
+				NextLineBlocks.Reserve(Blocks.Num());
+
+				const float LinePadding = InWidget.GetLinePadding();
+				float LineOffset = 0.0f;
+
+				// We keep chopping the last item off to get the current wrapping candidate
+				TConstArrayView<FBlock> SortedBlockView = MakeConstArrayView(WrappableSortedBlockView.GetData(), WrappableSortedBlockView.Num());
+
+				// Re-calculates the line height based on its blocks
+				auto UpdateLineHeight = [&](FBlockLine& InLine)
+				{
+					InLine.Height = 0.0f;
+					for (int32 BlockIndex = 0; BlockIndex < InLine.Blocks.Num(); ++BlockIndex)
+					{
+						InLine.Height = FMath::Max(InLine.Height, InLine.Blocks[BlockIndex].DesiredSize.Y);
+					}
+				};
+
+				auto AddBlockToLine = [](FBlockLine& InLine, const FBlock& InBlock)
+				{
+					InLine.MinLength += InBlock.MinLength;
+					InLine.MaxLength += InBlock.ReserveLength;
+					InLine.EffectiveLength += InBlock.GetEffectiveLength();
+					InLine.Height = FMath::Max(InLine.Height, InBlock.DesiredSize.Y);
+					InLine.Blocks.Emplace(InBlock);
+				};
+
+				auto RemoveBlockFromLine = [&](FBlockLine& InLine, const FBlock& InBlock) -> bool
+				{
+					if (const int32 FoundIndex = InLine.Blocks.Find(InBlock);
+						FoundIndex != INDEX_NONE)
+					{
+						InLine.EffectiveLength -= InBlock.GetEffectiveLength();
+						InLine.MinLength -= InBlock.MinLength;
+						InLine.MaxLength -= InBlock.ReserveLength;
+						InLine.Blocks.RemoveAt(FoundIndex, EAllowShrinking::No);
+
+						UpdateLineHeight(InLine);
+
+						return true;
+					}
+
+					return false;
+				};
+
+				auto MakeBlockLine = [&](FBlockLine& InLine, const TConstArrayView<FBlock> InBlocks)
+				{
+					for (int32 BlockIndex = 0; BlockIndex < InBlocks.Num(); ++BlockIndex)
+					{
+						const FBlock& CurrentBlock = InBlocks[BlockIndex];
+						AddBlockToLine(InLine, CurrentBlock);
+					}
+				};
+
+				auto HasMoreRemovableBlocks = [&]()
+				{
+					return !SortedBlockView.IsEmpty();
+				};
+
+				auto LineHasMultipleBlocksAndPriorities = [&](const FBlockLine& InLine)
+				{
+					const bool bHasMoreThanOneBlock = InLine.Blocks.Num() > 1;
+					if constexpr (UseGroupedWrapping)
+					{
+						return bHasMoreThanOneBlock;
+					}
+					else
+					{
+						if (bHasMoreThanOneBlock)
+						{
+							int32 PriorityToMatch = InLine.Blocks[0].WrapPriority;
+							return bHasMoreThanOneBlock
+								&& !Algo::AllOf(InLine.Blocks, [PriorityToMatch](const FBlock& InBlock)
+								{
+									return PriorityToMatch == InBlock.WrapPriority;
+								});
+						}
+
+						return false;
+					}					
+				};
+
+				int32 LineIndex = 0;
+				while (NumArrangedBlocks < Blocks.Num() && HasMoreRemovableBlocks())
+				{
+					FBlockLine& CurrentLine = Lines.Emplace_GetRef();
+					CurrentLine.Index = LineIndex;
+
+					MakeBlockLine(CurrentLine, CurrentLineBlocks);
+
+					int32 NumSortablesUncheckedForExpansion = SortedBlockView.Num();
+
+					// Remove vertical expandables, which rely on the full line length (not min + full combination)
+					while (NumSortablesUncheckedForExpansion > 0
+						&& CurrentLine.MaxLength > ParentLineLength
+						&& LineHasMultipleBlocksAndPriorities(CurrentLine)
+						&& HasMoreRemovableBlocks())
+					{
+						const FBlock& BlockToRemove = SortedBlockView.Last();
+						const bool bCanRemoveBlock = BlockToRemove.bCanWrap && BlockToRemove.bCanVerticallyExpand && BlockToRemove.WrapPriority != CurrentLine.Blocks[0].WrapPriority;
+
+						--NumSortablesUncheckedForExpansion;
+
+						if (bCanRemoveBlock)
+						{
+							bool bShouldRemoveBlock = true;
+							const FVector2f::FReal EffectiveLengthToRemove = BlockToRemove.GetEffectiveLength();
+
+							if (BlockToRemove.MinLength != BlockToRemove.ReserveLength)
+							{
+								const float LineLengthAfterRemovalOfEffective = CurrentLine.MaxLength - EffectiveLengthToRemove;
+								if (LineLengthAfterRemovalOfEffective + EffectiveLengthToRemove < ParentLineLength)
+								{
+									// We're cool
+									bShouldRemoveBlock = false;
+								}
+							}
+
+							if (bShouldRemoveBlock)
+							{
+								FBlock BlockCopy = BlockToRemove;
+								if (RemoveBlockFromLine(CurrentLine, BlockToRemove))
+								{
+									// Don't include 0-length entries, they'll dictate the line height without having any visual content
+									if (!FMath::IsNearlyZero(EffectiveLengthToRemove))
+									{
+										NextLineBlocks.Emplace(BlockCopy);
+									}
+
+									// Remove block from further consideration for this line
+									SortedBlockView.LeftChopInline(1);
+								}
+							}
+						}
+					}
+
+					// Horizontal, prioritized wrapping
+					while (CurrentLine.EffectiveLength > ParentLineLength
+						&& LineHasMultipleBlocksAndPriorities(CurrentLine)
+						&& HasMoreRemovableBlocks())
+					{
+						const FBlock& BlockToRemove = SortedBlockView.Last();
+						const bool bCanRemoveBlock = BlockToRemove.bCanWrap || BlockToRemove.WrapPriority != CurrentLine.Blocks[0].WrapPriority;
+
+						if (bCanRemoveBlock)
+						{
+							bool bShouldRemoveBlock = true;
+							const FVector2f::FReal EffectiveLengthToRemove = BlockToRemove.GetEffectiveLength();
+
+							// We can remove it, but should we?
+							// If we satisfy constraints using its MinSize, then we don't need to remove it
+							if (!BlockToRemove.bHasVerticalExpansion && BlockToRemove.MinLength != BlockToRemove.ReserveLength)
+							{
+								const float LineLengthAfterRemovalOfEffective = CurrentLine.EffectiveLength - EffectiveLengthToRemove;
+								if (LineLengthAfterRemovalOfEffective + BlockToRemove.MinLength < ParentLineLength)
+								{
+									// We're cool
+									bShouldRemoveBlock = false;
+								}
+							}
+
+							if (bShouldRemoveBlock)
+							{
+								FBlock BlockCopy = BlockToRemove;
+								if (RemoveBlockFromLine(CurrentLine, BlockToRemove))
+								{
+									// Don't include 0-length entries, they'll dictate the line height without having any visual content
+									if (!FMath::IsNearlyZero(EffectiveLengthToRemove))
+									{
+										NextLineBlocks.Emplace(BlockCopy);
+									}
+
+									NextLineBlocks.Emplace(BlockCopy);
+								}
+							}
+						}
+
+						// Remove block from further consideration for this line
+						SortedBlockView.LeftChopInline(1);
+					}
+
+					// No entries were valid, so no blocks were added - early out
+					if (CurrentLine.Blocks.IsEmpty())
+					{
+						// Empty line, early-out
+						break;
+					}
+
+					// Remove first and last slots, if necessary
+					FBlock& FirstBlockInCurrentLine = CurrentLine.Blocks[0];
+					FBlock& LastBlockInCurrentLine = CurrentLine.Blocks.Last();
+
+					int32 FirstSlotIndexInFirstBlock = FirstBlockInCurrentLine.SlotIndices[0];
+					int32 LastSlotIndexInLastBlock = LastBlockInCurrentLine.SlotIndices.Last();
+
+					const FSlotProxy& FirstChildProxy = ChildProxies[FirstSlotIndexInFirstBlock];
+					const FSlotProxy& LastChildProxy = ChildProxies[LastSlotIndexInLastBlock];
+
+					// Check first and last widget in line for exclude flag
+					{
+						const bool bIsSameBlock = FirstBlockInCurrentLine == LastBlockInCurrentLine;
+						const bool bIsSameSlotIndex = FirstSlotIndexInFirstBlock == LastSlotIndexInLastBlock;
+
+						if (FirstChildProxy.bExcludeIfFirstOrLast)
+						{
+							FirstBlockInCurrentLine.SlotIndices.Remove(FirstSlotIndexInFirstBlock);
+						}
+
+						if (!bIsSameSlotIndex && LastChildProxy.bExcludeIfFirstOrLast)
+						{
+							LastBlockInCurrentLine.SlotIndices.Remove(LastSlotIndexInLastBlock);
+						}
+
+						if (FirstChildProxy.bExcludeIfFirstOrLast
+							&& FirstBlockInCurrentLine.SlotIndices.IsEmpty()
+							&& RemoveBlockFromLine(CurrentLine, FirstBlockInCurrentLine))
+						{
+							// Compensate counter
+							NumArrangedBlocks++;
+						}
+
+						// If it's the same block, it's already removed
+						if (LastChildProxy.bExcludeIfFirstOrLast
+							&& (!bIsSameBlock
+								|| (LastBlockInCurrentLine.SlotIndices.IsEmpty() && RemoveBlockFromLine(CurrentLine, LastBlockInCurrentLine))))
+						{
+							// Compensate counter
+							NumArrangedBlocks++;
+						}
+					}
+
+					ensureAlways(!CurrentLine.Blocks.IsEmpty());
+
+					// Clamp height to Min
+					CurrentLine.Height = FMath::Max(MinLineHeight, CurrentLine.Height);
+
+					// Clamp height to Max
+					CurrentLine.Height = FMath::Min(MaxLineHeight, CurrentLine.Height);
+
+					// Constraints satisfied, complete this line...
+					DesiredSize.X = FMath::Max(DesiredSize.X, CurrentLine.MaxLength);
+					DesiredSize.Y = FMath::Max(DesiredSize.Y, CurrentLine.Height + LineOffset);
+
+					LineOffset += CurrentLine.Height + LinePadding;
+					++LineIndex;
+
+					NumArrangedBlocks += CurrentLine.Blocks.Num();
+
+					// ...and setup next
+					CurrentLineBlocks = NextLineBlocks;
+					NextLineBlocks.Reset();
+
+					CurrentLineBlocks.StableSort();
+
+					SortedBlockView = MakeConstArrayView(WrappableSortedBlockView.GetData(), WrappableSortedBlockView.Num());
+
+					// We've got as far as we can, exit
+					if (CurrentLine.MaxLength <= 0)
+					{
+						break;
+					}
+				}
+
+				LastPreferredLineLength = PreferredLineLength;
+				LastParentLineLength = ParentLineLength;
+				LastParentSize = InWidget.GetTickSpaceGeometry().Size;
+				LastDesiredSize = DesiredSize;
+				LastLinePadding = LinePadding;
+
+				Lines.SetNum(LineIndex);
+				BlockLines = Lines;
+			}
+
+			return FVector2D(DesiredSize);
+		}
+
+		void Arrange(const FGeometry& InAllottedGeometry, FArrangedChildren& InOutArrangedChildren)
+		{
+			if (BlockLines.IsEmpty())
+			{
+				return;
+			}
+
+			// The actual size available to the widget
+			const FVector2f ParentSize = InAllottedGeometry.GetLocalSize();
+
+			// These are re-calculated every time a new line is required
+			FVector2f::FReal LineOffset = 0.0f;
+			FVector2f LineOffsetXY = FVector2f::ZeroVector;
+			float LinePadding = LastLinePadding;
+
+			FVector2f LineSize = ParentSize;
+			LineSize.Y = BlockLines[0].Height;
+
+			for (int32 LineIndex = 0; LineIndex < BlockLines.Num(); ++LineIndex)
+			{
+				const FBlockLine& Line = BlockLines[LineIndex];
+
+				float OffsetX = 0.0f;
+
+				TArray<FSlotProxy> ChildProxiesCurrentLine;
+				ChildProxiesCurrentLine.Reserve(ChildProxies.Num());
+
+				for (int32 BlockIndex = 0; BlockIndex < Line.Blocks.Num(); ++BlockIndex)
+				{
+					const FBlock& Block = Line.Blocks[BlockIndex];
+
+					for (int32 ChildIndex = 0; ChildIndex < Block.SlotIndices.Num(); ++ChildIndex)
+					{
+						const int32 ChildIndexInSlotOrder = Block.SlotIndices[ChildIndex];
+
+						FSlotProxy& ChildProxy = ChildProxies[ChildIndexInSlotOrder];
+						ChildProxy.DesiredSize.Y = FMath::Min(ChildProxy.DesiredSize.Y, Line.Height);
+						OffsetX += ChildProxy.DesiredSize.X;
+
+						ChildProxiesCurrentLine.Emplace(ChildProxy);
+					}
+				}
+
+				LineSize.Y = Line.Height;
+				FSlateLayoutTransform LayoutTransform = FSlateLayoutTransform(1.0f, LineOffsetXY);
+				FGeometry LineGeometry = InAllottedGeometry.MakeChild(LineSize, LayoutTransform);
+
+				// Unused, required by ArrangeChildrenInStack
+				constexpr FVector2f::FReal Offset = 0.0f;
+
+				// By provided a non-const ArrayView, the ArrangedSize of the FSlotProxy is written
+				FVector2D ArrangedSize;
+				ArrangeChildrenInStack<Orientation>(GSlateFlowDirection, MakeArrayView(ChildProxiesCurrentLine), LineGeometry, InOutArrangedChildren, Offset, true, ArrangedSize);
+
+				for (int32 BlockIndex = 0; BlockIndex < Line.Blocks.Num(); ++BlockIndex)
+				{
+					const FBlock& Block = Line.Blocks[BlockIndex];
+
+					for (int32 ChildIndex = 0; ChildIndex < Block.SlotIndices.Num(); ++ChildIndex)
+					{
+						const int32 ChildIndexInSlotOrder = Block.SlotIndices[ChildIndex];
+						ChildProxies[ChildIndexInSlotOrder].ArrangedSize = ChildProxiesCurrentLine[ChildIndex].ArrangedSize;
+					}
+				}
+
+				// Zero out LinePadding if we're the last line
+				if (LineIndex == BlockLines.Num() - 1)
+				{
+					LinePadding = 0.0f;
+				}
+				// If we're not the last line, setup geometry for the next
+				else
+				{
+					LineOffset += ArrangedSize.Y + LinePadding;
+					LineOffsetXY = FVector2f(0.0f, LineOffset);
+				}
+			}
+		}
+
+	private:
+		bool UpdateChildProxies(const TPanelChildren<SPrioritizedWrapBox::FSlot>& InChildren)
+		{
+			bool bAnySlotProxyChanged = false;
+
+			const int32 NumChildren = InChildren.Num();
+
+			NumPreferredWrappingChildren = 0;
+			NumParentWrappingChildren = 0;
+
+			// First-time initialization
+			if (ChildProxies.IsEmpty())
+			{
+				ChildProxies.SetNum(NumChildren, EAllowShrinking::No);
+
+				for (int32 ChildIndex = 0; ChildIndex < NumChildren; ++ChildIndex)
+				{
+					ChildProxies[ChildIndex] = FSlotProxy(ChildIndex, InChildren[ChildIndex]);
+
+					if (ChildProxies[ChildIndex].WrapMode == EWrapMode::Preferred)
+					{
+						++NumPreferredWrappingChildren;
+					}
+					else if (ChildProxies[ChildIndex].WrapMode == EWrapMode::Parent)
+					{
+						++NumParentWrappingChildren;
+					}
+				}
+
+				bAnySlotProxyChanged = true;
+			}
+			else
+			{
+				for (int32 ChildIndex = 0; ChildIndex < NumChildren; ++ChildIndex)
+				{
+					if (ChildProxies.IsValidIndex(ChildIndex))
+					{
+						if (ChildProxies[ChildIndex].UpdateFromSlot(ChildIndex, InChildren[ChildIndex]))
+						{
+							bAnySlotProxyChanged = true;
+						}
+
+						if (ChildProxies[ChildIndex].WrapMode == EWrapMode::Preferred)
+						{
+							++NumPreferredWrappingChildren;
+						}
+						else if (ChildProxies[ChildIndex].WrapMode == EWrapMode::Parent)
+						{
+							++NumParentWrappingChildren;
+						}
+					}
+				}
+
+				ChildProxies.SetNum(NumChildren, EAllowShrinking::No);
+			}
+
+			return bAnySlotProxyChanged;
+		}
+
+		void UpdateBlocks(const TPanelChildren<SPrioritizedWrapBox::FSlot>& InChildren, const float InMinLineHeight)
+		{
+			if (ChildProxies.IsEmpty())
+			{
+				return;
+			}
+
+			Blocks.SetNum(InChildren.Num(), EAllowShrinking::No);
+			SortedBlocks.SetNum(InChildren.Num(), EAllowShrinking::No);
+
+			auto GetAreaRetainingWidthForHeight = [](const FVector2f& InOriginalSize, const int32 InNewHeight) -> FVector2f::FReal
+			{
+				const FVector2f::FReal Area = InOriginalSize.X * InOriginalSize.Y;
+				return FMath::CeilToFloat(Area / InNewHeight);
+			};
+
+			int32 NumBlocks = 0;
+			FBlock* CurrentBlock = &Blocks[NumBlocks];
+			CurrentBlock->ResetToDefault();
+			CurrentBlock->SequentialIndex = NumBlocks;
+			CurrentBlock->WrapPriority = ChildProxies[0].WrapPriority;
+
+			NumBlocks++;
+
+			for (int32 ChildIndex = 0; ChildIndex < ChildProxies.Num(); ++ChildIndex)
+			{
+				const SPrioritizedWrapBox::FSlot& Child = InChildren[ChildIndex];
+
+				// Different wrap priority, add a new block
+				if (CurrentBlock->WrapPriority != Child.GetWrapPriority())
+				{
+					CurrentBlock = &Blocks[NumBlocks];
+					CurrentBlock->ResetToDefault();
+					CurrentBlock->SequentialIndex = NumBlocks;
+					CurrentBlock->WrapPriority = Child.GetWrapPriority();
+
+					NumBlocks++;
+				}
+
+				CurrentBlock->SlotIndices.Emplace(ChildIndex);
+
+				TSlotAccessor<SPrioritizedWrapBox::FSlot> SlotAccessor;
+
+				if (SlotAccessor.GetVisibility(Child) == EVisibility::Collapsed)
+				{
+					continue;
+				}
+
+				const FVector2f ChildDesiredSize = SlotAccessor.GetDesiredSize(Child);
+				FVector2f ChildArrangedSize = SlotAccessor.GetArrangedSize(Child);
+
+				FVector2f ChildClampedSize = ChildDesiredSize;
+
+				// Min (if specified) or Desired (if not)
+				FVector2f::FReal MinDesiredSizeX = ChildDesiredSize.X;
+
+				FVector2f::FReal AppendSizeX = 0.0f;
+				FVector2f::FReal PruneSizeX = 0.0f;
+
+				if (const FVector2f::FReal MaxSize = SlotAccessor.GetMaxSize(Child);
+					MaxSize > 0)
+				{
+					ChildClampedSize.X = FMath::Min(MaxSize, ChildClampedSize.X);
+				}
+
+				if (const FVector2f::FReal MinSize = SlotAccessor.GetMinSize(Child);
+					MinSize > 0)
+				{
+					ChildClampedSize.X = FMath::Max(MinSize, ChildClampedSize.X); // Equal or greater to min size
+					MinDesiredSizeX = FMath::Min(MinSize, ChildClampedSize.X); // Min possible valid size
+				}
+
+				// Desired, or Adjusted (if needed)
+				FVector2f::FReal OverflowSizeX = ChildClampedSize.X;
+				FVector2f::FReal OverflowSizeY = ChildClampedSize.Y;
+
+				const float HorizontalPadding = SlotAccessor.GetPadding(Child).GetTotalSpaceAlong<EOrientation::Orient_Horizontal>();
+				const float VerticalPadding = SlotAccessor.GetPadding(Child).GetTotalSpaceAlong<EOrientation::Orient_Vertical>();
+
+				if (Child.GetVerticalOverflowBehavior() == EVerticalOverflowBehavior::ExpandProportional)
+				{
+					CurrentBlock->bCanVerticallyExpand = true;
+
+					// Default to 1.5x the minimum line height unless specified
+					const float VerticalExpansionThreshold = Child.GetVerticalExpansionThreshold().Get(InMinLineHeight * 1.5f);
+					if ((ChildClampedSize.Y + VerticalPadding > VerticalExpansionThreshold)
+						&& ChildArrangedSize.X > 0.0f)
+					{
+						FVector2f OriginalSize = ChildClampedSize;
+
+						// While the threshold is used for the test, the target line height is used here
+						OverflowSizeX = GetAreaRetainingWidthForHeight(OriginalSize, InMinLineHeight + VerticalPadding) + HorizontalPadding * 2;
+						OverflowSizeY = InMinLineHeight;;
+
+						CurrentBlock->bHasVerticalExpansion = true;
+					}
+				}
+
+				const FVector2f::FReal Area = ChildClampedSize.X * ChildClampedSize.Y;
+
+				CurrentBlock->DesiredSize = FVector2f(CurrentBlock->DesiredSize.X + ChildClampedSize.X, FMath::Max(CurrentBlock->DesiredSize.Y, ChildClampedSize.Y));
+				CurrentBlock->AdjustedSize = FVector2f(CurrentBlock->AdjustedSize.X + OverflowSizeX, FMath::Max(CurrentBlock->AdjustedSize.Y, OverflowSizeY));
+				CurrentBlock->ReserveLength += OverflowSizeX;
+				CurrentBlock->MinLength += MinDesiredSizeX;
+				CurrentBlock->bCanWrap = Child.GetAllowWrapping();
+				CurrentBlock->Area += Area;
+
+				FSlotProxy& ChildProxy = ChildProxies[ChildIndex];
+				ChildProxy.DesiredSize.X = CurrentBlock->bHasVerticalExpansion ? OverflowSizeX : ChildClampedSize.X;
+				ChildProxy.DesiredSize.Y = InMinLineHeight;
+			}
+
+			Blocks.SetNum(NumBlocks, EAllowShrinking::No);
+
+			SortedBlocks = Blocks;
+
+			// Needs re-sorting
+			SortedBlocks.StableSort([](const FBlock& A, const FBlock& B)
+			{
+				// Non-wrappable children are always at the start
+				if (!A.bCanWrap && !B.bCanWrap)
+				{
+					return A.SequentialIndex < B.SequentialIndex;
+				}
+
+				// Non-wrappable children come before wrappable ones
+				if (!A.bCanWrap) // B is wrappable
+				{
+					return true;
+				}
+
+				if (!B.bCanWrap) // A is wrappable
+				{
+					return false;
+				}
+
+				// Both are wrappable, with different priorities
+				if (A.WrapPriority != B.WrapPriority)
+				{
+					return A.SequentialIndex < B.SequentialIndex;
+				}
+
+				// Same Priority, maintain original order
+				return A.SequentialIndex < B.SequentialIndex;
+			});
+
+			int32 FirstWrappableBlockIndex = INDEX_NONE;
+			for (int32 BlockIndex = 0; BlockIndex < SortedBlocks.Num(); ++BlockIndex)
+			{
+				const FBlock& SortedBlock = SortedBlocks[BlockIndex];
+				if (SortedBlock.bCanWrap)
+				{
+					FirstWrappableBlockIndex = BlockIndex;
+					break;
+				}
+			}
+
+			if (FirstWrappableBlockIndex == INDEX_NONE)
+			{
+				// Make Empty
+				WrappableSortedBlockView = MakeConstArrayView<FBlock>({ });
+			}
+			else
+			{
+				WrappableSortedBlockView = MakeConstArrayView(SortedBlocks.GetData() + FirstWrappableBlockIndex, SortedBlocks.Num() - FirstWrappableBlockIndex);
+			}
+		}
+
+	private:
+		TArray<FSlotProxy> ChildProxies;
+
+		int32 NumPreferredWrappingChildren = 0;
+		int32 NumParentWrappingChildren = 0;
+
+		int32 LastPreferredLineLength = INDEX_NONE;
+		int32 LastParentLineLength = INDEX_NONE;
+		FVector2f LastParentSize = FVector2f::ZeroVector;
+
+		FVector2f LastDesiredSize;
+		float LastLinePadding;
+
+		TArray<FBlock> Blocks;
+
+		/** This should always be sorted by WrapPriority, then OriginalIndex (if priorities the same), then bCanWrap (such that all widgets that can't wrap are first). */
+		TArray<FBlock> SortedBlocks;
+
+		/** An ArrayView (of SortedChildren) that includes only wrappable children (bCanWrap == true). */
+		TConstArrayView<FBlock> WrappableSortedBlockView;
+
+		TArray<FBlockLine> BlockLines;
+	};
+}
+
+SLATE_IMPLEMENT_WIDGET(SPrioritizedWrapBox)
+void SPrioritizedWrapBox::PrivateRegisterAttributes(FSlateAttributeInitializer& AttributeInitializer)
+{
+	FSlateWidgetSlotAttributeInitializer SlotInitializer = SLATE_ADD_PANELCHILDREN_DEFINITION(AttributeInitializer, Slots);
+	FSlot::RegisterAttributes(SlotInitializer);
+
+	SLATE_ADD_MEMBER_ATTRIBUTE_DEFINITION(AttributeInitializer, PreferredSize, EInvalidateWidgetReason::Layout);
+	SLATE_ADD_MEMBER_ATTRIBUTE_DEFINITION(AttributeInitializer, MinLineHeight, EInvalidateWidgetReason::Layout);
+}
+
+SPrioritizedWrapBox::FSlot::FSlot()
+	: TBasicLayoutWidgetSlot<FSlot>(HAlign_Fill, VAlign_Fill)
+	, TResizingWidgetSlotMixin<FSlot>()
+	, bAllowWrapping(*this, true)
+	, WrapPriority(*this, 0)
+	, WrapMode(*this, UE::Slate::PrioritizedWrapBox::EWrapMode::Preferred)
+	, VerticalOverflowBehavior(UE::Slate::PrioritizedWrapBox::EVerticalOverflowBehavior::Default)
+	, VerticalExpansionThreshold(0.0f)
+	, bExcludeIfFirstOrLast(false)
+{
+}
+
+void SPrioritizedWrapBox::FSlot::Construct(const FChildren& SlotOwner, FSlotArguments&& InArgs)
+{
+	TBasicLayoutWidgetSlot<FSlot>::Construct(SlotOwner, MoveTemp(InArgs));
+	TResizingWidgetSlotMixin<FSlot>::ConstructMixin(SlotOwner, MoveTemp(InArgs));
+
+	if (InArgs._AllowWrapping.IsSet())
+	{
+		bAllowWrapping.Assign(*this, MoveTemp(InArgs._AllowWrapping));
+	}
+
+	if (InArgs._WrapPriority.IsSet())
+	{
+		WrapPriority.Assign(*this, MoveTemp(InArgs._WrapPriority));
+	}
+
+	if (InArgs._WrapMode.IsSet())
+	{
+		WrapMode.Assign(*this, MoveTemp(InArgs._WrapMode));
+	}
+
+	VerticalOverflowBehavior = InArgs._VerticalOverflowBehavior;
+	VerticalExpansionThreshold = InArgs._VerticalExpansionThreshold;
+	bExcludeIfFirstOrLast = InArgs._ExcludeIfFirstOrLast;
+}
+
+void SPrioritizedWrapBox::FSlot::RegisterAttributes(FSlateWidgetSlotAttributeInitializer& AttributeInitializer)
+{
+	TBasicLayoutWidgetSlot<FSlot>::RegisterAttributes(AttributeInitializer);
+	TResizingWidgetSlotMixin<FSlot>::RegisterAttributesMixin(AttributeInitializer);
+
+	SLATE_ADD_SLOT_ATTRIBUTE_DEFINITION_WITH_NAME(FSlot, AttributeInitializer, "Slot.AllowWrapping", bAllowWrapping, EInvalidateWidgetReason::Layout);
+	SLATE_ADD_SLOT_ATTRIBUTE_DEFINITION_WITH_NAME(FSlot, AttributeInitializer, "Slot.WrapPriority", WrapPriority, EInvalidateWidgetReason::Layout);
+	SLATE_ADD_SLOT_ATTRIBUTE_DEFINITION_WITH_NAME(FSlot, AttributeInitializer, "Slot.WrapMode", WrapMode, EInvalidateWidgetReason::Layout);
+}
+
+SPrioritizedWrapBox::SPrioritizedWrapBox()
+	: PreferredSize(*this, 100.0f)
+	, MinLineHeight(*this, TOptional<float>())
+	, LinePadding(0.0f)
+	, bGroupedWrapping(false)
+	, Slots(this, GET_MEMBER_NAME_CHECKED(SPrioritizedWrapBox, Slots))
+	, ChildArranger(MakeUnique<UE::Slate::PrioritizedWrapBox::FChildArranger>())
+{
+}
+
+SPrioritizedWrapBox::~SPrioritizedWrapBox()
+{
+}
+
+SPrioritizedWrapBox::FScopedWidgetSlotArguments SPrioritizedWrapBox::AddSlot()
+{
+	return FScopedWidgetSlotArguments{MakeUnique<FSlot>(), Slots, INDEX_NONE};
+}
+
+int32 SPrioritizedWrapBox::RemoveSlot(const TSharedRef<SWidget>& InSlot)
+{
+	return Slots.Remove(InSlot);
+}
+
+void SPrioritizedWrapBox::Construct(const FArguments& InArgs)
+{
+	PreferredSize.Assign(*this,  InArgs._PreferredSize);
+	MinLineHeight.Assign(*this, InArgs._MinLineHeight);
+	LinePadding = InArgs._LinePadding;;
+	bGroupedWrapping = InArgs._GroupedWrapping;
+
+	Slots.AddSlots(MoveTemp(const_cast<TArray<FSlot::FSlotArguments>&>(InArgs._Slots)));
+
+	SetCanTick(false);
+	SetClipping(EWidgetClipping::ClipToBounds);
+}
+
+FVector2D SPrioritizedWrapBox::ComputeDesiredSize(float LayoutScaleMultiplier) const
+{
+	TRACE_CPUPROFILER_EVENT_SCOPE(SPrioritizedWrapBox::ComputeDesiredSize);
+
+	check(ChildArranger.IsValid());
+
+	return bGroupedWrapping ? ChildArranger->GetDesiredSize<true>(*this, this->Slots) : ChildArranger->GetDesiredSize<false>(*this, this->Slots);
+}
+
+void SPrioritizedWrapBox::OnArrangeChildren(const FGeometry& AllottedGeometry, FArrangedChildren& ArrangedChildren) const
+{
+	TRACE_CPUPROFILER_EVENT_SCOPE(SPrioritizedWrapBox::OnArrangeChildren);
+
+	check(ChildArranger.IsValid());
+
+	return ChildArranger->Arrange(AllottedGeometry, ArrangedChildren);
+}
diff --git a/Engine/Source/Runtime/Slate/Private/Framework/MultiBox/SPrioritizedWrapBox.h b/Engine/Source/Runtime/Slate/Private/Framework/MultiBox/SPrioritizedWrapBox.h
new file mode 100644
--- /dev/null
+++ b/Engine/Source/Runtime/Slate/Private/Framework/MultiBox/SPrioritizedWrapBox.h
@@ -0,0 +1,157 @@
+// Copyright Epic Games, Inc. All Rights Reserved.
+
+#pragma once
+
+#include "CoreTypes.h"
+#include "Framework/MultiBox/MultiBoxDefs.h"
+#include "Framework/SlateDelegates.h"
+#include "Widgets/DeclarativeSyntaxSupport.h"
+#include "Widgets/Layout/SWrapBox.h"
+#include "Widgets/SBoxPanel.h"
+
+class FArrangedChildren;
+class FPaintArgs;
+class FSlateWindowElementList;
+class SComboButton;
+
+DECLARE_DELEGATE_RetVal_OneParam(FMenuEntryResizeParams, FOnGetWidgetResizeParams, const TSharedRef<SWidget>&);
+
+namespace UE::Slate::PrioritizedWrapBox
+{
+	class FChildArranger;
+}; // namespace UE::Slate
+
+/**
+ * An extension of SHorizontalBox that adds wrapping behavior (similar to SWrapBox), primarily for use with MultiBox.
+ * Wrapping candidates are determined by priority, so it doesn't necessarily occur sequentially.
+ * For example, if the middle of 3 widgets has a higher priority, that middle widget will wrap first.
+ */
+class SPrioritizedWrapBox : public SPanel
+{
+	SLATE_DECLARE_WIDGET(SPrioritizedWrapBox, SPanel)
+
+public:
+	class FSlot : public TBasicLayoutWidgetSlot<FSlot>, public TResizingWidgetSlotMixin<FSlot>
+	{
+	public:
+		SLATE_SLOT_BEGIN_ARGS_OneMixin(FSlot, TBasicLayoutWidgetSlot<FSlot>, TResizingWidgetSlotMixin<FSlot>)
+			SLATE_ATTRIBUTE(bool, AllowWrapping) // If set to false, this widget won't be considered for wrapping and will always be in the first line
+			SLATE_ATTRIBUTE(int32, WrapPriority)
+			SLATE_ATTRIBUTE(UE::Slate::PrioritizedWrapBox::EWrapMode, WrapMode)
+			SLATE_ARGUMENT_DEFAULT(UE::Slate::PrioritizedWrapBox::EVerticalOverflowBehavior, VerticalOverflowBehavior) = UE::Slate::PrioritizedWrapBox::EVerticalOverflowBehavior::Default;
+			SLATE_ARGUMENT(TOptional<float>, VerticalExpansionThreshold) // Vertical expansion will occur when the widget's desired size is at or beyond this value
+			SLATE_ARGUMENT_DEFAULT(bool, ExcludeIfFirstOrLast) = false;
+		SLATE_SLOT_END_ARGS()
+
+		FSlot();
+
+	public:
+		void Construct(const FChildren& SlotOwner, FSlotArguments&& InArgs);
+
+		static void RegisterAttributes(FSlateWidgetSlotAttributeInitializer& AttributeInitializer);
+
+		/** If false, this will always remain on the first line. */
+		bool GetAllowWrapping() const { return bAllowWrapping.Get(); }
+
+		/** Higher values will wrap first. Default is 0. */
+		int32 GetWrapPriority() const { return WrapPriority.Get(); }
+
+		/** Higher values will wrap first. Default is 0. */
+		void SetWrapPriority(int32 InPriority) { WrapPriority.Set(*this, InPriority); }
+
+		/** Get the WrapMode, which determines at what reference length the slot wraps. */
+		UE::Slate::PrioritizedWrapBox::EWrapMode GetWrapMode() const { return WrapMode.Get(); }
+
+		/** Get the vertical overflow behavior, specifying how a slot, and it's widget should behave when they need to wrap. */
+		UE::Slate::PrioritizedWrapBox::EVerticalOverflowBehavior GetVerticalOverflowBehavior() const { return VerticalOverflowBehavior; }
+
+		/** If the vertical overflow behavior requires expansion, this optionally specifies the height at which the slot should vertically expand. */
+		const TOptional<float>& GetVerticalExpansionThreshold() const { return VerticalExpansionThreshold; }
+
+		/** If, after wrapping, this slot is the first or last in the line, it will be excluded from the result if this is true. */
+		bool GetExcludeIfFirstOrLast() const { return bExcludeIfFirstOrLast; }
+
+	private:
+		template <typename InObjectType>
+		using TSlateSlotAttribute = ::SlateAttributePrivate::TSlateContainedAttribute<InObjectType, ::SlateAttributePrivate::FSlateAttributeNoInvalidationReason, TSlateAttributeComparePredicate<>>;
+
+		/** If false, this will always remain on the first line. */
+		TSlateSlotAttribute<bool> bAllowWrapping;
+
+		/** Higher values will wrap first. Default is 0. */
+		TSlateSlotAttribute<int32> WrapPriority;
+
+		TSlateSlotAttribute<UE::Slate::PrioritizedWrapBox::EWrapMode> WrapMode;
+
+		UE::Slate::PrioritizedWrapBox::EVerticalOverflowBehavior VerticalOverflowBehavior = UE::Slate::PrioritizedWrapBox::EVerticalOverflowBehavior::Default;
+		TOptional<float> VerticalExpansionThreshold;
+
+		bool bExcludeIfFirstOrLast = false;
+	};
+
+	SLATE_BEGIN_ARGS(SPrioritizedWrapBox)
+		: _PreferredSize(100.0f)
+		, _MinLineHeight(TOptional<float>())
+		, _LinePadding(0.0f)
+		, _GroupedWrapping(false)
+	{ }
+		SLATE_SLOT_ARGUMENT(FSlot, Slots)
+
+		/** The preferred size at which wrapping occurs, applicable only when the Slots WrapMode is Preferred  */
+		SLATE_ATTRIBUTE(float, PreferredSize)
+
+		/** An optional minimum line height, useful to reduce height variance (which changes when wrapping) */
+		SLATE_ATTRIBUTE(TOptional<float>, MinLineHeight)
+
+		/** The padding to add between lines. This only affects spacing between lines, not around each entry */
+		SLATE_ARGUMENT(float, LinePadding)
+
+		/** If true, slots with the same wrap priority are treated as a single monolithic element, rather than per-slot */
+		SLATE_ARGUMENT(bool, GroupedWrapping)
+	SLATE_END_ARGS()
+
+	SPrioritizedWrapBox();
+	virtual ~SPrioritizedWrapBox() override;
+
+	static FSlot::FSlotArguments Slot()
+	{
+		return FSlot::FSlotArguments(MakeUnique<FSlot>());
+	}
+
+	using FScopedWidgetSlotArguments = TPanelChildren<FSlot>::FScopedWidgetSlotArguments;
+	FScopedWidgetSlotArguments AddSlot();
+	int32 RemoveSlot(const TSharedRef<SWidget>& InSlot);
+
+	void Construct(const FArguments& InArgs);
+
+	//~ Begin SWidget
+	virtual FVector2D ComputeDesiredSize(float LayoutScaleMultiplier) const override;
+	//~ End SWidget
+
+	//~ Begin SPanel
+	virtual void OnArrangeChildren( const FGeometry& AllottedGeometry, FArrangedChildren& ArrangedChildren ) const override;
+	virtual FChildren* GetChildren() override { return &Slots; }
+	//~ End SPanel
+
+	float GetPreferredSize() const { return PreferredSize.Get(); }
+	TOptional<float> GetMinLineHeight() const { return MinLineHeight.Get(); }
+	float GetLinePadding() const { return LinePadding; }
+	bool GetUseGroupedWrapping() const { return bGroupedWrapping; }
+
+private:
+	/** How wide or long, dependently of the orientation, this panel should appear to be. Any widgets past this line will be wrapped onto the next line. */
+	TSlateAttribute<float> PreferredSize;
+
+	/** An optional minimum line height, useful to reduce height variance (which changes when wrapping) */
+	TSlateAttribute<TOptional<float>> MinLineHeight;
+
+	/** The padding to add between lines. This only affects spacing between lines, not around each entry */
+	float LinePadding;
+
+	/** If true, slots with the same wrap priority are treated as a single monolithic element, rather than per-slot  */
+	bool bGroupedWrapping;
+
+	TPanelChildren<FSlot> Slots;
+
+	TUniquePtr<class UE::Slate::PrioritizedWrapBox::FChildArranger> ChildArranger;
+};
diff --git a/Engine/Source/Runtime/Slate/Private/Framework/MultiBox/SToolBarButtonBlock.cpp b/Engine/Source/Runtime/Slate/Private/Framework/MultiBox/SToolBarButtonBlock.cpp
--- a/Engine/Source/Runtime/Slate/Private/Framework/MultiBox/SToolBarButtonBlock.cpp
+++ b/Engine/Source/Runtime/Slate/Private/Framework/MultiBox/SToolBarButtonBlock.cpp
@@ -245,6 +245,7 @@ void SToolBarButtonBlock::BuildMultiBlockWidget(const ISlateStyle* StyleSet, con
 
 	IconWidget->AddLayer(TAttribute<const FSlateBrush*>(this, &SToolBarButtonBlock::GetOverlayIconBrush));
 	const bool bIsSlimHorizontalUniformToolBar = MultiBox->GetType() == EMultiBoxType::SlimHorizontalUniformToolBar;
+	const bool bIsSlimWrappingToolBar = MultiBox->GetType() == EMultiBoxType::SlimWrappingToolBar;
 
 	const TSharedRef<STextBlock> TextBlock =
 	SNew(STextBlock)
@@ -274,7 +275,8 @@ void SToolBarButtonBlock::BuildMultiBlockWidget(const ISlateStyle* StyleSet, con
 	// Create the content for our button
 	TSharedRef<SWidget> ButtonContent = SNullWidget::NullWidget;
 	if (MultiBox->GetType() == EMultiBoxType::SlimHorizontalToolBar 
-		|| bIsSlimHorizontalUniformToolBar)
+		|| bIsSlimHorizontalUniformToolBar
+		|| bIsSlimWrappingToolBar)
 	{
 		const FVector2f IconSize = ToolBarStyle.IconSize;
 
diff --git a/Engine/Source/Runtime/Slate/Private/Framework/MultiBox/SToolBarComboButtonBlock.cpp b/Engine/Source/Runtime/Slate/Private/Framework/MultiBox/SToolBarComboButtonBlock.cpp
--- a/Engine/Source/Runtime/Slate/Private/Framework/MultiBox/SToolBarComboButtonBlock.cpp
+++ b/Engine/Source/Runtime/Slate/Private/Framework/MultiBox/SToolBarComboButtonBlock.cpp
@@ -114,7 +114,8 @@ void SToolBarComboButtonBlock::BuildMultiBlockWidget(const ISlateStyle* StyleSet
 			ActualIconWidget->AddLayer(TAttribute<const FSlateBrush*>(this, &SToolBarComboButtonBlock::GetOverlayIconBrush));
 
 			if (MultiBox->GetType() == EMultiBoxType::SlimHorizontalToolBar
-				|| MultiBox->GetType() == EMultiBoxType::SlimHorizontalUniformToolBar)
+				|| MultiBox->GetType() == EMultiBoxType::SlimHorizontalUniformToolBar
+				|| MultiBox->GetType() == EMultiBoxType::SlimWrappingToolBar)
 			{
 				const FVector2f IconSize = ToolBarStyle.IconSize;
 
@@ -158,7 +159,8 @@ void SToolBarComboButtonBlock::BuildMultiBlockWidget(const ISlateStyle* StyleSet
 	// Create the content for our button
 	TSharedRef<SWidget> ButtonContent = SNullWidget::NullWidget;
 	if (MultiBox->GetType() == EMultiBoxType::SlimHorizontalToolBar
-		|| MultiBox->GetType() == EMultiBoxType::SlimHorizontalUniformToolBar)
+		|| MultiBox->GetType() == EMultiBoxType::SlimHorizontalUniformToolBar
+		|| MultiBox->GetType() == EMultiBoxType::SlimWrappingToolBar)
 	{
 		// clang-format off
 		ButtonContent =
diff --git a/Engine/Source/Runtime/Slate/Private/Framework/MultiBox/SToolBarStackButtonBlock.cpp b/Engine/Source/Runtime/Slate/Private/Framework/MultiBox/SToolBarStackButtonBlock.cpp
--- a/Engine/Source/Runtime/Slate/Private/Framework/MultiBox/SToolBarStackButtonBlock.cpp
+++ b/Engine/Source/Runtime/Slate/Private/Framework/MultiBox/SToolBarStackButtonBlock.cpp
@@ -283,7 +283,8 @@ void SToolBarStackButtonBlock::BuildMultiBlockWidget(const ISlateStyle* StyleSet
 
 	TSharedRef<SWidget> ButtonContent = SNullWidget::NullWidget;
 	if (MultiBox->GetType() == EMultiBoxType::SlimHorizontalToolBar
-		|| MultiBox->GetType() == EMultiBoxType::SlimHorizontalUniformToolBar)
+		|| MultiBox->GetType() == EMultiBoxType::SlimHorizontalUniformToolBar
+		|| MultiBox->GetType() == EMultiBoxType::SlimWrappingToolBar)
 	{
 		const FVector2f IconSize = ToolBarStyle.IconSize;
 
diff --git a/Engine/Source/Runtime/Slate/Public/Framework/MultiBox/MultiBox.h b/Engine/Source/Runtime/Slate/Public/Framework/MultiBox/MultiBox.h
--- a/Engine/Source/Runtime/Slate/Public/Framework/MultiBox/MultiBox.h
+++ b/Engine/Source/Runtime/Slate/Public/Framework/MultiBox/MultiBox.h
@@ -872,7 +872,7 @@ public:
 
 private:
 	/** Adds a block Widget to this widget */
-	SLATE_API void AddBlockWidget(const FMultiBlock& Block, TSharedPtr<SHorizontalBox> HorizontalBox, TSharedPtr<SVerticalBox> VerticalBox, EMultiBlockLocation::Type InLocation, bool bSectionContainsIcons, TSharedPtr<const FToolBarComboButtonBlock> OptionsBlock);
+	SLATE_API void AddBlockWidget(const FMultiBlock& Block, TSharedPtr<SHorizontalBox> HorizontalBox, TSharedPtr<SVerticalBox> VerticalBox, TSharedPtr<SPanel> Panel, EMultiBlockLocation::Type InLocation, bool bSectionContainsIcons, TSharedPtr<const FToolBarComboButtonBlock> OptionsBlock);
 
 	/**
 	 * Updates the preview block being dragged.  The drag area is where the users dragged block will be dropped
@@ -959,6 +959,9 @@ private:
 	/** Box panel used for vertically-oriented boxes, e.g., vertical toolbar or menu */
 	TSharedPtr<SVerticalBox> MainVerticalBox;
 
+	/** Box panel used for other boxes, e.g., horizontal wrapping toolbar */
+	TSharedPtr<SPanel> MainPanelBox;
+
 	/** Specialized box widget to handle clipping of toolbars and menubars */
 	TSharedPtr<class SClippingHorizontalBox> ClippedHorizontalBox;
 	
diff --git a/Engine/Source/Runtime/Slate/Public/Framework/MultiBox/MultiBoxBuilder.h b/Engine/Source/Runtime/Slate/Public/Framework/MultiBox/MultiBoxBuilder.h
--- a/Engine/Source/Runtime/Slate/Public/Framework/MultiBox/MultiBoxBuilder.h
+++ b/Engine/Source/Runtime/Slate/Public/Framework/MultiBox/MultiBoxBuilder.h
@@ -349,15 +349,16 @@ public:
 	 */
 	SLATE_API void AddMenuSeparator(FName InExtensionHook = NAME_None, const TAttribute<EVisibility>& InVisibility = TAttribute<EVisibility>());
 	SLATE_API void AddSeparator(FName InExtensionHook = NAME_None, const TAttribute<EVisibility>& InVisibility = TAttribute<EVisibility>());
-	
+
 	/**
 	 * Starts a section on to the extender section hook stack
 	 * 
 	 * @param InExtensionHook	The section hook. Can be NAME_None
 	 * @param InHeadingText		The heading text to use. If none, only a separator is used
 	 * @param InVisibility Optional Visibility Override. Can be used to show/hide this Section dynamically.
+	 * @param InResizeParams	Optional resize parameters and overrides, @see FMenuEntryResizeParams
 	 */
-	SLATE_API void BeginSection( FName InExtensionHook, const TAttribute< FText >& InHeadingText = TAttribute<FText>(), const TAttribute<EVisibility>& InVisibility = TAttribute<EVisibility>() );
+	SLATE_API void BeginSection( FName InExtensionHook, const TAttribute< FText >& InHeadingText = TAttribute<FText>(), const TAttribute<EVisibility>& InVisibility = TAttribute<EVisibility>(), const TAttribute<FMenuEntryResizeParams>& InResizeParams = TAttribute<FMenuEntryResizeParams>() );
 
 	/**
 	 * Ends the current section
@@ -410,6 +411,20 @@ public:
 	 */
 	SLATE_API void AddWidget(const TSharedRef<SWidget>& InWidget, const FText& InLabel, const FMenuEntryStyleParams& InStyleParams, bool bInSearchable = true, const TAttribute<FText>& InToolTipText = FText(), const TAttribute<FSlateIcon>& InIcon = TAttribute<FSlateIcon>(), const TAttribute<EVisibility>& InVisibility = TAttribute<EVisibility>());
 
+	/**
+	 * Adds any widget to the menu
+	 *
+	 * @param	InWidget			The widget that should be shown in the menu
+	 * @param	InLabel				Optional label text to be added to the left of the content
+	 * @param	InStyleParams		Optional style parameters and overrides, @see FMenuEntryStyleParams
+	 * @param	InResizeParams		Optional resize parameters and overrides, @see FMenuEntryResizeParams
+	 * @param	bInSearchable		If true, widget will be searchable (default == true)
+	 * @param	InToolTipText		Optional tooltip text to be added to the widget and label
+	 * @param	InIcon				Optional icon to be shown to the left of the label/content
+	 */
+	SLATE_API void AddWidget(const TSharedRef<SWidget>& InWidget, const FText& InLabel, const FMenuEntryStyleParams& InStyleParams, const TAttribute<FMenuEntryResizeParams>& InResizeParams, bool bInSearchable = true, const TAttribute<FText>& InToolTipText = FText(), const TAttribute<FSlateIcon>& InIcon = TAttribute<FSlateIcon>(), const TAttribute<EVisibility>& InVisibility = TAttribute<EVisibility>());
+
+
 	/**
 	* Adds the widget the multibox will use for searching
 	*/
@@ -455,6 +470,9 @@ private:
 	/** Keep track of current section visibility, in case the section needs to be dynamically hidden */
 	TAttribute<EVisibility> CurrentSectionVisibility;
 
+	/** Resize parameters for the current section */
+	TAttribute<FMenuEntryResizeParams> CurrentSectionResizeParams;
+
 	/** True if there is a pending section that needs to be applied */
 	bool bSectionNeedsToBeApplied;
 
@@ -723,11 +741,12 @@ public:
 	/**
 	 * Starts a section on to the extender section hook stack
 	 *
-	 * @param InExtensionHook	The section hook. Can be NAME_None
+	 * @param InExtensionHook				The section hook. Can be NAME_None
 	 * @param bInSectionShouldHaveSeparator	Whether this section should have a separator to its left (this doesn't apply
+	 * @param InResizeParams				Optional resize parameters to control how this section behaves during toolbar resizing. Note that this currently only supports wrapping options.
 	 * to the first section added, because that never gets a separator).
 	 */
-	SLATE_API void BeginSection(FName InExtensionHook, bool bInSectionShouldHaveSeparator = true);
+	SLATE_API void BeginSection(FName InExtensionHook, bool bInSectionShouldHaveSeparator = true, const TAttribute<FMenuEntryResizeParams>& InResizeParams = TAttribute<FMenuEntryResizeParams>());
 
 	/**
 	 * Ends the current section
@@ -783,6 +802,9 @@ private:
 	/** Current extension hook name for sections to determine where sections begin and end */
 	FName CurrentSectionExtensionHook;
 
+	/** Resize parameters for the current section */
+	TAttribute<FMenuEntryResizeParams> CurrentSectionResizeParams;
+
 	FName CurrentStyleOverride;
 
 	TOptional< EVisibility > LabelVisibility;
diff --git a/Engine/Source/Runtime/Slate/Public/Framework/MultiBox/MultiBoxDefs.h b/Engine/Source/Runtime/Slate/Public/Framework/MultiBox/MultiBoxDefs.h
--- a/Engine/Source/Runtime/Slate/Public/Framework/MultiBox/MultiBoxDefs.h
+++ b/Engine/Source/Runtime/Slate/Public/Framework/MultiBox/MultiBoxDefs.h
@@ -12,6 +12,28 @@
 class SToolTip;
 class SWidget;
 
+/** Types specific to SlimWrappingToolBar (and it's variants). */
+namespace UE::Slate::PrioritizedWrapBox
+{
+	enum class EWrapMode : uint8
+	{
+		Preferred = 0,					// Wraps at the line length specified in PreferredSize
+		Parent = 1,						// Wraps at the wrap box's resulting/actual size
+	};
+
+	/**
+	 * The behavior when a slot's desired height exceeds the allotted/available size (including MaxLineHeight when specified)
+	 * By default this Clips.
+	 */
+	enum class EVerticalOverflowBehavior : uint8
+	{
+		Clip = 0,						// Always Clip to the line height as it's been calculated so far, excluding this slots height from affecting the line height
+		ExpandProportional = 1,			// Uses the desired height of the slot, maintaining the total area after clamping the width to the available space
+
+		Default = Clip,
+	};
+}
+
 /**
  * Types of MultiBoxes
  */
@@ -41,6 +63,9 @@ enum class EMultiBoxType : uint8
 
 	/** A toolbar with horizontally-oriented buttons that tries to arrange all toolbar items uniformly */
 	SlimHorizontalUniformToolBar,
+
+	/** Horizontal tool bar that can (optionally) wrap to subsequent rows */
+	SlimWrappingToolBar,
 };
 
 
@@ -255,4 +280,25 @@ struct FMenuEntryResizeParams
 
 	/** If true (default), this entry will be visible in a toolbar overflow menu. If false, this entry will disappear when it overflows. */
 	TAttribute<bool> VisibleInOverflow;
+
+	struct FWrappingParams
+	{
+		// Keep default values separate from the attributes. This allows the attributes to act like unset optionals by default.
+		static constexpr bool DefaultAllowWrapping = true;
+		static constexpr int32 DefaultPriority = 0;
+		static constexpr UE::Slate::PrioritizedWrapBox::EWrapMode DefaultMode = UE::Slate::PrioritizedWrapBox::EWrapMode::Parent;
+		static constexpr UE::Slate::PrioritizedWrapBox::EVerticalOverflowBehavior DefaultVerticalOverflowBehavior = UE::Slate::PrioritizedWrapBox::EVerticalOverflowBehavior::Default;
+
+		/** If true (default), allow this entry to be wrapped to the next line during resizing of toolbars. If false, this entry will never be wrapped. */
+		TAttribute<bool> Allow;
+
+		/** Override to specify a wrap priority, where a higher priority means the entry will be wrapped to the next line first. */
+		TAttribute<int32> Priority;
+
+		/** Override to specify the wrap mode. By default, this is "Preferred". */
+		TAttribute<UE::Slate::PrioritizedWrapBox::EWrapMode> Mode;
+
+		/** Override to specify the vertical overflow behavior. By default, this is "ClipOrExpandProportional". */
+		TOptional<UE::Slate::PrioritizedWrapBox::EVerticalOverflowBehavior> VerticalOverflowBehavior;
+	} Wrapping;
 };
diff --git a/Engine/Source/Runtime/SlateCore/Public/Styling/ToolBarStyle.h b/Engine/Source/Runtime/SlateCore/Public/Styling/ToolBarStyle.h
--- a/Engine/Source/Runtime/SlateCore/Public/Styling/ToolBarStyle.h
+++ b/Engine/Source/Runtime/SlateCore/Public/Styling/ToolBarStyle.h
@@ -151,6 +151,11 @@ struct FToolBarStyle : public FSlateWidgetStyle
 	bool bAllowWrapButton = true;
 	FToolBarStyle& SetAllowWrapButton(bool bInAllowWrapButton) { bAllowWrapButton = bInAllowWrapButton; return *this; }
 
+	/** Set to false if the toolbar should not wrap (to the next line) by default, but can be overridden per section or entry */
+	UPROPERTY(EditAnywhere, BlueprintReadWrite, Category=Appearance)
+	bool bAllowWrappingDefault = true;
+	FToolBarStyle& SetAllowWrappingDefault(bool bInAllowWrapping) { bAllowWrappingDefault = bInAllowWrapping; return *this; }
+
 	UPROPERTY(EditAnywhere, BlueprintReadWrite, Category = Appearance)
 	FDeprecateSlateVector2D IconSize;
 	FToolBarStyle& SetIconSize(const UE::Slate::FDeprecateVector2DParameter& InIconSize) { IconSize = InIconSize; return *this; }
