[SSpinBox] Extend the SpinBox arguments to include RevertTextOnEscape due to virtual keyoboard workarounds

diff --git a/Engine/Source/Runtime/Slate/Private/Widgets/Input/SSpinBox.cpp b/Engine/Source/Runtime/Slate/Private/Widgets/Input/SSpinBox.cpp
--- a/Engine/Source/Runtime/Slate/Private/Widgets/Input/SSpinBox.cpp
+++ b/Engine/Source/Runtime/Slate/Private/Widgets/Input/SSpinBox.cpp
@@ -171,7 +171,7 @@ void SSpinBox<NumericType>::Construct(const FArguments& InArgs)
 				.Font(InArgs._Font)
 				.SelectAllTextWhenFocused(true)
 				.Text(this, &SSpinBox<NumericType>::GetValueAsText)
-				.RevertTextOnEscape(true)
+				.RevertTextOnEscape(InArgs._RevertTextOnEscape)
 				.OnIsTypedCharValid(this, &SSpinBox<NumericType>::IsCharacterValid)
 				.OnTextChanged(this, &SSpinBox<NumericType>::TextField_OnTextChanged)
 				.OnTextCommitted(this, &SSpinBox<NumericType>::TextField_OnTextCommitted)
@@ -847,6 +847,10 @@ template <typename NumericType>
 void SSpinBox<NumericType>::SetTextSelectAllTextOnCommit(bool bNewValue)
 { EditableText->SetSelectAllTextOnCommit(bNewValue); }
 
+template <typename NumericType>
+void SSpinBox<NumericType>::SetTextRevertTextOnEscape(bool bNewValue)
+{ EditableText->SetRevertTextOnEscape(bNewValue); }
+
 template <typename NumericType>
 void SSpinBox<NumericType>::EnterTextMode()
 {
diff --git a/Engine/Source/Runtime/Slate/Public/Widgets/Input/SSpinBox.h b/Engine/Source/Runtime/Slate/Public/Widgets/Input/SSpinBox.h
--- a/Engine/Source/Runtime/Slate/Public/Widgets/Input/SSpinBox.h
+++ b/Engine/Source/Runtime/Slate/Public/Widgets/Input/SSpinBox.h
@@ -101,6 +101,7 @@ public:
 		, _Justification(ETextJustify::Left)
 		, _KeyboardType(Keyboard_Default)
 		, _PreventThrottling(true)
+		, _RevertTextOnEscape(true)
 	{}
 
 	/** The style used to draw this spinbox */
@@ -180,6 +181,8 @@ public:
 		SLATE_ATTRIBUTE(TSharedPtr< INumericTypeInterface<NumericType> >, TypeInterface)
 		/** If refresh requests for the viewport should happen for all value changes **/
 		SLATE_ARGUMENT(bool, PreventThrottling)
+		/** If the text should be reverted when pressing the escape key **/
+		SLATE_ARGUMENT(bool, RevertTextOnEscape)
 		/** Menu extender for the right-click context menu */
 		SLATE_EVENT(FMenuExtensionDelegate, ContextMenuExtender)
 
@@ -321,6 +324,7 @@ public:
 	SLATE_API void SetTextBlockFont(FSlateFontInfo InFont);
 	SLATE_API void SetTextJustification(ETextJustify::Type InJustification);
 	SLATE_API void SetTextClearKeyboardFocusOnCommit(bool bNewValue);
+	SLATE_API void SetTextRevertTextOnEscape(bool bNewValue);
 	SLATE_API void SetTextSelectAllTextOnCommit(bool bNewValue);
 
 	/** Reset the cached string. Typically used when the value is the same but the display format changed (through the callback). */
