[Slate] Fixed duplicate entries in SPrioritzedWrapBox

diff --git a/Engine/Source/Runtime/Slate/Private/Framework/MultiBox/SPrioritizedWrapBox.cpp b/Engine/Source/Runtime/Slate/Private/Framework/MultiBox/SPrioritizedWrapBox.cpp
--- a/Engine/Source/Runtime/Slate/Private/Framework/MultiBox/SPrioritizedWrapBox.cpp
+++ b/Engine/Source/Runtime/Slate/Private/Framework/MultiBox/SPrioritizedWrapBox.cpp
@@ -363,8 +363,6 @@ namespace UE::Slate::PrioritizedWrapBox
 									{
 										NextLineBlocks.Emplace(BlockCopy);
 									}
-
-									NextLineBlocks.Emplace(BlockCopy);
 								}
 							}
 						}
