Hotfix an issue in 41783588 where Block resize params could be stomped by section resize params.

#rnx

diff --git a/Engine/Source/Runtime/Slate/Private/Framework/MultiBox/MultiBoxBuilder.cpp b/Engine/Source/Runtime/Slate/Private/Framework/MultiBox/MultiBoxBuilder.cpp
--- a/Engine/Source/Runtime/Slate/Private/Framework/MultiBox/MultiBoxBuilder.cpp
+++ b/Engine/Source/Runtime/Slate/Private/Framework/MultiBox/MultiBoxBuilder.cpp
@@ -78,7 +78,7 @@ namespace UE::Slate::Private
 				BlockResizeParamsAttribute = NewResizeParams;
 			}
 
-			InMultiBlock->SetResizeParams(InSectionResizeParams);
+			InMultiBlock->SetResizeParams(BlockResizeParamsAttribute);
 		}
 	}
 }
