Expose the raw resize parameters attribute in the FMultiBlock, enabling more correct section overrides in FMultiBoxBuilder.

diff --git a/Engine/Source/Runtime/Slate/Private/Framework/MultiBox/MultiBox.cpp b/Engine/Source/Runtime/Slate/Private/Framework/MultiBox/MultiBox.cpp
--- a/Engine/Source/Runtime/Slate/Private/Framework/MultiBox/MultiBox.cpp
+++ b/Engine/Source/Runtime/Slate/Private/Framework/MultiBox/MultiBox.cpp
@@ -280,9 +280,9 @@ void FMultiBlock::SetVisibilityOverride(TAttribute<EVisibility> InVisibilityOver
 	VisibilityOverride = InVisibilityOverride;
 }
 
-FMenuEntryResizeParams FMultiBlock::GetResizeParams() const
+const TAttribute<FMenuEntryResizeParams>& FMultiBlock::GetResizeParams() const
 {
-	return ResizeParams.Get();
+	return ResizeParams;
 }
 
 void FMultiBlock::SetResizeParams(const TAttribute<FMenuEntryResizeParams>& InResizeParams)
@@ -902,7 +902,7 @@ void SMultiBoxWidget::AddBlockWidget(const FMultiBlock& Block, TSharedPtr<SHoriz
 					break;
 				}
 
-				const FMenuEntryResizeParams ResizeParams = Block.GetResizeParams();
+				const FMenuEntryResizeParams ResizeParams = Block.GetResizeParams().Get();
 
 				TAttribute<bool> AllowWrapping = ResizeParams.Wrapping.Allow;
 				if (!AllowWrapping.IsSet()
@@ -1185,7 +1185,7 @@ void SMultiBoxWidget::BuildMultiBoxWidget()
 							{
 								if (InWidget == Pair.Key)
 								{
-									return Pair.Value->GetResizeParams();
+									return Pair.Value->GetResizeParams().Get();
 								}
 							}
 
diff --git a/Engine/Source/Runtime/Slate/Private/Framework/MultiBox/MultiBoxBuilder.cpp b/Engine/Source/Runtime/Slate/Private/Framework/MultiBox/MultiBoxBuilder.cpp
--- a/Engine/Source/Runtime/Slate/Private/Framework/MultiBox/MultiBoxBuilder.cpp
+++ b/Engine/Source/Runtime/Slate/Private/Framework/MultiBox/MultiBoxBuilder.cpp
@@ -51,7 +51,7 @@ namespace UE::Slate::Private
 			{
 				BlockResizeParamsAttribute = InSectionResizeParams;
 			}
-			else if (BlockResizeParamsAttribute.IsSet() && !BlockResizeParamsAttribute.IsBound())
+			else if (!BlockResizeParamsAttribute.IsBound())
 			{
 				const FMenuEntryResizeParams& SectionResizeParams = InSectionResizeParams.Get();
 
@@ -61,7 +61,7 @@ namespace UE::Slate::Private
 
 				BlockResizeParamsAttribute = BlockResizeParams;
 			}
-			else if (BlockResizeParamsAttribute.IsBound())
+			else
 			{
 				// Intercept and inject wrapping properties
 				TAttribute<FMenuEntryResizeParams> NewResizeParams(
diff --git a/Engine/Source/Runtime/Slate/Public/Framework/MultiBox/MultiBox.h b/Engine/Source/Runtime/Slate/Public/Framework/MultiBox/MultiBox.h
--- a/Engine/Source/Runtime/Slate/Public/Framework/MultiBox/MultiBox.h
+++ b/Engine/Source/Runtime/Slate/Public/Framework/MultiBox/MultiBox.h
@@ -252,7 +252,7 @@ public:
 	void SetVisibilityOverride(TAttribute<EVisibility> InVisibilityOverride);
 
 	/** Get the resize parameters which determine how this block behaves during toolbar resizing.*/
-	FMenuEntryResizeParams GetResizeParams() const;
+	const TAttribute<FMenuEntryResizeParams>& GetResizeParams() const;
 
 	/** Set the resize parameters which determine how this block behaves during toolbar resizing.*/
 	void SetResizeParams(const TAttribute<FMenuEntryResizeParams>& InResizeParams);
