[Slate/MainFrame/LevelEditor]
* Reverted Commands change to resolve linker error. Will need to revisit this

diff --git a/Engine/Source/Editor/LevelEditor/Private/LevelViewportActions.cpp b/Engine/Source/Editor/LevelEditor/Private/LevelViewportActions.cpp
--- a/Engine/Source/Editor/LevelEditor/Private/LevelViewportActions.cpp
+++ b/Engine/Source/Editor/LevelEditor/Private/LevelViewportActions.cpp
@@ -21,8 +21,6 @@
 
 #define LOCTEXT_NAMESPACE "LevelViewportActions"
 
-UE_DEFINE_COMMANDS_TLS(FLevelViewportCommands)
-
 FLevelViewportCommands::FOnNewStatCommandAdded FLevelViewportCommands::NewStatCommandDelegate;
 
 PRAGMA_DISABLE_DEPRECATION_WARNINGS
diff --git a/Engine/Source/Editor/LevelEditor/Public/LevelViewportActions.h b/Engine/Source/Editor/LevelEditor/Public/LevelViewportActions.h
--- a/Engine/Source/Editor/LevelEditor/Public/LevelViewportActions.h
+++ b/Engine/Source/Editor/LevelEditor/Public/LevelViewportActions.h
@@ -35,8 +35,6 @@ namespace LevelViewportConfigurationNames
 	static FName OnePane("OnePane");
 }
 
-UE_DECLARE_COMMANDS_TLS(class FLevelViewportCommands, LEVELEDITOR_API)
-
 /**
  * Class containing commands for level viewport actions
  */
diff --git a/Engine/Source/Editor/MainFrame/MainFrame.Build.cs b/Engine/Source/Editor/MainFrame/MainFrame.Build.cs
--- a/Engine/Source/Editor/MainFrame/MainFrame.Build.cs
+++ b/Engine/Source/Editor/MainFrame/MainFrame.Build.cs
@@ -22,7 +22,6 @@ public class MainFrame : ModuleRules
 				"Engine",
 				"EngineSettings",
 				"InputCore",
-				"LevelEditor",
 				"RHI",
 				"RenderCore",
 				"Slate",
@@ -54,6 +53,7 @@ public class MainFrame : ModuleRules
 				"LauncherPlatform",
 				"GameProjectGeneration",
 				"ProjectTargetPlatformEditor",
+				"LevelEditor",
 				"Settings",
 				"SourceCodeAccess",
 				"HotReload",
@@ -65,6 +65,7 @@ public class MainFrame : ModuleRules
 				"AssetTools",
 				"GameProjectGeneration",
 				"ProjectTargetPlatformEditor",
+				"LevelEditor",
 				"SourceCodeAccess",
 				"HotReload",
 				"LocalizationDashboard", // not required but causes circular depends issues on Linux/Mac
diff --git a/Engine/Source/Runtime/Slate/Private/Framework/Commands/GenericCommands.cpp b/Engine/Source/Runtime/Slate/Private/Framework/Commands/GenericCommands.cpp
--- a/Engine/Source/Runtime/Slate/Private/Framework/Commands/GenericCommands.cpp
+++ b/Engine/Source/Runtime/Slate/Private/Framework/Commands/GenericCommands.cpp
@@ -4,8 +4,6 @@
 
 #define LOCTEXT_NAMESPACE "GenericCommands"
 
-UE_DEFINE_COMMANDS_TLS(FGenericCommands)
-
 void FGenericCommands::RegisterCommands()
 {
 	UI_COMMAND(Cut, "Cut", "Cut selection", EUserInterfaceActionType::Button, FInputChord(EModifierKey::Control, EKeys::X))
diff --git a/Engine/Source/Runtime/Slate/Public/Framework/Commands/Commands.h b/Engine/Source/Runtime/Slate/Public/Framework/Commands/Commands.h
--- a/Engine/Source/Runtime/Slate/Public/Framework/Commands/Commands.h
+++ b/Engine/Source/Runtime/Slate/Public/Framework/Commands/Commands.h
@@ -24,9 +24,6 @@ SLATE_API void MakeUICommand_InternalUseOnly(FBindingContext* This, TSharedPtr<
 
 #undef LOC_DEFINE_REGION
 
-#define UE_DECLARE_COMMANDS_TLS(Type, Api) template<> Api TWeakPtr<Type>& TCommands<Type>::GetInstance();
-#define UE_DEFINE_COMMANDS_TLS(Type) template<> TWeakPtr<Type>& TCommands<Type>::GetInstance() { static TWeakPtr<Type> Inst; return Inst; }
-
 /** A base class for a set of commands. Inherit from it to make a set of commands. See MainFrameActions for an example. */
 template<typename CommandContextType>
 class TCommands : public FBindingContext
@@ -35,7 +32,7 @@ public:
 	
 
 	/** Use this method to register commands. Usually done in StartupModule(). */
-	inline static void Register()
+	FORCENOINLINE static void Register()
 	{
 		if ( !Instance.IsValid() )
 		{
@@ -63,13 +60,13 @@ public:
 		}
 	}
 
-	inline static bool IsRegistered()
+	FORCENOINLINE static bool IsRegistered()
 	{
 		return Instance.IsValid();
 	}
 
 	/** Get the singleton instance of this set of commands. */
-	inline static const CommandContextType& Get()
+	FORCENOINLINE static const CommandContextType& Get()
 	{
 		// If you get a crash here, then either your commands object wasn't registered, or it was not kept
 		//  alive. The latter could happen if your RegisterCommands() function does not actually register 
@@ -78,7 +75,7 @@ public:
 	}
 
 	/** Use this method to clean up any resources used by the command set. Usually done in ShutdownModule() */
-	inline static void Unregister()
+	FORCENOINLINE static void Unregister()
 	{
 		// The instance may not be valid if it was never used.
 		if( Instance.IsValid() )
@@ -95,7 +92,7 @@ public:
 	}
 
 	/** Get the BindingContext for this set of commands. */
-	inline static const FBindingContext& GetContext()
+	FORCENOINLINE static const FBindingContext& GetContext()
 	{
 		check( Instance.IsValid() );
 		return *(Instance.Pin());
@@ -114,25 +111,16 @@ protected:
 
 	/** A static instance of the command set. */
 #if PLATFORM_UNIX || PLATFORM_APPLE
-	static SLATE_API TWeakPtr<CommandContextType>& GetInstance()
-	{
-		static TWeakPtr<CommandContextType> Inst;
-		return Inst;
-	}
+	static SLATE_API TWeakPtr< CommandContextType > Instance;
 #else
-	static TWeakPtr<CommandContextType>& GetInstance()
-	{
-		static TWeakPtr<CommandContextType> Inst;
-		return Inst;
-	}
+	static TWeakPtr< CommandContextType > Instance;
 #endif
 
-	static TWeakPtr<CommandContextType>& Instance;
-
 	/** Pure virtual to override; describe and instantiate the commands in here by using the UI COMMAND macro. */
 	virtual void RegisterCommands() = 0;
 
 };
 
 template<typename CommandContextType>
-TWeakPtr<CommandContextType>& TCommands<CommandContextType>::Instance = TCommands<CommandContextType>::GetInstance();
+TWeakPtr< CommandContextType > TCommands<CommandContextType>::Instance = NULL;
+
diff --git a/Engine/Source/Runtime/Slate/Public/Framework/Commands/GenericCommands.h b/Engine/Source/Runtime/Slate/Public/Framework/Commands/GenericCommands.h
--- a/Engine/Source/Runtime/Slate/Public/Framework/Commands/GenericCommands.h
+++ b/Engine/Source/Runtime/Slate/Public/Framework/Commands/GenericCommands.h
@@ -6,8 +6,6 @@
 #include "Styling/CoreStyle.h"
 #include "Framework/Commands/Commands.h"
 
-UE_DECLARE_COMMANDS_TLS(class FGenericCommands, SLATE_API)
-
 class FGenericCommands : public TCommands<FGenericCommands>
 {
 public:
