[Slate]
* readded dll export on FGenericCommands

diff --git a/Engine/Source/Runtime/Slate/Public/Framework/Commands/GenericCommands.h b/Engine/Source/Runtime/Slate/Public/Framework/Commands/GenericCommands.h
--- a/Engine/Source/Runtime/Slate/Public/Framework/Commands/GenericCommands.h
+++ b/Engine/Source/Runtime/Slate/Public/Framework/Commands/GenericCommands.h
@@ -6,7 +6,7 @@
 #include "Styling/CoreStyle.h"
 #include "Framework/Commands/Commands.h"
 
-class FGenericCommands : public TCommands<FGenericCommands>
+class SLATE_API FGenericCommands : public TCommands<FGenericCommands>
 {
 public:
 	
@@ -19,7 +19,7 @@ public:
 	{
 	}
 
-	SLATE_API virtual void RegisterCommands() override;	
+	virtual void RegisterCommands() override;	
 
 	TSharedPtr<FUICommandInfo> Cut;
 	TSharedPtr<FUICommandInfo> Copy;
