[Slate]
* Removed static TCommands::Instance and replaced it with static TCommands::GetInstance() with static Instance inside. This makes it possible to specialize GetInstance to move Instance out of the header. It will also work with current "hack" if using dllstorage on entire subclass
* Added UE_DECLARE_COMMANDS_TLS/UE_DEFINE_COMMANDS_TLS that should be used to instantiate GetInstance to make sure same instance is always returned regardless of dll.

Using the macros is the right way to make sure there is only one instace, current code relies on the entire subclass having dllstorage which somehow makes the linker use the imported version of the subclass to access instance which according to microsoft "works by accident".

diff --git a/Engine/Plugins/Cameras/CameraShakePreviewer/Source/CameraShakePreviewer/Private/CameraShakePreviewerModule.cpp b/Engine/Plugins/Cameras/CameraShakePreviewer/Source/CameraShakePreviewer/Private/CameraShakePreviewerModule.cpp
--- a/Engine/Plugins/Cameras/CameraShakePreviewer/Source/CameraShakePreviewer/Private/CameraShakePreviewerModule.cpp
+++ b/Engine/Plugins/Cameras/CameraShakePreviewer/Source/CameraShakePreviewer/Private/CameraShakePreviewerModule.cpp
@@ -77,11 +77,6 @@ public:
 	{
 	}
 
-	static FCameraShakePreviewerCommands& Get()
-	{
-		return *(Instance.Pin());
-	}
-
 	virtual void RegisterCommands() override
 	{
 		UI_COMMAND(ToggleCameraShakesPreview, "Allow Camera Shakes", "If enabled, allows the camera shakes previewer panel to apply shakes to this viewport", EUserInterfaceActionType::ToggleButton, FInputChord());
diff --git a/Engine/Plugins/Editor/ModelingToolsEditorMode/Source/ModelingToolsEditorMode/Private/ModelingToolsManagerActions.cpp b/Engine/Plugins/Editor/ModelingToolsEditorMode/Source/ModelingToolsEditorMode/Private/ModelingToolsManagerActions.cpp
--- a/Engine/Plugins/Editor/ModelingToolsEditorMode/Source/ModelingToolsEditorMode/Private/ModelingToolsManagerActions.cpp
+++ b/Engine/Plugins/Editor/ModelingToolsEditorMode/Source/ModelingToolsEditorMode/Private/ModelingToolsManagerActions.cpp
@@ -311,9 +311,9 @@ TSharedPtr<FUICommandInfo> FModelingToolsManagerCommands::RegisterExtensionPalet
 	const FText& Tooltip,
 	const FSlateIcon& Icon)
 {
-	if (Instance.IsValid())
+	if (IsRegistered())
 	{
-		TSharedPtr<FModelingToolsManagerCommands> Commands = Instance.Pin();
+		TSharedPtr<FModelingToolsManagerCommands> Commands = GetInstance().Pin();
 
 		for (FDynamicExtensionCommand& ExtensionCommand : Commands->ExtensionPaletteCommands)
 		{
diff --git a/Engine/Plugins/Editor/ScriptableToolsEditorMode/Source/ScriptableToolsEditorMode/Private/ScriptableToolsEditorModeManagerCommands.cpp b/Engine/Plugins/Editor/ScriptableToolsEditorMode/Source/ScriptableToolsEditorMode/Private/ScriptableToolsEditorModeManagerCommands.cpp
--- a/Engine/Plugins/Editor/ScriptableToolsEditorMode/Source/ScriptableToolsEditorMode/Private/ScriptableToolsEditorModeManagerCommands.cpp
+++ b/Engine/Plugins/Editor/ScriptableToolsEditorMode/Source/ScriptableToolsEditorMode/Private/ScriptableToolsEditorModeManagerCommands.cpp
@@ -37,9 +37,9 @@ TSharedPtr<FUICommandInfo> FScriptableToolsEditorModeManagerCommands::FindToolBy
 
 void FScriptableToolsEditorModeManagerCommands::NotifyCommandsChanged() const
 {
-	if (TSharedPtr<FScriptableToolsEditorModeManagerCommands> Commands = Instance.Pin())
+	if (IsRegistered())
 	{
-		CommandsChanged.Broadcast(*Commands);
+		CommandsChanged.Broadcast(Get());
 	}
 }
 
diff --git a/Engine/Source/Developer/WidgetRegistration/Private/Inputs/BuilderCommandCreationManager.cpp b/Engine/Source/Developer/WidgetRegistration/Private/Inputs/BuilderCommandCreationManager.cpp
--- a/Engine/Source/Developer/WidgetRegistration/Private/Inputs/BuilderCommandCreationManager.cpp
+++ b/Engine/Source/Developer/WidgetRegistration/Private/Inputs/BuilderCommandCreationManager.cpp
@@ -15,11 +15,6 @@ FBuilderCommandCreationManager::FBuilderCommandCreationManager()
 {
 }
 
-const FBuilderCommandCreationManager& FBuilderCommandCreationManager::Get()
-{
-	return TCommands<FBuilderCommandCreationManager>::Get();
-}
-
 void FBuilderCommandCreationManager::RegisterCommands()
 {
 	/*
@@ -34,9 +29,9 @@ void FBuilderCommandCreationManager::RegisterCommands()
 
 void FBuilderCommandCreationManager::RegisterCommandForBuilder( UE::DisplayBuilders::FBuilderInput& Input ) const
 {
-	if (Instance.IsValid() && !Input.IsNameNone() )
+	if (IsRegistered() && !Input.IsNameNone() )
 	{
-		const TSharedPtr<FBuilderCommandCreationManager> Commands = Instance.Pin();
+		const TSharedPtr<FBuilderCommandCreationManager> Commands = GetInstance().Pin();
 
 		TSharedPtr<FUICommandInfo> NewCommandInfo;
 
@@ -59,7 +54,7 @@ void FBuilderCommandCreationManager::UnregisterCommandForBuilder( UE::DisplayBui
 {
 	if ( Input.UICommandInfo.IsValid() )
 	{
-		FUICommandInfo::UnregisterCommandInfo( Instance.Pin().ToSharedRef(), Input.UICommandInfo.ToSharedRef() );		
+		FUICommandInfo::UnregisterCommandInfo( GetInstance().Pin().ToSharedRef(), Input.UICommandInfo.ToSharedRef() );		
 	}
 }
 
diff --git a/Engine/Source/Developer/WidgetRegistration/Private/Inputs/BuilderCommandCreationManager.h b/Engine/Source/Developer/WidgetRegistration/Private/Inputs/BuilderCommandCreationManager.h
--- a/Engine/Source/Developer/WidgetRegistration/Private/Inputs/BuilderCommandCreationManager.h
+++ b/Engine/Source/Developer/WidgetRegistration/Private/Inputs/BuilderCommandCreationManager.h
@@ -40,9 +40,4 @@ public:
     * @param OutBuilderInput the FBuilderInput for which to unregister and tear down TSharedPtr<FUICommandInfo>
     */
 	void UnregisterCommandForBuilder( UE::DisplayBuilders::FBuilderInput& OutBuilderInput ) const;
-
-	/**
-	 * @return the singleton of the FBuilderCommandCreationManager
-	 */
-	static const FBuilderCommandCreationManager& Get();
 };
\ No newline at end of file
diff --git a/Engine/Source/Editor/LevelEditor/Private/LevelViewportActions.cpp b/Engine/Source/Editor/LevelEditor/Private/LevelViewportActions.cpp
--- a/Engine/Source/Editor/LevelEditor/Private/LevelViewportActions.cpp
+++ b/Engine/Source/Editor/LevelEditor/Private/LevelViewportActions.cpp
@@ -21,6 +21,8 @@
 
 #define LOCTEXT_NAMESPACE "LevelViewportActions"
 
+UE_DEFINE_COMMANDS_TLS(FLevelViewportCommands)
+
 FLevelViewportCommands::FOnNewStatCommandAdded FLevelViewportCommands::NewStatCommandDelegate;
 
 PRAGMA_DISABLE_DEPRECATION_WARNINGS
diff --git a/Engine/Source/Editor/LevelEditor/Public/LevelViewportActions.h b/Engine/Source/Editor/LevelEditor/Public/LevelViewportActions.h
--- a/Engine/Source/Editor/LevelEditor/Public/LevelViewportActions.h
+++ b/Engine/Source/Editor/LevelEditor/Public/LevelViewportActions.h
@@ -35,10 +35,12 @@ namespace LevelViewportConfigurationNames
 	static FName OnePane("OnePane");
 }
 
+UE_DECLARE_COMMANDS_TLS(class FLevelViewportCommands, LEVELEDITOR_API)
+
 /**
  * Class containing commands for level viewport actions
  */
-class LEVELEDITOR_API FLevelViewportCommands : public TCommands<FLevelViewportCommands>
+class FLevelViewportCommands : public TCommands<FLevelViewportCommands>
 {
 public:
 
@@ -63,12 +65,6 @@ PRAGMA_ENABLE_DEPRECATION_WARNINGS
 
 	virtual ~FLevelViewportCommands();
 
-	/** Get the singleton instance of this set of commands. */
-	FORCENOINLINE static FLevelViewportCommands& Get()
-	{
-		return *(Instance.Pin());
-	}
-
 	struct FShowMenuCommand
 	{
 		TSharedPtr<FUICommandInfo> ShowMenuItem;
diff --git a/Engine/Source/Runtime/Slate/Private/Framework/Commands/GenericCommands.cpp b/Engine/Source/Runtime/Slate/Private/Framework/Commands/GenericCommands.cpp
--- a/Engine/Source/Runtime/Slate/Private/Framework/Commands/GenericCommands.cpp
+++ b/Engine/Source/Runtime/Slate/Private/Framework/Commands/GenericCommands.cpp
@@ -4,6 +4,8 @@
 
 #define LOCTEXT_NAMESPACE "GenericCommands"
 
+UE_DEFINE_COMMANDS_TLS(FGenericCommands)
+
 void FGenericCommands::RegisterCommands()
 {
 	UI_COMMAND(Cut, "Cut", "Cut selection", EUserInterfaceActionType::Button, FInputChord(EModifierKey::Control, EKeys::X))
diff --git a/Engine/Source/Runtime/Slate/Public/Framework/Commands/Commands.h b/Engine/Source/Runtime/Slate/Public/Framework/Commands/Commands.h
--- a/Engine/Source/Runtime/Slate/Public/Framework/Commands/Commands.h
+++ b/Engine/Source/Runtime/Slate/Public/Framework/Commands/Commands.h
@@ -24,6 +24,9 @@ SLATE_API void MakeUICommand_InternalUseOnly(FBindingContext* This, TSharedPtr<
 
 #undef LOC_DEFINE_REGION
 
+#define UE_DECLARE_COMMANDS_TLS(Type, Api) template<> Api TWeakPtr<Type>& TCommands<Type>::GetInstance();
+#define UE_DEFINE_COMMANDS_TLS(Type) template<> TWeakPtr<Type>& TCommands<Type>::GetInstance() { static TWeakPtr<Type> Inst; return Inst; }
+
 /** A base class for a set of commands. Inherit from it to make a set of commands. See MainFrameActions for an example. */
 template<typename CommandContextType>
 class TCommands : public FBindingContext
@@ -32,8 +35,9 @@ public:
 	
 
 	/** Use this method to register commands. Usually done in StartupModule(). */
-	FORCENOINLINE static void Register()
+	inline static void Register()
 	{
+		TWeakPtr<CommandContextType>& Instance = GetInstance();
 		if ( !Instance.IsValid() )
 		{
 			// We store the singleton instances in the FInputBindingManager in order to prevent
@@ -60,23 +64,24 @@ public:
 		}
 	}
 
-	FORCENOINLINE static bool IsRegistered()
+	inline static bool IsRegistered()
 	{
-		return Instance.IsValid();
+		return GetInstance().IsValid();
 	}
 
 	/** Get the singleton instance of this set of commands. */
-	FORCENOINLINE static const CommandContextType& Get()
+	inline static CommandContextType& Get()
 	{
 		// If you get a crash here, then either your commands object wasn't registered, or it was not kept
 		//  alive. The latter could happen if your RegisterCommands() function does not actually register 
 		//  any FUICommandInfo objects that use your TCommands object, so nothing keeps it alive.
-		return *( Instance.Pin() );
+		return *( GetInstance().Pin() );
 	}
 
 	/** Use this method to clean up any resources used by the command set. Usually done in ShutdownModule() */
-	FORCENOINLINE static void Unregister()
+	inline static void Unregister()
 	{
+		TWeakPtr<CommandContextType>& Instance = GetInstance();
 		// The instance may not be valid if it was never used.
 		if( Instance.IsValid() )
 		{
@@ -92,8 +97,9 @@ public:
 	}
 
 	/** Get the BindingContext for this set of commands. */
-	FORCENOINLINE static const FBindingContext& GetContext()
+	inline static const FBindingContext& GetContext()
 	{
+		TWeakPtr<CommandContextType>& Instance = GetInstance();
 		check( Instance.IsValid() );
 		return *(Instance.Pin());
 	}
@@ -111,16 +117,20 @@ protected:
 
 	/** A static instance of the command set. */
 #if PLATFORM_UNIX || PLATFORM_APPLE
-	static SLATE_API TWeakPtr< CommandContextType > Instance;
+	static SLATE_API TWeakPtr<CommandContextType>& GetInstance()
+	{
+		static TWeakPtr<CommandContextType> Inst;
+		return Inst;
+	}
 #else
-	static TWeakPtr< CommandContextType > Instance;
+	static TWeakPtr<CommandContextType>& GetInstance()
+	{
+		static TWeakPtr<CommandContextType> Inst;
+		return Inst;
+	}
 #endif
 
 	/** Pure virtual to override; describe and instantiate the commands in here by using the UI COMMAND macro. */
 	virtual void RegisterCommands() = 0;
 
 };
-
-template<typename CommandContextType>
-TWeakPtr< CommandContextType > TCommands<CommandContextType>::Instance = NULL;
-
diff --git a/Engine/Source/Runtime/Slate/Public/Framework/Commands/GenericCommands.h b/Engine/Source/Runtime/Slate/Public/Framework/Commands/GenericCommands.h
--- a/Engine/Source/Runtime/Slate/Public/Framework/Commands/GenericCommands.h
+++ b/Engine/Source/Runtime/Slate/Public/Framework/Commands/GenericCommands.h
@@ -6,7 +6,9 @@
 #include "Styling/CoreStyle.h"
 #include "Framework/Commands/Commands.h"
 
-class SLATE_API FGenericCommands : public TCommands<FGenericCommands>
+UE_DECLARE_COMMANDS_TLS(class FGenericCommands, SLATE_API)
+
+class FGenericCommands : public TCommands<FGenericCommands>
 {
 public:
 	
@@ -19,7 +21,7 @@ public:
 	{
 	}
 
-	virtual void RegisterCommands() override;	
+	SLATE_API virtual void RegisterCommands() override;	
 
 	TSharedPtr<FUICommandInfo> Cut;
 	TSharedPtr<FUICommandInfo> Copy;
