[Slate] Added remaining wrapping attributes to MultiBox ResizeParams injection

diff --git a/Engine/Source/Runtime/Slate/Private/Framework/MultiBox/MultiBoxBuilder.cpp b/Engine/Source/Runtime/Slate/Private/Framework/MultiBox/MultiBoxBuilder.cpp
--- a/Engine/Source/Runtime/Slate/Private/Framework/MultiBox/MultiBoxBuilder.cpp
+++ b/Engine/Source/Runtime/Slate/Private/Framework/MultiBox/MultiBoxBuilder.cpp
@@ -23,18 +23,26 @@ namespace UE::Slate::Private
 {
 	/** Utility to set a Target attribute to the given Source attribute, if the target is not already set. */
 	template <typename AttributeType>
-	static void ApplyUnsetAttributeOverride(TAttribute<AttributeType>* InTargetAttribute, const TAttribute<AttributeType>* InSourceAttribute)
+	static void ApplyUnsetAttributeOverride(TAttribute<AttributeType>& InTargetAttribute, const TAttribute<AttributeType>& InSourceAttribute)
 	{
-		if (InTargetAttribute->IsSet() && !InTargetAttribute->IsBound())
+		if (!InTargetAttribute.IsSet())
 		{
-			// If the target is already set, don't override it
-			return;
+			if (InSourceAttribute.IsSet())
+			{
+				// Target is not set, and source is, so override Target with Source
+				InTargetAttribute = InSourceAttribute;
+			}
 		}
+	}
 
-		if (InSourceAttribute->IsSet())
+	/** Utility to set a Target TOptional to the given Source TOptional, if the target is not already set. */
+	template <typename OptionalType>
+	static void ApplyUnsetOptionalOverride(TOptional<OptionalType>& InTargetOptional, const TOptional<OptionalType>& InSourceOptional)
+	{
+		if (!InTargetOptional.IsSet() && InSourceOptional.IsSet())
 		{
 			// Target is not set, and source is, so override Target with Source
-			*InTargetAttribute = *InSourceAttribute;
+			InTargetOptional = InSourceOptional;
 		}
 	}
 
@@ -56,8 +64,10 @@ namespace UE::Slate::Private
 				const FMenuEntryResizeParams& SectionResizeParams = InSectionResizeParams.Get();
 
 				FMenuEntryResizeParams BlockResizeParams = BlockResizeParamsAttribute.Get();
-				ApplyUnsetAttributeOverride(&BlockResizeParams.Wrapping.Allow, &SectionResizeParams.Wrapping.Allow);
-				ApplyUnsetAttributeOverride(&BlockResizeParams.Wrapping.Priority, &SectionResizeParams.Wrapping.Priority);
+				ApplyUnsetAttributeOverride(BlockResizeParams.Wrapping.Allow, SectionResizeParams.Wrapping.Allow);
+				ApplyUnsetAttributeOverride(BlockResizeParams.Wrapping.Mode, SectionResizeParams.Wrapping.Mode);
+				ApplyUnsetAttributeOverride(BlockResizeParams.Wrapping.Priority, SectionResizeParams.Wrapping.Priority);
+				ApplyUnsetOptionalOverride(BlockResizeParams.Wrapping.VerticalOverflowBehavior, SectionResizeParams.Wrapping.VerticalOverflowBehavior);
 
 				BlockResizeParamsAttribute = BlockResizeParams;
 			}
@@ -69,8 +79,10 @@ namespace UE::Slate::Private
 						TAttribute<FMenuEntryResizeParams>::FGetter::CreateLambda([BlockResizeParamsAttribute, SectionResizeParams = InSectionResizeParams.Get()]
 						{
 							FMenuEntryResizeParams ResizeParams = BlockResizeParamsAttribute.Get();
-							ApplyUnsetAttributeOverride(&ResizeParams.Wrapping.Allow, &SectionResizeParams.Wrapping.Allow);
-							ApplyUnsetAttributeOverride(&ResizeParams.Wrapping.Priority, &SectionResizeParams.Wrapping.Priority);
+							ApplyUnsetAttributeOverride(ResizeParams.Wrapping.Allow, SectionResizeParams.Wrapping.Allow);
+							ApplyUnsetAttributeOverride(ResizeParams.Wrapping.Mode, SectionResizeParams.Wrapping.Mode);
+							ApplyUnsetAttributeOverride(ResizeParams.Wrapping.Priority, SectionResizeParams.Wrapping.Priority);
+							ApplyUnsetOptionalOverride(ResizeParams.Wrapping.VerticalOverflowBehavior, SectionResizeParams.Wrapping.VerticalOverflowBehavior);
 
 							return ResizeParams;
 						})));
