Fix MultiBoxBuilder AddWidget

Make sure we forward along all parameters to inner AddWidget function

#rnx

diff --git a/Engine/Source/Runtime/Slate/Private/Framework/MultiBox/MultiBoxBuilder.cpp b/Engine/Source/Runtime/Slate/Private/Framework/MultiBox/MultiBoxBuilder.cpp
--- a/Engine/Source/Runtime/Slate/Private/Framework/MultiBox/MultiBoxBuilder.cpp
+++ b/Engine/Source/Runtime/Slate/Private/Framework/MultiBox/MultiBoxBuilder.cpp
@@ -489,7 +489,7 @@ void FMenuBuilder::AddWidget( TSharedRef<SWidget> InWidget, const FText& Label,
 
 void FMenuBuilder::AddWidget(const TSharedRef<SWidget>& InWidget, const FText& InLabel, const FMenuEntryStyleParams& InStyleParams, const bool bInSearchable, const TAttribute<FText>& InToolTipText, const TAttribute<FSlateIcon>& InIcon, const TAttribute<EVisibility>& InVisibility)
 {
-	AddWidget(InWidget, InLabel, InStyleParams, FMenuEntryResizeParams(), bInSearchable, InToolTipText, InIcon);
+	AddWidget(InWidget, InLabel, InStyleParams, FMenuEntryResizeParams(), bInSearchable, InToolTipText, InIcon, InVisibility);
 }
 
 void FMenuBuilder::AddWidget(const TSharedRef<SWidget>& InWidget, const FText& InLabel, const FMenuEntryStyleParams& InStyleParams, const TAttribute<FMenuEntryResizeParams>& InResizeParams, const bool bInSearchable, const TAttribute<FText>& InToolTipText, const TAttribute<FSlateIcon>& InIcon, const TAttribute<EVisibility>& InVisibility)
