[Slate] Add an accessor for checking the FSlateUser's cursor visibilitiy.

#rnx

diff --git a/Engine/Source/Runtime/Slate/Private/Framework/Application/SlateUser.cpp b/Engine/Source/Runtime/Slate/Private/Framework/Application/SlateUser.cpp
--- a/Engine/Source/Runtime/Slate/Private/Framework/Application/SlateUser.cpp
+++ b/Engine/Source/Runtime/Slate/Private/Framework/Application/SlateUser.cpp
@@ -392,6 +392,12 @@ void FSlateUser::SetCursorVisibility(bool bDrawCursor)
 	}
 }
 
+bool FSlateUser::IsCursorVisible() const
+{
+	// Consider the cursor valid so long as we can draw it and it is not the "None" type
+	return bCanDrawCursor && Cursor.IsValid() && Cursor->GetType() != EMouseCursor::None;
+}
+
 void FSlateUser::SetCursorPosition(int32 PosX, int32 PosY)
 {
 	SetPointerPosition(FSlateApplication::CursorPointerIndex, PosX, PosY);
diff --git a/Engine/Source/Runtime/Slate/Public/Framework/Application/SlateUser.h b/Engine/Source/Runtime/Slate/Public/Framework/Application/SlateUser.h
--- a/Engine/Source/Runtime/Slate/Public/Framework/Application/SlateUser.h
+++ b/Engine/Source/Runtime/Slate/Public/Framework/Application/SlateUser.h
@@ -87,6 +87,7 @@ public:
 	SLATE_API TSharedPtr<SWidget> GetPointerCaptor(uint32 PointerIndex) const;
 
 	SLATE_API void SetCursorVisibility(bool bDrawCursor);
+	SLATE_API bool IsCursorVisible() const;
 
 	SLATE_API void SetCursorPosition(int32 PosX, int32 PosY);
 	SLATE_API void SetCursorPosition(const UE::Slate::FDeprecateVector2DParameter& NewCursorPos);
