Slate: Prevent zero-delta platform mouse move events from triggering mouse move events in the engine

Typically, users encountering this issue would see their UI switch from gamepad icons to keyboard and mouse icons even though they never moved their mouse.

The primary cause is on Windows when other applications would trigger WM_MOUSEMOVE messages (eg. Resource Monitor)

#rnx

diff --git a/Engine/Source/Runtime/Slate/Private/Framework/Application/SlateApplication.cpp b/Engine/Source/Runtime/Slate/Private/Framework/Application/SlateApplication.cpp
--- a/Engine/Source/Runtime/Slate/Private/Framework/Application/SlateApplication.cpp
+++ b/Engine/Source/Runtime/Slate/Private/Framework/Application/SlateApplication.cpp
@@ -6059,7 +6059,12 @@ FReply FSlateApplication::RouteMouseWheelOrGestureEvent(const FWidgetPath& Widge
 
 bool FSlateApplication::OnMouseMove()
 {
-	// If the left button is pressed we fake 
+	if (!PlatformApplication->Cursor.IsValid())
+    {
+    	return false;
+    }
+
+    // If the left button is pressed we fake 
 	if (IsFakingTouchEvents() && (GetPressedMouseButtons().Num() == 0 || GetPressedMouseButtons().Contains(EKeys::LeftMouseButton)))
 	{
 		// convert to touch event if we are faking it
@@ -6073,38 +6078,36 @@ bool FSlateApplication::OnMouseMove()
 	}
 
 	bool Result = true;
-	const FVector2f CurrentCursorPosition = GetCursorPos();
-	const FVector2f LastCursorPosition = GetLastCursorPos();
-	
-	// Force the cursor user index to use the platform cursor since we've been notified that the platform 
-	// cursor position has changed. This is done intentionally after getting the positions in order to avoid
-	// false positives.
+	const FVector2f CurrentCursorPosition = UE::Slate::CastToVector2f(PlatformApplication->Cursor->GetPosition());
+	if (LastPlatformCursorPosition != CurrentCursorPosition)
+	{
+		// Force the cursor user index to use the platform cursor since we've been notified that the platform 
+		// cursor position has changed. This is done intentionally after getting the positions in order to avoid
+		// false positives.
 
-	// NOTE: When we swap out the real OS cursor for the faux slate cursor ie. UsePlatformCursorForCursorUser(false)
-	// we reset this event count to 0.  This occurs typically when a gamepad is being used and you don't want to manipulate the real
-	// OS cursor, instead move around a fake cursor so that you can still do development stuff outside the game window.  Anyway,
-	// when this occurs, in the future the OS will send you a long delayed mouse movement.  I'm not exactly sure what's triggering it,
-	// it's not a movement from the application, I think it's something more subtle, like swapping true cursor visibility for using
-	// a None cursor, it could also be some combination.
-	// 
-	// In any event, we track the number of events, and if we get more than 3, then we start trying to swap back to the OS cursor.
-	// the 3 should give us any buffer needed for either a last frame mouse movement, or a weird condition like noted above.
-	PlatformMouseMovementEvents++;
+		// NOTE: When we swap out the real OS cursor for the faux slate cursor ie. UsePlatformCursorForCursorUser(false)
+		// we reset this event count to 0.  This occurs typically when a gamepad is being used and you don't want to manipulate the real
+		// OS cursor, instead move around a fake cursor so that you can still do development stuff outside the game window.  Anyway,
+		// when this occurs, in the future the OS will send you a long delayed mouse movement.  I'm not exactly sure what's triggering it,
+		// it's not a movement from the application, I think it's something more subtle, like swapping true cursor visibility for using
+		// a None cursor, it could also be some combination.
+		// 
+		// In any event, we track the number of events, and if we get more than 3, then we start trying to swap back to the OS cursor.
+		// the 3 should give us any buffer needed for either a last frame mouse movement, or a weird condition like noted above.
+		PlatformMouseMovementEvents++;
+
+		if (PlatformMouseMovementEvents > 3)
+		{
+			UsePlatformCursorForCursorUser(true);
+		}
 
-	if (PlatformMouseMovementEvents > 3)
-	{
-		UsePlatformCursorForCursorUser(true);
-	}
-	
-	if ( LastCursorPosition != CurrentCursorPosition )
-	{
 		LastMouseMoveTime = GetCurrentTime();
 
 		FPointerEvent MouseEvent(
 			GetUserIndexForMouse(),
 			CursorPointerIndex,
 			CurrentCursorPosition,
-			LastCursorPosition,
+			LastPlatformCursorPosition,
 			PressedMouseButtons,
 			EKeys::Invalid,
 			0,
@@ -6112,6 +6115,7 @@ bool FSlateApplication::OnMouseMove()
 			);
 
 		Result = ProcessMouseMoveEvent( MouseEvent );
+		LastPlatformCursorPosition = CurrentCursorPosition;
 	}
 
 	return Result;
diff --git a/Engine/Source/Runtime/Slate/Public/Framework/Application/SlateApplication.h b/Engine/Source/Runtime/Slate/Public/Framework/Application/SlateApplication.h
--- a/Engine/Source/Runtime/Slate/Public/Framework/Application/SlateApplication.h
+++ b/Engine/Source/Runtime/Slate/Public/Framework/Application/SlateApplication.h
@@ -2164,6 +2164,9 @@ private:
 	/** Platform mouse movement event count. */
 	uint64 PlatformMouseMovementEvents = 0;
 
+	/** The location of the cursor the last time OnMouseMove was called */
+	FVector2f LastPlatformCursorPosition = FVector2f::ZeroVector;
+
 	/** Constant delta time used on every Slate widget Tick if CVarSlateUseFixedDeltaTime is enabled. */
 	static SLATE_API double FixedDeltaTime;
 	/**
