Update the style of overflow menus for slim & viewport toolbars to be more differentiated from normal toolbar content.
- Separate out wrap button values in FToolBarStyle to FWrapButtonStyle for better namespacing.
- Remove FLevelEditorStyle as its only use was to shift where the wrap was in the level viewport toolbar.
- Fix an issue in SClippingHorizontalBox where the overflow button was not actually removed when not needed.
- Allow negative wrap button padding values to overflow normal bounds.

#rnx

diff --git a/Engine/Source/Editor/EditorStyle/Private/StarshipStyle.cpp b/Engine/Source/Editor/EditorStyle/Private/StarshipStyle.cpp
--- a/Engine/Source/Editor/EditorStyle/Private/StarshipStyle.cpp
+++ b/Engine/Source/Editor/EditorStyle/Private/StarshipStyle.cpp
@@ -1809,11 +1809,10 @@ void FStarshipEditorStyle::FStyle::SetupGeneralStyles()
 		FVerticalToolBarStyle.SetBackgroundPadding(   FMargin(6.0f, 4.0f));
 		FVerticalToolBarStyle.ButtonStyle.SetNormalPadding(FMargin(12.0f, 6.0f));
 		FVerticalToolBarStyle.ButtonStyle.SetPressedPadding(FMargin(12.0f, 6.0f));
-		FVerticalToolBarStyle.SetExpandBrush(IMAGE_BRUSH("Icons/toolbar_expand_16x", Icon8x8));
 		FVerticalToolBarStyle.SetIconPadding(FMargin(8.0f, 8.0f));
 		FVerticalToolBarStyle.SetIconPaddingWithVisibleLabel(FMargin(8.0f, 8.0f, 8.0f, 2.0f));
  
-		FVerticalToolBarStyle.SetExpandBrush(CORE_IMAGE_BRUSH_SVG("Starship/Common/ellipsis-horizontal-narrow", Icon16x16));
+		FVerticalToolBarStyle.WrapButtonStyle.SetExpandBrush(CORE_IMAGE_BRUSH_SVG("Starship/Common/ellipsis-horizontal-narrow", Icon16x16));
 		FCheckBoxStyle CheckBoxStyle = FCheckBoxStyle(FVerticalToolBarStyle.ToggleButton)
 			.SetCheckedImage(FSlateRoundedBoxBrush(FStyleColors::Primary, 4.f, FLinearColor(0.f, 0.f, 0.f, .8f), 0.5f))
 			.SetCheckedHoveredImage(FSlateRoundedBoxBrush(FStyleColors::PrimaryHover, 4.f, FLinearColor(0.f, 0.f, 0.f, .8f), 0.5f))
@@ -2621,6 +2620,19 @@ void FStarshipEditorStyle::FStyle::SetupViewportStyles()
 				.SetDownArrowImage(IMAGE_BRUSH_SVG("Starship/EditorViewport/small-chevron-down", FVector2f(6.f, 6.f)))
 				.SetDownArrowPadding(FMargin(4.0f, 0.f, 0.f, 0.f));
 		}
+		
+		FComboButtonStyle WrapButtonComboStyle = FComboButtonStyle(ComboButtonStyle)
+			.SetButtonStyle(FButtonStyle()
+				.SetNormal(FSlateColorBrush(FStyleColors::Dropdown))
+				.SetHovered(FSlateColorBrush(FStyleColors::Hover))
+				.SetPressed(FSlateColorBrush(FStyleColors::Header))
+				.SetNormalForeground(FStyleColors::Foreground)
+				.SetHoveredForeground(FStyleColors::ForegroundHover)
+				.SetPressedForeground(FStyleColors::ForegroundHover)
+				.SetDisabledForeground(FStyleColors::Hover2)
+				.SetNormalPadding(FMargin(6.0f, 4.0f))
+				.SetPressedPadding(FMargin(6.0f, 4.0f))
+			);
 
 		FToolBarStyle ViewportToolbarStyle = SlimToolBarStyle;
 		{
@@ -2638,8 +2650,14 @@ void FStarshipEditorStyle::FStyle::SetupViewportStyles()
 				.SetSeparatorThickness(1.0f)
 				.SetSeparatorPadding(FMargin(4.0f, 0.f))
 				.SetBackgroundPadding(FMargin(4.0f))
-				.SetWrapButtonPadding(FMargin(4.0f, 0.f))
-				.SetWrapButtonIndex(-1)
+				.SetWrapButtonStyle(
+					FWrapButtonStyle(SlimToolBarStyle.WrapButtonStyle)
+					.SetComboButtonStyle(WrapButtonComboStyle)
+					.SetWrapButtonIndex(-1)
+					.SetWrapButtonPadding(FMargin(4.f, -4.f, -4.f, -4.f))
+					.SetSeparatorBrush(SlimToolBarStyle.SeparatorBrush)
+					.SetSeparatorThickness(2.0f)
+				)
 				.SetRaisedChildrenRightPadding(16.0f);
 
 			Set("ViewportToolbar", ViewportToolbarStyle);
@@ -3044,8 +3062,10 @@ void FStarshipEditorStyle::FStyle::SetupViewportStyles()
 			.SetToggleButtonStyle(SoloToggleButton)
 			.SetButtonStyle(ViewportMenuButton)
 			.SetSeparatorBrush(FSlateNoResource())
-			.SetSeparatorPadding(FMargin(2.0f, 0.0f))
-			.SetExpandBrush(IMAGE_BRUSH("Icons/toolbar_expand_16x", Icon8x8));
+			.SetSeparatorPadding(FMargin(2.0f, 0.0f));
+		
+		ViewportToolbarStyle.WrapButtonStyle.SetExpandBrush(IMAGE_BRUSH("Icons/toolbar_expand_16x", Icon8x8));
+		
 		Set("EditorViewportToolBar", ViewportToolbarStyle);
 
 		FButtonStyle ViewportMenuWarningButton = FButtonStyle(ViewportMenuButton)
@@ -3101,7 +3121,6 @@ void FStarshipEditorStyle::FStyle::SetupViewportStyles()
 		FToolBarStyle ViewportToolbarStyle =
 			FToolBarStyle()
 			.SetBackground(BOX_BRUSH("Old/Menu_Background", FMargin(8.0f / 64.0f), FLinearColor::Transparent))
-			.SetExpandBrush(IMAGE_BRUSH("Icons/toolbar_expand_16x", Icon8x8))
 			.SetComboButtonPadding(FMargin(0.f))
 			.SetButtonPadding(FMargin(0.f))
 			.SetCheckBoxPadding(FMargin(4.f))
@@ -3117,6 +3136,8 @@ void FStarshipEditorStyle::FStyle::SetupViewportStyles()
 				.SetFont(DEFAULT_FONT("Bold", 9))
 				.SetColorAndOpacity(FLinearColor(0.0f, 0.0f, 0.0f, 1.0f))
 			);
+			
+		ViewportToolbarStyle.WrapButtonStyle.SetExpandBrush(IMAGE_BRUSH("Icons/toolbar_expand_16x", Icon8x8));
 		
 		const FString SmallRoundedButton(TEXT("Common/SmallRoundedButton"));
 		const FString SmallRoundedButtonStart(TEXT("Common/SmallRoundedButtonLeft"));
@@ -3656,7 +3677,6 @@ void FStarshipEditorStyle::FStyle::SetupPropertyEditorStyles()
 
 		FToolBarStyle DetailsExtensionToolBarStyle = FToolBarStyle(FStarshipCoreStyle::GetCoreStyle().GetWidgetStyle<FToolBarStyle>("SlimToolBar"))
 			.SetButtonStyle(DetailsExtensionMenuButton)
-			.SetExpandBrush(CORE_IMAGE_BRUSH_SVG("Starship/Common/ellipsis-vertical-narrow", FVector2f(4.f, 16.f)))
 			.SetIconSize(Icon16x16)
 			.SetBackground(FSlateNoResource())
 			.SetLabelPadding(FMargin(0.f))
@@ -3668,6 +3688,9 @@ void FStarshipEditorStyle::FStyle::SetupPropertyEditorStyles()
 			.SetCheckBoxPadding(FMargin(0.f))
 			.SetSeparatorBrush(FSlateNoResource())
 			.SetSeparatorPadding(FMargin(0.f));
+		
+		DetailsExtensionToolBarStyle.WrapButtonStyle
+			.SetExpandBrush(CORE_IMAGE_BRUSH_SVG("Starship/Common/ellipsis-vertical-narrow", FVector2f(4.f, 16.f)));
 
 		Set("DetailsView.ExtensionToolBar", DetailsExtensionToolBarStyle);
 
diff --git a/Engine/Source/Editor/LevelEditor/Private/LevelEditorStyle.h b/Engine/Source/Editor/LevelEditor/Private/LevelEditorStyle.h
deleted file mode 100644
diff --git a/Engine/Source/Editor/LevelEditor/Private/SLevelViewport.cpp b/Engine/Source/Editor/LevelEditor/Private/SLevelViewport.cpp
--- a/Engine/Source/Editor/LevelEditor/Private/SLevelViewport.cpp
+++ b/Engine/Source/Editor/LevelEditor/Private/SLevelViewport.cpp
@@ -107,7 +107,6 @@
 #include "FunctionalUIScreenshotTest.h"
 #include "GPUSkinCache.h"
 #include "GPUSkinCacheVisualizationMenuCommands.h"
-#include "LevelEditorStyle.h"
 #include "SActionableMessageViewportWidget.h"
 #include "ShowFlagMenuCommands.h"
 #include "SkeletalRenderPublic.h"
@@ -2294,7 +2293,6 @@ TSharedPtr<SWidget> SLevelViewport::BuildViewportToolbar()
 				ViewportToolbarName, NAME_None /* parent */, EMultiBoxType::SlimHorizontalToolBar
 			);
 
-			ViewportToolbarMenu->SetStyleSet(&FLevelEditorStyle::Get());
 			ViewportToolbarMenu->StyleName = "ViewportToolbar";
 			ViewportToolbarMenu->bSeparateSections = false;
 
diff --git a/Engine/Source/Editor/LevelEditor/Private/ViewportToolbar/LevelEditorStyle.cpp b/Engine/Source/Editor/LevelEditor/Private/ViewportToolbar/LevelEditorStyle.cpp
deleted file mode 100644
diff --git a/Engine/Source/Editor/Sequencer/Private/Filters/Widgets/SFilterBarClippingHorizontalBox.cpp b/Engine/Source/Editor/Sequencer/Private/Filters/Widgets/SFilterBarClippingHorizontalBox.cpp
--- a/Engine/Source/Editor/Sequencer/Private/Filters/Widgets/SFilterBarClippingHorizontalBox.cpp
+++ b/Engine/Source/Editor/Sequencer/Private/Filters/Widgets/SFilterBarClippingHorizontalBox.cpp
@@ -111,7 +111,7 @@ TSharedRef<SComboButton> SFilterBarClippingHorizontalBox::CreateWrapButton()
 		[
 			SNew(SImage)
 			.DesiredSizeOverride(FVector2D(10.f))
-			.Image(&ToolBarStyle.ExpandBrush)
+			.Image(&ToolBarStyle.WrapButtonStyle.ExpandBrush)
 		];
 
 	WeakWrapButton = WrapButton;
diff --git a/Engine/Source/Runtime/Slate/Private/Framework/MultiBox/SClippingHorizontalBox.cpp b/Engine/Source/Runtime/Slate/Private/Framework/MultiBox/SClippingHorizontalBox.cpp
--- a/Engine/Source/Runtime/Slate/Private/Framework/MultiBox/SClippingHorizontalBox.cpp
+++ b/Engine/Source/Runtime/Slate/Private/Framework/MultiBox/SClippingHorizontalBox.cpp
@@ -8,6 +8,7 @@
 #include "Styling/ToolBarStyle.h"
 #include "Widgets/Images/SImage.h"
 #include "Widgets/Input/SComboButton.h"
+#include "Widgets/Layout/SSeparator.h"
 
 namespace UE::Slate
 {
@@ -292,7 +293,7 @@ void SClippingHorizontalBox::OnArrangeChildren(const FGeometry& AllottedGeometry
 	if (bAllowWrapButton)
 	{
 		UE::Slate::PrioritizedResize(
-			InAllottedWidth, WrapButtonWidth, ToolBarStyle.WrapButtonPadding, WrapButtonIndex, ClippingInfos, WrapButtonX
+			InAllottedWidth, WrapButtonWidth, ToolBarStyle.WrapButtonStyle.Padding, WrapButtonIndex, ClippingInfos, WrapButtonX
 		);
 	}
 	else
@@ -328,15 +329,7 @@ void SClippingHorizontalBox::OnArrangeChildren(const FGeometry& AllottedGeometry
 		}
 	}
 
-	// Position the wrap button.
-	if (const float* WrapX = WrapButtonX.GetPtrOrNull())
-	{
-		FGeometry& WrapButtonGeometry = ArrangedChildren[ArrangedChildren.Num() - 1].Geometry;
-		const FVector2D WrapButtonSize = FVector2D(WrapButtonWidth, WrapButtonGeometry.GetLocalSize().Y);
-		WrapButtonGeometry = AllottedGeometry.MakeChild(WrapButtonSize, FSlateLayoutTransform(FVector2f(*WrapX, 0)));
-	}
-
-	// Finally, position all children using the clipping information.
+	// Position all children using the clipping information.
 	for (int32 Index = 0; Index < ArrangedChildren.Num() - 1; ++Index)
 	{
 		const UE::Slate::FClippingInfo* const Info = ClippingInfos.FindByPredicate(
@@ -360,6 +353,45 @@ void SClippingHorizontalBox::OnArrangeChildren(const FGeometry& AllottedGeometry
 		const FVector2D WrapButtonSize = FVector2D(Info->Width, Geometry.GetLocalSize().Y);
 		Geometry = AllottedGeometry.MakeChild(WrapButtonSize, FSlateLayoutTransform(FVector2f(Info->X, 0)));
 	}
+	
+	// Position or remove the wrap button
+	if (const float* WrapX = WrapButtonX.GetPtrOrNull())
+	{
+		const FMargin& Padding = ToolBarStyle.WrapButtonStyle.Padding;  
+	
+		FGeometry& WrapButtonGeometry = ArrangedChildren[ArrangedChildren.Num() - 1].Geometry;
+		FVector2f WrapButtonSize(WrapButtonWidth, WrapButtonGeometry.GetLocalSize().Y);
+		FVector2f WrapButtonPosition(*WrapX, 0.0f);
+		
+		// Allow configured padding to break the widget out of its normal bounds
+		if (Padding.Left < 0.f)
+		{
+			WrapButtonSize.X -= Padding.Left;
+			WrapButtonPosition.X -= Padding.Left;
+		}
+		
+		if (Padding.Right < 0.f)
+		{
+			WrapButtonSize.X -= Padding.Right;
+		}
+		
+		if (Padding.Top < 0.f)
+		{
+			WrapButtonSize.Y -= Padding.Top;
+			WrapButtonPosition.Y += Padding.Top;	
+		}
+		
+		if (Padding.Bottom < 0.f)
+		{
+			WrapButtonSize.Y -= Padding.Bottom;	
+		}
+		
+		WrapButtonGeometry = AllottedGeometry.MakeChild(WrapButtonSize, FSlateLayoutTransform(WrapButtonPosition));
+	}
+	else
+	{
+		ArrangedChildren.Remove(ArrangedChildren.Num() - 1);
+	}
 }
 
 int32 SClippingHorizontalBox::OnPaint( const FPaintArgs& Args, const FGeometry& AllottedGeometry, const FSlateRect& MyCullingRect, FSlateWindowElementList& OutDrawElements, int32 LayerId, const FWidgetStyle& InWidgetStyle, bool bParentEnabled ) const
@@ -397,11 +429,22 @@ void SClippingHorizontalBox::AddWrapButton()
 
 	// Construct the wrap button used in toolbars and menubars
 	// Always allow this to be focusable to prevent the menu from collapsing during interaction
+	
+	const FWrapButtonStyle& WrapButtonStyle = ToolBarStyle.WrapButtonStyle;
+	
+	const FComboButtonStyle* ComboButtonStyle = WrapButtonStyle.ComboButtonStyle.GetPtrOrNull();
+	const FButtonStyle* ButtonStyle = nullptr;
+	if (!ComboButtonStyle)
+	{
+		ComboButtonStyle = &FAppStyle::Get().GetWidgetStyle< FComboButtonStyle >( "ComboButton" );
+		ButtonStyle = &ToolBarStyle.ButtonStyle;
+	}
+	
 	WrapButton = 
 		SNew( SComboButton )
-		.HasDownArrow( false )
-		.ButtonStyle(&ToolBarStyle.ButtonStyle)
-		.ContentPadding( FMargin(4.f, 0.f) )
+		.HasDownArrow( WrapButtonStyle.bHasDownArrow )
+		.ComboButtonStyle(ComboButtonStyle)
+		.ButtonStyle(ButtonStyle)
 		.ToolTipText( NSLOCTEXT("Slate", "ExpandToolbar", "Click to expand toolbar") )
 		.OnGetMenuContent( OnWrapButtonClicked )
 		.Cursor( EMouseCursor::Default )
@@ -410,20 +453,62 @@ void SClippingHorizontalBox::AddWrapButton()
 		.ButtonContent()
 		[
 			SNew(SImage)
-			.Image(&ToolBarStyle.ExpandBrush)
+			.ColorAndOpacity(FSlateColor::UseForeground())
+			.Image(&WrapButtonStyle.ExpandBrush)
+		];
+
+	TSharedRef<SWidget> RootWidget = WrapButton.ToSharedRef();
+
+	if (WrapButtonStyle.bIncludeSeparator)
+	{
+		TSharedRef<SHorizontalBox> Box = SNew(SHorizontalBox);
+		
+		if (WrapButtonStyle.WrapButtonIndex != 0) // Not at the start
+		{
+			// Insert a separator before
+			Box->AddSlot()
+			.AutoWidth()
+			.Padding(WrapButtonStyle.SeparatorPadding.Get(ToolBarStyle.SeparatorPadding))
+			[
+				SNew(SSeparator)
+				.Orientation(Orient_Vertical)
+				.Thickness(WrapButtonStyle.SeparatorThickness.Get(ToolBarStyle.SeparatorThickness))
+				.SeparatorImage(&WrapButtonStyle.SeparatorBrush.Get(ToolBarStyle.SeparatorBrush))
+			];	
+		}
+		
+		Box->AddSlot()
+		.AutoWidth()
+		[
+			RootWidget
 		];
+		
+		if (WrapButtonStyle.WrapButtonIndex != -1) // Not at the end 
+		{
+			// Insert a separator after
+			Box->AddSlot()
+			[
+				SNew(SSeparator)
+				.Orientation(Orient_Vertical)
+				.Thickness(WrapButtonStyle.SeparatorThickness.Get(ToolBarStyle.SeparatorThickness))
+				.SeparatorImage(&WrapButtonStyle.SeparatorBrush.Get(ToolBarStyle.SeparatorBrush))
+			];
+		}
+	
+		RootWidget = Box;
+	}
 
 	// Perform a prepass to get a valid DesiredSize value below
-	WrapButton->SlatePrepass(1.0f);
-	WrapButtonWidth = bAllowWrapButton ? WrapButton->GetDesiredSize().X : 0.0f;
-	WrapButtonIndex = ToolBarStyle.WrapButtonIndex;
+	RootWidget->SlatePrepass(1.0f);
+	WrapButtonWidth = bAllowWrapButton ? RootWidget->GetDesiredSize().X : 0.0f;
+	WrapButtonIndex = WrapButtonStyle.WrapButtonIndex;
 
 	// Add the wrap button
 	AddSlot()
 	.FillWidth(0.0f) // Effectively makes this widget 0 width, so it exists as a slot/child, but isn't considered for layout
 	.Padding(0.f)
 	[
-		WrapButton.ToSharedRef()
+		RootWidget
 	];
 }
 
diff --git a/Engine/Source/Runtime/Slate/Private/Framework/MultiBox/SClippingVerticalBox.cpp b/Engine/Source/Runtime/Slate/Private/Framework/MultiBox/SClippingVerticalBox.cpp
--- a/Engine/Source/Runtime/Slate/Private/Framework/MultiBox/SClippingVerticalBox.cpp
+++ b/Engine/Source/Runtime/Slate/Private/Framework/MultiBox/SClippingVerticalBox.cpp
@@ -125,7 +125,7 @@ void SClippingVerticalBox::InitializeWrapButton( TSharedPtr<SComboButton>& Butto
 	const TSharedRef<SImage> IconWidget =
 		SNew( SImage )
 		.Visibility(EVisibility::HitTestInvisible)
-		.Image( &ToolBarStyle.ExpandBrush );
+		.Image( &ToolBarStyle.WrapButtonStyle.ExpandBrush );
 
 
 	Style = ToolBarStyle.ButtonStyle;
diff --git a/Engine/Source/Runtime/Slate/Private/Framework/MultiBox/SUniformToolbarPanel.cpp b/Engine/Source/Runtime/Slate/Private/Framework/MultiBox/SUniformToolbarPanel.cpp
--- a/Engine/Source/Runtime/Slate/Private/Framework/MultiBox/SUniformToolbarPanel.cpp
+++ b/Engine/Source/Runtime/Slate/Private/Framework/MultiBox/SUniformToolbarPanel.cpp
@@ -42,7 +42,7 @@ void SUniformToolbarPanel::Construct( const FArguments& InArgs )
 		.ButtonContent()
 		[
 			SNew(SImage)
-			.Image(&ToolBarStyle.ExpandBrush)
+			.Image(&ToolBarStyle.WrapButtonStyle.ExpandBrush)
 		]
 		)));
 }
diff --git a/Engine/Source/Runtime/SlateCore/Private/Styling/StarshipCoreStyle.cpp b/Engine/Source/Runtime/SlateCore/Private/Styling/StarshipCoreStyle.cpp
--- a/Engine/Source/Runtime/SlateCore/Private/Styling/StarshipCoreStyle.cpp
+++ b/Engine/Source/Runtime/SlateCore/Private/Styling/StarshipCoreStyle.cpp
@@ -2241,7 +2241,6 @@ void FStarshipCoreStyle::SetupMultiboxStyles(TSharedRef<FStyle>& Style)
 			FToolBarStyle()
 			.SetBackground(FSlateColorBrush(FStyleColors::Panel))
 			.SetBackgroundPadding(FMargin(0.f, 4.f))
-			.SetExpandBrush(IMAGE_BRUSH("Icons/toolbar_expand_16x", Icon16x16))
 			.SetComboButtonPadding(FMargin(4.0f, 0.0f))
 			.SetButtonPadding(FMargin(2.0f, 0.f))
 			.SetCheckBoxPadding(FMargin(4.0f, 0.f))
@@ -2279,6 +2278,10 @@ void FStarshipCoreStyle::SetupMultiboxStyles(TSharedRef<FStyle>& Style)
 
 		NormalToolbarStyle.SetButtonStyle(ToolbarButton);
 		NormalToolbarStyle.SetIconSize(Icon40x40);
+		
+		NormalToolbarStyle.WrapButtonStyle
+			.SetIncludeSeparator(false)
+			.SetExpandBrush(IMAGE_BRUSH_SVG("Common/menu", Icon16x16));
 
 		Style->Set("ToolBar", NormalToolbarStyle);
 
@@ -2303,7 +2306,6 @@ void FStarshipCoreStyle::SetupMultiboxStyles(TSharedRef<FStyle>& Style)
 			FToolBarStyle()
 			.SetBackground(*SlimToolbarBackground)
 			.SetBackgroundPadding(FMargin(4.f, 6.f))
-			.SetExpandBrush(IMAGE_BRUSH("Icons/toolbar_expand_16x", Icon16x16))
 			.SetComboButtonPadding(FMargin(0.0f, 0.0f))
 			.SetButtonPadding(FMargin(4.0f, 0.0f))
 			.SetCheckBoxPadding(FMargin(8.0f, 0.0f))
@@ -2313,7 +2315,26 @@ void FStarshipCoreStyle::SetupMultiboxStyles(TSharedRef<FStyle>& Style)
 			.SetComboButtonStyle(SlimToolBarComboButton)
 			.SetLabelPadding(FMargin(4.f, 0.f, 0.f, 0.f))
 			.SetEditableTextStyle(FEditableTextBoxStyle(NormalEditableTextBoxStyle));
-
+			
+		SlimToolbarStyle.WrapButtonStyle
+			.SetWrapButtonPadding(FMargin(0.0f, -6.f, -4.f, -6.f))
+			.SetSeparatorPadding(0.f)
+			.SetComboButtonStyle(FComboButtonStyle(SlimToolbarStyle.ComboButtonStyle)
+			.SetContentPadding(FMargin(4.0f, 0.f, 0.f, 0.f))
+			.SetDownArrowPadding(FMargin(4.f, 0.f, 0.f, 0.f))
+			.SetButtonStyle(FButtonStyle()
+				.SetNormal(FSlateColorBrush(FStyleColors::Dropdown))
+				.SetHovered(FSlateColorBrush(FStyleColors::Hover))
+				.SetPressed(FSlateColorBrush(FStyleColors::Header))
+				.SetNormalForeground(FStyleColors::Foreground)
+				.SetHoveredForeground(FStyleColors::ForegroundHover)
+				.SetPressedForeground(FStyleColors::ForegroundHover)
+				.SetDisabledForeground(FStyleColors::Hover2)
+				.SetNormalPadding(FMargin(6.0f, 4.0f))
+				.SetPressedPadding(FMargin(6.0f, 4.0f))
+			))
+			.SetExpandBrush(IMAGE_BRUSH_SVG("Common/menu", Icon16x16));
+			
 		FCheckBoxStyle SlimToolBarToggleButtonCheckBoxStyle = FCheckBoxStyle(Style->GetWidgetStyle<FCheckBoxStyle>("ToggleButtonCheckbox"))
 			.SetPadding(FMargin(8.f, 4.f, 8.f, 4.f));
 
diff --git a/Engine/Source/Runtime/SlateCore/Private/Styling/ToolBarStyle.cpp b/Engine/Source/Runtime/SlateCore/Private/Styling/ToolBarStyle.cpp
--- a/Engine/Source/Runtime/SlateCore/Private/Styling/ToolBarStyle.cpp
+++ b/Engine/Source/Runtime/SlateCore/Private/Styling/ToolBarStyle.cpp
@@ -5,11 +5,37 @@
 
 #include UE_INLINE_GENERATED_CPP_BY_NAME(ToolBarStyle)
 
+FWrapButtonStyle::FWrapButtonStyle()
+	: Padding(FMargin(0.f))
+	, WrapButtonIndex(-1) // Default to the right side of a menu
+	, SeparatorThickness(2.0f)
+	, SeparatorPadding(0)
+{
+}
+
+FWrapButtonStyle::FWrapButtonStyle(const FWrapButtonStyle&) = default;
+
+void FWrapButtonStyle::GetResources(TArray<const FSlateBrush*>& OutBrushes) const
+{
+	OutBrushes.Add(&ExpandBrush);
+	
+	if (const FComboButtonStyle* ComboButton = ComboButtonStyle.GetPtrOrNull())
+	{
+		ComboButton->GetResources(OutBrushes);
+	}
+	
+	if (const FSlateBrush* Brush = SeparatorBrush.GetPtrOrNull())
+	{
+		OutBrushes.Add(Brush);
+	}
+}
+
+
 const FName FToolBarStyle::TypeName(TEXT("FToolbarStyle"));
 
 FToolBarStyle::FToolBarStyle()
 	: BackgroundBrush(FSlateNoResource())
-	, ExpandBrush(FSlateNoResource())
+	, ExpandBrush_DEPRECATED(FSlateNoResource())
 	, SeparatorBrush(FSlateNoResource())
 	, LabelStyle()
 	, EditableTextStyle()
@@ -30,8 +56,8 @@ FToolBarStyle::FToolBarStyle()
 	, IndentedBlockPadding(0)
 	, BlockHovered(FSlateNoResource())
 	, BackgroundPadding(0)
-	, WrapButtonPadding(0)
-	, WrapButtonIndex(-1)
+	, WrapButtonPadding_DEPRECATED(0)
+	, WrapButtonIndex_DEPRECATED(-1)
 	, IconSize(16, 16)
 	, bShowLabels(true)
 	, ComboContentMinWidth(0)
@@ -52,7 +78,7 @@ const FToolBarStyle& FToolBarStyle::GetDefault()
 void FToolBarStyle::GetResources(TArray<const FSlateBrush*>& OutBrushes) const
 {
 	OutBrushes.Add(&BackgroundBrush);
-	OutBrushes.Add(&ExpandBrush);
+	OutBrushes.Add(&ExpandBrush_DEPRECATED);
 	OutBrushes.Add(&SeparatorBrush);
 	OutBrushes.Add(&BlockHovered);
 
@@ -63,6 +89,5 @@ void FToolBarStyle::GetResources(TArray<const FSlateBrush*>& OutBrushes) const
 	SettingsToggleButton.GetResources(OutBrushes);
 	SettingsButtonStyle.GetResources(OutBrushes);
 	ButtonStyle.GetResources(OutBrushes);
+	WrapButtonStyle.GetResources(OutBrushes);
 }
-
-
diff --git a/Engine/Source/Runtime/SlateCore/Public/Styling/ToolBarStyle.h b/Engine/Source/Runtime/SlateCore/Public/Styling/ToolBarStyle.h
--- a/Engine/Source/Runtime/SlateCore/Public/Styling/ToolBarStyle.h
+++ b/Engine/Source/Runtime/SlateCore/Public/Styling/ToolBarStyle.h
@@ -8,6 +8,62 @@
 
 #include "ToolBarStyle.generated.h"
 
+USTRUCT(BlueprintType)
+struct FWrapButtonStyle
+{
+	GENERATED_BODY()
+	
+	SLATECORE_API FWrapButtonStyle();
+	SLATECORE_API FWrapButtonStyle(const FWrapButtonStyle&);
+	
+	SLATECORE_API void GetResources(TArray<const FSlateBrush*>& OutBrushes) const;
+
+	/** The padding around the wrap button */
+	UPROPERTY(EditAnywhere, BlueprintReadWrite, Category = Appearance)
+	FMargin Padding;
+	FWrapButtonStyle& SetWrapButtonPadding(const FMargin& InWrapButtonPadding) { Padding = InWrapButtonPadding; return *this; }
+	
+	/** Where in the toolbar the wrap button should appear. e.g. 0 for the left side, -1 for the right side. */
+	UPROPERTY(EditAnywhere, BlueprintReadWrite, Category = Appearance)
+	int32 WrapButtonIndex;
+	FWrapButtonStyle& SetWrapButtonIndex(int32 InWrapButtonIndex) { WrapButtonIndex = InWrapButtonIndex; return *this; }
+
+	/** The brush used for the expand arrow when the toolbar runs out of room and needs to display toolbar items in a menu*/
+	UPROPERTY(EditAnywhere, BlueprintReadWrite, Category = Appearance)
+	FSlateBrush ExpandBrush;
+	FWrapButtonStyle& SetExpandBrush(const FSlateBrush& InExpandBrush) { ExpandBrush = InExpandBrush; return *this; }
+
+	/** Whether the combo box includes a down arrow */
+	UPROPERTY(EditAnywhere, BlueprintReadWrite, Category = Appearance)
+	bool bHasDownArrow = true;
+	FWrapButtonStyle& SetHasDownArrow(bool bInHasDownArrow) { bHasDownArrow = bInHasDownArrow; return *this; }
+
+	/** The styling of the combo button that opens the wrapping menu */
+	UPROPERTY(EditAnywhere, BlueprintReadWrite, Category = Appearance)
+	TOptional<FComboButtonStyle> ComboButtonStyle;
+	FWrapButtonStyle& SetComboButtonStyle(const FComboButtonStyle& InComboButtonStyle) { ComboButtonStyle = InComboButtonStyle; return *this; }
+	
+	/** Whether a separator should appear adjacent to the combo button */
+	UPROPERTY(EditAnywhere, BlueprintReadWrite, Category = Appearance)
+	bool bIncludeSeparator = true;
+	FWrapButtonStyle& SetIncludeSeparator(bool bInIncludeSeparator) { bIncludeSeparator = bInIncludeSeparator; return *this; }
+	
+	/** The appearance of the separator */
+	UPROPERTY(EditAnywhere, BlueprintReadWrite, Category = Appearance)
+	TOptional<FSlateBrush> SeparatorBrush;
+	FWrapButtonStyle& SetSeparatorBrush(const FSlateBrush& InSeparatorBrush) { SeparatorBrush = InSeparatorBrush; return *this; }
+
+	/** How wide/tall the separator should be */
+	UPROPERTY(EditAnywhere, BlueprintReadWrite, Category = Appearance)
+	TOptional<float> SeparatorThickness;
+	FWrapButtonStyle& SetSeparatorThickness(float InSeparatorThickness) { SeparatorThickness = InSeparatorThickness; return *this; }
+	
+	/** Any padding around the separator */
+	UPROPERTY(EditAnywhere, BlueprintReadWrite, Category = Appearance)
+	TOptional<FMargin> SeparatorPadding;
+	FWrapButtonStyle& SetSeparatorPadding(const FMargin& InSeparatorPadding) { SeparatorPadding = InSeparatorPadding; return *this; }
+};
+
 /**
  * Represents the appearance of a toolbar 
  */
@@ -33,9 +89,13 @@ struct FToolBarStyle : public FSlateWidgetStyle
 	FToolBarStyle& SetBackground(const FSlateBrush& InBackground) { BackgroundBrush = InBackground; return *this; }
 
 	/** The brush used for the expand arrow when the toolbar runs out of room and needs to display toolbar items in a menu*/
-	UPROPERTY(EditAnywhere, BlueprintReadWrite, Category = Appearance)
-	FSlateBrush ExpandBrush;
-	FToolBarStyle& SetExpandBrush(const FSlateBrush& InExpandBrush) { ExpandBrush = InExpandBrush; return *this; }
+	UPROPERTY(meta = (DeprecatedProperty, DeprecationMessage = "Use WrapButtonStyle.ExpandBrush instead"))
+	FSlateBrush ExpandBrush_DEPRECATED;
+	UE_DEPRECATED(5.6, "Use WrapButtonStyle.SetExpandBrush() instead")
+	FToolBarStyle& SetExpandBrush(const FSlateBrush& InExpandBrush)
+	{
+		ExpandBrush_DEPRECATED = InExpandBrush; return *this;
+	}
 
 	UPROPERTY(EditAnywhere, BlueprintReadWrite, Category = Appearance)
 	FSlateBrush SeparatorBrush;
@@ -131,18 +191,26 @@ struct FToolBarStyle : public FSlateWidgetStyle
 	FToolBarStyle& SetBackgroundPadding(const FMargin& InMargin) { BackgroundPadding = InMargin; return *this; }
 
 	UPROPERTY(EditAnywhere, BlueprintReadWrite, Category = Appearance)
-	FMargin WrapButtonPadding;
+	FWrapButtonStyle WrapButtonStyle;
+	FToolBarStyle& SetWrapButtonStyle(const FWrapButtonStyle& InButtonStyle) { WrapButtonStyle = InButtonStyle; return *this; }
+
+	UPROPERTY(meta = (DeprecatedProperty, DeprecationMessage = "Use WrapButtonStyle.Padding instead"))
+	FMargin WrapButtonPadding_DEPRECATED;
+	UE_DEPRECATED(5.6, "Use WrapButtonStyle instead")
 	FToolBarStyle& SetWrapButtonPadding(const FMargin& InWrapButtonPadding)
 	{
-		WrapButtonPadding = InWrapButtonPadding;
+		WrapButtonPadding_DEPRECATED = InWrapButtonPadding;
+		WrapButtonStyle.SetWrapButtonPadding(InWrapButtonPadding);
 		return *this;
 	}
 
-	UPROPERTY(EditAnywhere, BlueprintReadWrite, Category = Appearance)
-	int32 WrapButtonIndex;
+	UPROPERTY(meta = (DeprecatedProperty, DeprecationMessage = "Use WrapButtonStyle.Padding instead"))
+	int32 WrapButtonIndex_DEPRECATED;
+	UE_DEPRECATED(5.6, "Use WrapButtonStyle instead")
 	FToolBarStyle& SetWrapButtonIndex(int32 InWrapButtonIndex)
 	{
-		WrapButtonIndex = InWrapButtonIndex;
+		WrapButtonIndex_DEPRECATED = InWrapButtonIndex;
+		WrapButtonStyle.SetWrapButtonIndex(InWrapButtonIndex);
 		return *this;
 	}
 
