InertialScrollManager
- Promote FrictionCoefficient and StaticVelocityDrag as CVar instead of static variable.

diff --git a/Engine/Source/Runtime/Slate/Private/Framework/Layout/InertialScrollManager.cpp b/Engine/Source/Runtime/Slate/Private/Framework/Layout/InertialScrollManager.cpp
--- a/Engine/Source/Runtime/Slate/Private/Framework/Layout/InertialScrollManager.cpp
+++ b/Engine/Source/Runtime/Slate/Private/Framework/Layout/InertialScrollManager.cpp
@@ -1,9 +1,21 @@
 // Copyright Epic Games, Inc. All Rights Reserved.
 
 #include "Framework/Layout/InertialScrollManager.h"
+#include "HAL/IConsoleManager.h"
 
-float FInertialScrollManager::FrictionCoefficient = 2.0f;
-float FInertialScrollManager::StaticVelocityDrag = 100;
+float FrictionCoefficient = 2.0f;
+FAutoConsoleVariableRef CVarFrictionCoefficient(
+	TEXT("Slate.InertialScroll.FrictionCoefficient"),
+	FrictionCoefficient,
+	TEXT("This is the percentage of velocity loss per second."),
+	ECVF_Default);
+
+float StaticVelocityDrag = 100;
+FAutoConsoleVariableRef CVarStaticVelocityDrag(
+	TEXT("Slate.InertialScroll.StaticVelocityDrag"),
+	StaticVelocityDrag,
+	TEXT("This is a constant amount of velocity lost per second."),
+	ECVF_Default);
 
 FInertialScrollManager::FInertialScrollManager(double InSampleTimeout)
 	: ScrollVelocity(0.f)
diff --git a/Engine/Source/Runtime/Slate/Public/Framework/Layout/InertialScrollManager.h b/Engine/Source/Runtime/Slate/Public/Framework/Layout/InertialScrollManager.h
--- a/Engine/Source/Runtime/Slate/Public/Framework/Layout/InertialScrollManager.h
+++ b/Engine/Source/Runtime/Slate/Public/Framework/Layout/InertialScrollManager.h
@@ -13,15 +13,6 @@
  */
 class FInertialScrollManager
 {
-	/**
-	 * This is the percentage of velocity loss per second.
-	 */
-	static SLATE_API float FrictionCoefficient;
-
-	/**
-	 * This is a constant amount of velocity lost per second.
-	 */
-	static SLATE_API float StaticVelocityDrag;
 
 public:
 	/**
