Fast follow up fixes to SVectorInputBox
* Only apply swizzle to Vector3s for now - to follow up to narrow it down more to just translations
* Color attribute applied everywhere

diff --git a/Engine/Source/Runtime/Slate/Public/Widgets/Input/SVectorInputBox.h b/Engine/Source/Runtime/Slate/Public/Widgets/Input/SVectorInputBox.h
--- a/Engine/Source/Runtime/Slate/Public/Widgets/Input/SVectorInputBox.h
+++ b/Engine/Source/Runtime/Slate/Public/Widgets/Input/SVectorInputBox.h
@@ -62,8 +62,8 @@ private:
 				{
 					_XDisplayName.Set(AxisDisplayInfo::GetAxisDisplayName(EAxisList::X));
 				}
-				_XColor.Set(AxisDisplayInfo::GetAxisColor(EAxisList::X));
 			}
+			_XColor.Set(AxisDisplayInfo::GetAxisColor(EAxisList::X));
 		}
 
 		/** X Component of the vector */
@@ -112,8 +112,8 @@ private:
 				{
 					_YDisplayName.Set(AxisDisplayInfo::GetAxisDisplayName(EAxisList::Y));
 				}
-				_YColor.Set(AxisDisplayInfo::GetAxisColor(EAxisList::Y));
 			}
+			_YColor.Set(AxisDisplayInfo::GetAxisColor(EAxisList::Y));
 		}
 
 		/** Y Component of the vector */
@@ -162,8 +162,8 @@ private:
 				{
 					_ZDisplayName.Set(AxisDisplayInfo::GetAxisDisplayName(EAxisList::Z));
 				}
-				_ZColor.Set(AxisDisplayInfo::GetAxisColor(EAxisList::Z));
 			}
+			_ZColor.Set(AxisDisplayInfo::GetAxisColor(EAxisList::Z));
 		}
 
 		/** Z Component of the vector */
@@ -259,11 +259,19 @@ public:
 			, _SpinDelta(1)
 			, _LinearDeltaSensitivity(1)
 			, _bColorAxisLabels(false)
-			, _Swizzle(AxisDisplayInfo::DefaultAxisComponentDisplaySwizzle()) /* X, Y, Z, W */
 			, _DisplayToggle(false)
 			, _TogglePadding(FMargin(1.f,0.f,1.f,0.f) )
 			, _PreventThrottling(false)
-		{}
+		{
+			if constexpr (NumberOfComponents == 3)
+			{
+				_Swizzle = AxisDisplayInfo::DefaultAxisComponentDisplaySwizzle();
+			}
+			else
+			{
+				_Swizzle = FIntVector4(0, 1, 2, 3);
+			}
+		}
 
 		/** Optional Value of the vector */
 		SLATE_ATTRIBUTE(TOptional<VectorType>, Vector)
