Added missing init of bUseCopyToClipboard in the other SNotificationItem constructor.

diff --git a/Engine/Source/Runtime/Slate/Public/Widgets/Notifications/SNotificationList.h b/Engine/Source/Runtime/Slate/Public/Widgets/Notifications/SNotificationList.h
--- a/Engine/Source/Runtime/Slate/Public/Widgets/Notifications/SNotificationList.h
+++ b/Engine/Source/Runtime/Slate/Public/Widgets/Notifications/SNotificationList.h
@@ -143,6 +143,7 @@ struct FNotificationInfo
 		bUseThrobber(true),
 		bUseSuccessFailIcons(true),
 		bUseLargeFont(true),
+		bUseCopyToClipboard(false),
 		WidthOverride(320.0f),
 		bFireAndForget(true),
 		CheckBoxState(ECheckBoxState::Unchecked),
