Add "Close Tabs to the Left" to SDockingTabStack

Add new Enum ECloseTabsDirection
Rename CloseTabsToRightOfForegroundTab to CloseTabsInDirectionFromForegroundTab
Rename CanCloseTabsToRightOfForegroundTab to CanCloseTabsInDirectionOfForegroundTab

diff --git a/Engine/Source/Runtime/Slate/Private/Framework/Docking/SDockingTabStack.cpp b/Engine/Source/Runtime/Slate/Private/Framework/Docking/SDockingTabStack.cpp
--- a/Engine/Source/Runtime/Slate/Private/Framework/Docking/SDockingTabStack.cpp
+++ b/Engine/Source/Runtime/Slate/Private/Framework/Docking/SDockingTabStack.cpp
@@ -562,26 +562,28 @@ void SDockingTabStack::CloseForegroundTab()
 	}
 }
 
-void SDockingTabStack::CloseTabsToRightOfForegroundTab(ETabsToClose TabsToClose)
+void SDockingTabStack::CloseTabsInDirectionFromForegroundTab(ETabsToClose TabsToClose, ECloseTabsDirection Direction)
 {
 	TSharedPtr<SDockTab> ForegroundTab = TabWell->GetForegroundTab();
 	if (ForegroundTab.IsValid())
 	{
-		int32 DestroyIndex = TabWell->GetForegroundTabIndex() + 1;
-		int NumTabsToClose = TabWell->GetNumTabs() - TabWell->GetForegroundTabIndex() - 1;
-		while ((NumTabsToClose > 0) && (DestroyIndex < TabWell->GetNumTabs()))
+		const int32 ForegroundTabIndex = TabWell->GetForegroundTabIndex();
+		const int32 Increment = GetIncrementFromCloseTabsDirection(Direction);
+
+		int32 DestroyIndex = ForegroundTabIndex + Increment;
+		while (DestroyIndex >= 0 && DestroyIndex < TabWell->GetNumTabs())
 		{
 			const TSharedRef<SDockTab>& Tab = TabWell->GetTabs()[DestroyIndex];
-
 			const ETabRole VisualTabRole = Tab->GetVisualTabRole();
 			const bool bCanClose =
-				(TabsToClose == CloseAllTabs) ||
-				(TabsToClose == CloseDocumentTabs && VisualTabRole == ETabRole::DocumentTab) ||
-				(TabsToClose == CloseDocumentAndMajorTabs && (VisualTabRole == ETabRole::DocumentTab || VisualTabRole == ETabRole::MajorTab));
+				(TabsToClose == ETabsToClose::CloseAllTabs) ||
+				(TabsToClose == ETabsToClose::CloseDocumentTabs && VisualTabRole == ETabRole::DocumentTab) ||
+				(TabsToClose == ETabsToClose::CloseDocumentAndMajorTabs && (VisualTabRole == ETabRole::DocumentTab || VisualTabRole == ETabRole::MajorTab));
 
-			if (!bCanClose || !Tab->RequestCloseTab())
+			// Need to increment index if we failed to close a tab to the right or if we're closing tabs to the left regardless of if we succeeded
+			if (!bCanClose || !Tab->RequestCloseTab() || Direction == ECloseTabsDirection::Left)
 			{
-				++DestroyIndex;
+				DestroyIndex += Increment;
 			}
 		}
 	}
@@ -599,9 +601,9 @@ void SDockingTabStack::CloseAllButForegroundTab(ETabsToClose TabsToClose)
 
 			const ETabRole VisualTabRole = Tab->GetVisualTabRole();
 			const bool bCanClose = 
-				(TabsToClose == CloseAllTabs) ||
-				(TabsToClose == CloseDocumentTabs && VisualTabRole == ETabRole::DocumentTab) ||
-				(TabsToClose == CloseDocumentAndMajorTabs && (VisualTabRole == ETabRole::DocumentTab || VisualTabRole == ETabRole::MajorTab));
+				(TabsToClose == ETabsToClose::CloseAllTabs) ||
+				(TabsToClose == ETabsToClose::CloseDocumentTabs && VisualTabRole == ETabRole::DocumentTab) ||
+				(TabsToClose == ETabsToClose::CloseDocumentAndMajorTabs && (VisualTabRole == ETabRole::DocumentTab || VisualTabRole == ETabRole::MajorTab));
 
 			if ((Tab == ForegroundTab) || !bCanClose || !Tab->RequestCloseTab())
 			{
@@ -691,15 +693,25 @@ TSharedRef<SWidget> SDockingTabStack::MakeContextMenu()
 				)
 			);
 
-			const ETabsToClose TabsToClose = CloseDocumentAndMajorTabs;
+			const ETabsToClose TabsToClose = ETabsToClose::CloseDocumentAndMajorTabs;
+
+			MenuBuilder.AddMenuEntry(
+				LOCTEXT("CloseTabsToTheLeft", "Close Tabs to the Left"),
+				LOCTEXT("CloseTabsToTheLeftTooltip", "Closes all tabs to the left of the active tab."),
+				FSlateIcon(),
+				FUIAction(
+					FExecuteAction::CreateSP( this, &SDockingTabStack::CloseTabsInDirectionFromForegroundTab, TabsToClose, ECloseTabsDirection::Left),
+					FCanExecuteAction::CreateSP( this, &SDockingTabStack::CanCloseTabsInDirectionOfForegroundTab, ECloseTabsDirection::Left)
+				)
+			);
 
 			MenuBuilder.AddMenuEntry(
 				LOCTEXT("CloseTabsToTheRight", "Close Tabs to the Right"),
 				LOCTEXT("CloseTabsToTheRightTooltip", "Closes all tabs to the right of the active tab."),
 				FSlateIcon(),
 				FUIAction(
-					FExecuteAction::CreateSP( this, &SDockingTabStack::CloseTabsToRightOfForegroundTab, TabsToClose ),
-					FCanExecuteAction::CreateSP( this, &SDockingTabStack::CanCloseTabsToRightOfForegroundTab )
+					FExecuteAction::CreateSP( this, &SDockingTabStack::CloseTabsInDirectionFromForegroundTab, TabsToClose, ECloseTabsDirection::Right),
+					FCanExecuteAction::CreateSP( this, &SDockingTabStack::CanCloseTabsInDirectionOfForegroundTab, ECloseTabsDirection::Right)
 				)
 			);
 
@@ -1091,7 +1103,17 @@ bool SDockingTabStack::CanCloseForegroundTab() const
 	return ForegroundTabPtr.IsValid() && ForegroundTabPtr->CanCloseTab();
 }
 
-bool SDockingTabStack::CanCloseTabsToRightOfForegroundTab() const
+int32 SDockingTabStack::GetIncrementFromCloseTabsDirection(ECloseTabsDirection Direction)
+{
+	switch (Direction)
+	{
+		case ECloseTabsDirection::Left: 	return -1;
+		case ECloseTabsDirection::Right:	return  1;
+		default: 							return  1;
+	}
+}
+
+bool SDockingTabStack::CanCloseTabsInDirectionOfForegroundTab(ECloseTabsDirection Direction) const
 {
 	TSharedPtr<SDockTab> ForegroundTabPtr = TabWell->GetForegroundTab();
 	if (!ForegroundTabPtr.IsValid())
@@ -1100,10 +1122,11 @@ bool SDockingTabStack::CanCloseTabsToRightOfForegroundTab() const
 	}
 
 	const ETabRole VisualTabRole = ForegroundTabPtr->GetVisualTabRole();
-	if ((VisualTabRole == ETabRole::DocumentTab || VisualTabRole == ETabRole::MajorTab) && (TabWell->GetNumTabs() > 1) && (TabWell->GetForegroundTabIndex() != TabWell->GetNumTabs() - 1))
+	if (VisualTabRole == ETabRole::DocumentTab || VisualTabRole == ETabRole::MajorTab)
 	{
 		const TArray< TSharedRef<SDockTab> > MyTabs = this->GetTabs().AsArrayCopy();
-		for (int32 TabIndex = TabWell->GetForegroundTabIndex() + 1; TabIndex < MyTabs.Num(); ++TabIndex)
+		const int32 Increment = GetIncrementFromCloseTabsDirection(Direction);
+		for (int32 TabIndex = TabWell->GetForegroundTabIndex() + Increment; TabIndex >= 0 && TabIndex < MyTabs.Num(); TabIndex += Increment)
 		{
 			TSharedRef<SDockTab> Tab = MyTabs[TabIndex];
 			if (Tab->CanCloseTab())
diff --git a/Engine/Source/Runtime/Slate/Private/Framework/Docking/SDockingTabStack.h b/Engine/Source/Runtime/Slate/Private/Framework/Docking/SDockingTabStack.h
--- a/Engine/Source/Runtime/Slate/Private/Framework/Docking/SDockingTabStack.h
+++ b/Engine/Source/Runtime/Slate/Private/Framework/Docking/SDockingTabStack.h
@@ -175,14 +175,20 @@ protected:
 
 	SLATE_API void CloseForegroundTab();
 
-	enum ETabsToClose
+	enum class ETabsToClose : uint8
 	{
 		CloseDocumentTabs,
 		CloseDocumentAndMajorTabs,
 		CloseAllTabs
 	};
 
-	SLATE_API void CloseTabsToRightOfForegroundTab(ETabsToClose TabsToClose);
+	enum class ECloseTabsDirection : uint8
+	{
+		Left,
+		Right
+	};
+
+	SLATE_API void CloseTabsInDirectionFromForegroundTab(ETabsToClose TabsToClose, ECloseTabsDirection Direction);
 
 	/**
 	 * Close all the background tabs.
@@ -249,8 +255,10 @@ private:
 	/** Only allow closing the tab well when the tab allows it. */
 	SLATE_API bool CanCloseForegroundTab() const;
 
+	static int32 GetIncrementFromCloseTabsDirection(ECloseTabsDirection);
+
 	/** Only allow closing tabs to the right when there is more than one tab open, the tab is of type Document or Major, and the tab is not furthest to the right. */
-	SLATE_API bool CanCloseTabsToRightOfForegroundTab() const;
+	SLATE_API bool CanCloseTabsInDirectionOfForegroundTab(ECloseTabsDirection Direction) const;
 
 	/** Only allow closing all other tabs when there are more then one tab open and the tab is of type Document or Major. */
 	SLATE_API bool CanCloseAllButForegroundTab() const;
