Restore the ability for nomad tabs to be placed in separate windows owned by major tabs.

#rnx

diff --git a/Engine/Source/Runtime/Slate/Private/Framework/Docking/FDockingDragOperation.cpp b/Engine/Source/Runtime/Slate/Private/Framework/Docking/FDockingDragOperation.cpp
--- a/Engine/Source/Runtime/Slate/Private/Framework/Docking/FDockingDragOperation.cpp
+++ b/Engine/Source/Runtime/Slate/Private/Framework/Docking/FDockingDragOperation.cpp
@@ -218,28 +218,15 @@ bool FDockingDragOperation::CanDockInNode(const TSharedRef<SDockingNode>& DockNo
 
 	if (TabBeingDragged->GetTabRole() == ETabRole::NomadTab)
 	{
-		if (TargetTabManager != FGlobalTabmanager::Get())
+		if ( IsDockingViaTabwell == FDockingDragOperation::DockingViaTabWell )
 		{
-			if (TabManagerOfOriginToUse == TargetTabManager && IsDockingViaTabwell == FDockingDragOperation::DockingViaTabWell)
-			{
-				return true;
-			}
-
-			const TSharedPtr<SDockingArea>& DockingArea = DockNode->GetDockArea();
-			if (DockingArea.IsValid())
-			{
-				if (!DockingArea->GetParentWindow() || DockingArea->GetParentWindow() == TabOwnerAreaOfOrigin->GetParentWindow())
-				{
-					// Nomad tab can be docked in any non-floating dock area or if the ParentWindow is the same
-					return true;
-				}
-			}
-			return false;
+			// Nomad tabs can be docked in in any tab well.
+			return true;
+		}
+		else
+		{
+			return TargetTabManager != FGlobalTabmanager::Get();
 		}
-
-		// Nomad tab can be docked to any GlobalTabManager area.
-		// This additional check is needed to not make the cross visible for the main root window of the engine, but to make it able anyway to dock it like a major tab.
-		return IsDockingViaTabwell == FDockingDragOperation::DockingViaTabWell || FGlobalTabmanager::Get()->GetRootWindow() != DockNode->GetDockArea()->GetParentWindow();
 	}
 	else if (TabBeingDragged->GetTabRole() == ETabRole::MajorTab)
 	{
