PR #12773: Added missing OnTabActivated initialization in SDockTab

#tests Ran the editor, displaying various SDockTab without issue
#rnx

diff --git a/Engine/Source/Runtime/Slate/Private/Widgets/Docking/SDockTab.cpp b/Engine/Source/Runtime/Slate/Private/Widgets/Docking/SDockTab.cpp
--- a/Engine/Source/Runtime/Slate/Private/Widgets/Docking/SDockTab.cpp
+++ b/Engine/Source/Runtime/Slate/Private/Widgets/Docking/SDockTab.cpp
@@ -566,6 +566,7 @@ void SDockTab::Construct( const FArguments& InArgs )
 	this->TabWellContentRight = InArgs._TabWellContentRight.Widget;
 	this->TabRole = InArgs._TabRole;
 	this->OnTabClosed = InArgs._OnTabClosed;
+	this->OnTabActivated = InArgs._OnTabActivated;
 	this->OnCanCloseTab = InArgs._OnCanCloseTab;
 	this->bCanEverClose = InArgs._CanEverClose;
 	this->OnPersistVisualState = InArgs._OnPersistVisualState;
