Integration from CDPR partner stream.
[Slate]: Fix problem where users can't type polish letters s and a (with a tail) into the text input
#tests Validated locally that it was still working correctly (while fixing an issue with my keyboard layout too)

diff --git a/Engine/Source/Runtime/Slate/Private/Widgets/Text/SlateEditableTextLayout.cpp b/Engine/Source/Runtime/Slate/Private/Widgets/Text/SlateEditableTextLayout.cpp
--- a/Engine/Source/Runtime/Slate/Private/Widgets/Text/SlateEditableTextLayout.cpp
+++ b/Engine/Source/Runtime/Slate/Private/Widgets/Text/SlateEditableTextLayout.cpp
@@ -1228,6 +1228,13 @@ FReply FSlateEditableTextLayout::HandleKeyDown(const FKeyEvent& InKeyEvent)
 		// Absorb this event so it is not bubbled and handled by other widgets that could have something bound to the key press.
 		Reply = FReply::Handled();
 	}
+
+	else if (InKeyEvent.IsRightAltDown() && InKeyEvent.GetCharacter() != 0)
+	{
+		// Alt Gr button and a single character was pressed. We will type something in an upcoming OnKeyChar event.  
+		// Absorb this event so it is not bubbled and handled by other widgets that could have something bound to the key press.
+		Reply = FReply::Handled();
+	}
 	
 	if (!Reply.IsEventHandled())
 	{
