Inertial Scrolling
- Use the average DeltaTime when using animated scrolling in STableViewBase.

#rnx

diff --git a/Engine/Source/Runtime/Slate/Private/Widgets/Views/STableViewBase.cpp b/Engine/Source/Runtime/Slate/Private/Widgets/Views/STableViewBase.cpp
--- a/Engine/Source/Runtime/Slate/Private/Widgets/Views/STableViewBase.cpp
+++ b/Engine/Source/Runtime/Slate/Private/Widgets/Views/STableViewBase.cpp
@@ -19,6 +19,20 @@
 
 #include UE_INLINE_GENERATED_CPP_BY_NAME(STableViewBase)
 
+int32 MaxNumberOfDeltaTimes = 20;
+FAutoConsoleVariableRef CVarMaxNumberOfDeltaTimes(
+	TEXT("Slate.InertialScrolling.MaxNumberOfDeltaTimes"),
+	MaxNumberOfDeltaTimes,
+	TEXT("The number of DeltaTimes used to calculate the smooth DeltaTime used for animated scrolling."),
+	ECVF_Default);
+
+float RatioOfOutliers = 0.4f;
+FAutoConsoleVariableRef CVarRatioOfOutliers(
+	TEXT("Slate.InertialScrolling.RatioOfOutliers"),
+	RatioOfOutliers,
+	TEXT("The ratio of outliers to remove from DeltaTimeQueue for the smooth DeltaTime calculation."),
+	ECVF_Default);
+
 namespace ListConstants
 {
 	static const float OvershootMax = 150.0f;
@@ -236,7 +250,8 @@ EActiveTimerReturnType STableViewBase::UpdateInertialScroll(double InCurrentTime
 		}
 		else
 		{
-			InertialScrollManager.UpdateScrollVelocity(InDeltaTime);
+			const float CurrentSmoothDeltaTime = GetSmoothDeltaTime(InDeltaTime);
+			InertialScrollManager.UpdateScrollVelocity(CurrentSmoothDeltaTime);
 			const float ScrollVelocity = InertialScrollManager.GetScrollVelocity();
 
 			if (ScrollVelocity != 0.f)
@@ -244,7 +259,7 @@ EActiveTimerReturnType STableViewBase::UpdateInertialScroll(double InCurrentTime
 				if (CanUseInertialScroll(ScrollVelocity))
 				{
 					bKeepTicking = true;
-					ScrollBy(GetTickSpaceGeometry(), ScrollVelocity * InDeltaTime, AllowOverscroll);
+					ScrollBy(GetTickSpaceGeometry(), ScrollVelocity * CurrentSmoothDeltaTime, AllowOverscroll);
 				}
 				else
 				{
@@ -291,6 +306,8 @@ void STableViewBase::Tick( const FGeometry& AllottedGeometry, const double InCur
 		
 		if ( bItemsNeedRefresh || bPanelGeometryChanged)
 		{
+			UpdateSmoothDeltaTime(InDeltaTime);
+
 			PanelGeometryLastTick = PanelGeometry;
 			
 			const int32 NumItemsPerLine = GetNumItemsPerLine();
@@ -307,7 +324,7 @@ void STableViewBase::Tick( const FGeometry& AllottedGeometry, const double InCur
 			bool bEndScroll = true;
 			if ((bStartedTouchInteraction && bEnableTouchAnimatedScrolling) || (!bStartedTouchInteraction && bEnableAnimatedScrolling))
 			{
-				CurrentScrollOffset = FMath::FInterpTo(CurrentScrollOffset, TargetScrollOffset, (double)InDeltaTime, ScrollingAnimationInterpolationSpeed);
+				CurrentScrollOffset = FMath::FInterpTo(CurrentScrollOffset, TargetScrollOffset, (double)GetSmoothDeltaTime(InDeltaTime), ScrollingAnimationInterpolationSpeed);
 				// We may need to wait until a new target scroll offset is recalculated at 0 velocity before we can end the scroll
 				const bool bWaitForZeroVelocityAutoScroll = ShouldScrollToFixedLineOffsetAtZeroVelocity() && InertialScrollManager.GetScrollVelocity() != 0.f;
 				bEndScroll = !bWaitForZeroVelocityAutoScroll && FMath::IsNearlyEqual(CurrentScrollOffset, TargetScrollOffset, 0.01);
@@ -866,6 +883,46 @@ bool STableViewBase::ShouldScrollToFixedLineOffsetAtZeroVelocity() const
 	return !bWasAtEndOfList && FixedLineScrollOffset.IsSet() && !IsRightClickScrolling();
 }
 
+float STableViewBase::GetSmoothDeltaTime(float InDeltaTime)
+{
+	// This is to be sure we have an initial value.
+	return SmoothDeltaTime >= 0 ? SmoothDeltaTime : InDeltaTime;
+}
+
+void STableViewBase::UpdateSmoothDeltaTime(float InDeltaTime)
+{
+	if (DeltaTimeQueue.IsValidIndex(DeltaTimeCircularIndex))
+	{
+		DeltaTimeQueue[DeltaTimeCircularIndex] = InDeltaTime;
+	}
+	else
+	{
+		DeltaTimeQueue.Add(InDeltaTime);
+	}
+	DeltaTimeCircularIndex = ++DeltaTimeCircularIndex % MaxNumberOfDeltaTimes;
+	
+	SortedDeltaTimes.Empty(MaxNumberOfDeltaTimes);
+	SortedDeltaTimes.Append(DeltaTimeQueue);
+	SortedDeltaTimes.Sort();
+
+	float SumOfDeltaTimes = 0.0f;
+
+	const int32 NumberOfOutliers = FMath::RoundFromZero(SortedDeltaTimes.Num() * FMath::Abs(RatioOfOutliers));
+	const int32 HalfNumberOfOutliers = NumberOfOutliers / 2;
+	
+	const int32 StartIndex = HalfNumberOfOutliers;
+	const int32 EndIndex = SortedDeltaTimes.Num() - HalfNumberOfOutliers;
+
+	ensure(EndIndex - StartIndex > 0);
+
+	for (int32 i = StartIndex; i < EndIndex; ++i)
+	{
+		SumOfDeltaTimes += SortedDeltaTimes[i];
+	}
+
+	SmoothDeltaTime = SumOfDeltaTimes / (EndIndex - StartIndex);
+}
+
 float STableViewBase::ScrollBy(const FGeometry& MyGeometry, float ScrollByAmountInSlateUnits, EAllowOverscroll InAllowOverscroll)
 {
 	const int32 NumItemsBeingObserved = GetNumItemsBeingObserved();
diff --git a/Engine/Source/Runtime/Slate/Public/Widgets/Views/STableViewBase.h b/Engine/Source/Runtime/Slate/Public/Widgets/Views/STableViewBase.h
--- a/Engine/Source/Runtime/Slate/Public/Widgets/Views/STableViewBase.h
+++ b/Engine/Source/Runtime/Slate/Public/Widgets/Views/STableViewBase.h
@@ -573,6 +573,9 @@ protected:
 	/** How we should handle scrolling with the mouse wheel */
 	EConsumeMouseWheel ConsumeMouseWheel;
 
+	float GetSmoothDeltaTime(float InDeltaTime);
+	void UpdateSmoothDeltaTime(float InDeltaTime);
+
 private:
 	/** When true, a refresh should occur the next tick */
 	bool bItemsNeedRefresh = false;
@@ -583,6 +586,13 @@ private:
 	/** Whether or not the current scroll is inertial */
 	bool bIsInertialScroll = false;
 
+	TArray<float> DeltaTimeQueue;
+	int32 DeltaTimeCircularIndex = 0;
+	TArray<float> SortedDeltaTimes;
+
+	// Initialized at -1.0f to detect if the variable has been set and is usable since DeltaTimes cannot be negative.
+	float SmoothDeltaTime = -1.0f;
+
 	union
 	{
 		// vertical scroll bar is stored in horizontal box and vice versa
