SUniformGridPanel now properly collapses rows and columns that only contain collapsed widgets when they precede rows and columns with non-collapsed widgets.

[RN]

diff --git a/Engine/Source/Runtime/Slate/Private/Widgets/Layout/SUniformGridPanel.cpp b/Engine/Source/Runtime/Slate/Private/Widgets/Layout/SUniformGridPanel.cpp
--- a/Engine/Source/Runtime/Slate/Private/Widgets/Layout/SUniformGridPanel.cpp
+++ b/Engine/Source/Runtime/Slate/Private/Widgets/Layout/SUniformGridPanel.cpp
@@ -3,6 +3,26 @@
 #include "Widgets/Layout/SUniformGridPanel.h"
 #include "Layout/LayoutUtils.h"
 
+namespace UE::Slate::Private
+{
+
+// Get the column/row index minus the count of preceding collapsed columns/rows
+int32 GetAdjustedIndex(const int32 OriginalIndex, const TSet<int32>& CollapsedSet)
+{
+	int32 ResultIndex = OriginalIndex;
+	for (int32 CheckIndex = 0; CheckIndex < OriginalIndex; CheckIndex++)
+	{
+		if (CollapsedSet.Contains(CheckIndex))
+		{
+			--ResultIndex;
+		}
+	}
+
+	return ResultIndex;
+}
+
+}
+
 SUniformGridPanel::SUniformGridPanel()
 	: Children(this)
 	, SlotPadding(*this, FMargin(0.0f))
@@ -48,7 +68,8 @@ void SUniformGridPanel::OnArrangeChildren( const FGeometry& AllottedGeometry, FA
 
 				ArrangedChildren.AddWidget(ChildVisibility,
 					AllottedGeometry.MakeChild(Child.GetWidget(),
-					FVector2D(CellSize.X*Child.GetColumn() + XAxisResult.Offset, CellSize.Y*Child.GetRow() + YAxisResult.Offset),
+					FVector2D(CellSize.X * UE::Slate::Private::GetAdjustedIndex(Child.GetColumn(), CollapsedColumns) + XAxisResult.Offset,
+					          CellSize.Y * UE::Slate::Private::GetAdjustedIndex(Child.GetRow(), CollapsedRows) + YAxisResult.Offset),
 					FVector2D(XAxisResult.Size, YAxisResult.Size)
 					));
 			}
@@ -68,26 +89,56 @@ FVector2D SUniformGridPanel::ComputeDesiredSize( float ) const
 	NumColumns = 0;
 	NumRows = 0;
 
-	for ( int32 ChildIndex=0; ChildIndex < Children.Num(); ++ChildIndex )
+	CollapsedColumns.Reset();
+	CollapsedRows.Reset();
+
+	// Track rows/columns that have been verified as containing visible widgets
+	TSet<int32> VisibleColumns;
+	TSet<int32> VisibleRows;
+
+	for (int32 ChildIndex=0; ChildIndex < Children.Num(); ++ChildIndex)
 	{
-		const FSlot& Child = Children[ ChildIndex ];
+		const FSlot& Child = Children[ChildIndex];
 
-		if (Child.GetWidget()->GetVisibility() != EVisibility::Collapsed)
-		{
-			// A single cell at (N,M) means our grid size is (N+1, M+1)
-			NumColumns = FMath::Max(Child.GetColumn() + 1, NumColumns);
-			NumRows = FMath::Max(Child.GetRow() + 1, NumRows);
+		// A single cell at (N,M) means our grid size is (N+1, M+1)
+		NumColumns = FMath::Max(Child.GetColumn() + 1, NumColumns);
+		NumRows = FMath::Max(Child.GetRow() + 1, NumRows);
 
-			FVector2D ChildDesiredSize = Child.GetWidget()->GetDesiredSize() + SlotPaddingDesiredSize;
+		// If collapsed, we may want to collapse the entire row/column
+		if (Child.GetWidget()->GetVisibility() == EVisibility::Collapsed)
+		{
+			if (!VisibleColumns.Contains(Child.GetColumn()))
+			{
+				CollapsedColumns.Add(Child.GetColumn());
+			}
 
-			ChildDesiredSize.X = FMath::Max( ChildDesiredSize.X, CachedMinDesiredSlotWidth);
-			ChildDesiredSize.Y = FMath::Max( ChildDesiredSize.Y, CachedMinDesiredSlotHeight);
+			if (!VisibleRows.Contains(Child.GetRow()))
+			{
+				CollapsedRows.Add(Child.GetRow());
+			}
 
-			MaxChildDesiredSize.X = FMath::Max( MaxChildDesiredSize.X, ChildDesiredSize.X );
-			MaxChildDesiredSize.Y = FMath::Max( MaxChildDesiredSize.Y, ChildDesiredSize.Y );
+			continue;
 		}
+
+		// Verify the row & column as visible
+		CollapsedColumns.Remove(Child.GetColumn());
+		CollapsedRows.Remove(Child.GetRow());
+		VisibleColumns.Add(Child.GetColumn());
+		VisibleRows.Add(Child.GetRow());
+
+		FVector2D ChildDesiredSize = Child.GetWidget()->GetDesiredSize() + SlotPaddingDesiredSize;
+
+		ChildDesiredSize.X = FMath::Max(ChildDesiredSize.X, CachedMinDesiredSlotWidth);
+		ChildDesiredSize.Y = FMath::Max(ChildDesiredSize.Y, CachedMinDesiredSlotHeight);
+
+		MaxChildDesiredSize.X = FMath::Max(MaxChildDesiredSize.X, ChildDesiredSize.X);
+		MaxChildDesiredSize.Y = FMath::Max(MaxChildDesiredSize.Y, ChildDesiredSize.Y);
 	}
 
+	// Final row & column count shouldn't consider collapsed rows & columns
+	NumColumns -= CollapsedColumns.Num();
+	NumRows -= CollapsedRows.Num();
+
 	return FVector2D( NumColumns*MaxChildDesiredSize.X, NumRows*MaxChildDesiredSize.Y );
 }
 
diff --git a/Engine/Source/Runtime/Slate/Public/Widgets/Layout/SUniformGridPanel.h b/Engine/Source/Runtime/Slate/Public/Widgets/Layout/SUniformGridPanel.h
--- a/Engine/Source/Runtime/Slate/Public/Widgets/Layout/SUniformGridPanel.h
+++ b/Engine/Source/Runtime/Slate/Public/Widgets/Layout/SUniformGridPanel.h
@@ -148,6 +148,8 @@ private:
 	/** These values are recomputed and cached during compute desired size, as they may have changed since the previous frame. */
 	mutable int32 NumColumns;
 	mutable int32 NumRows;
+	mutable TSet<int32> CollapsedColumns;
+	mutable TSet<int32> CollapsedRows;
 
 	TSlateAttribute<float, EInvalidateWidgetReason::Layout> MinDesiredSlotWidth;
 	TSlateAttribute<float, EInvalidateWidgetReason::Layout> MinDesiredSlotHeight;
