Fix warning V567: The modification of the 'DeltaTimeCircularIndex' variable is unsequenced relative to another operation on the same variable. This may lead to undefined behavior.
#rnx

diff --git a/Engine/Source/Runtime/Slate/Private/Widgets/Views/STableViewBase.cpp b/Engine/Source/Runtime/Slate/Private/Widgets/Views/STableViewBase.cpp
--- a/Engine/Source/Runtime/Slate/Private/Widgets/Views/STableViewBase.cpp
+++ b/Engine/Source/Runtime/Slate/Private/Widgets/Views/STableViewBase.cpp
@@ -899,7 +899,7 @@ void STableViewBase::UpdateSmoothDeltaTime(float InDeltaTime)
 	{
 		DeltaTimeQueue.Add(InDeltaTime);
 	}
-	DeltaTimeCircularIndex = ++DeltaTimeCircularIndex % MaxNumberOfDeltaTimes;
+	DeltaTimeCircularIndex = (DeltaTimeCircularIndex + 1) % MaxNumberOfDeltaTimes;
 	
 	SortedDeltaTimes.Empty(MaxNumberOfDeltaTimes);
 	SortedDeltaTimes.Append(DeltaTimeQueue);
