Add AllowContentToShrink argument to SScrollBox to allow the underlying SScrollPanel to shrink with the overall size of the ScrollBox instead of always taking up maximum space.

[RN]

diff --git a/Engine/Source/Runtime/Slate/Private/Widgets/Layout/SScrollBox.cpp b/Engine/Source/Runtime/Slate/Private/Widgets/Layout/SScrollBox.cpp
--- a/Engine/Source/Runtime/Slate/Private/Widgets/Layout/SScrollBox.cpp
+++ b/Engine/Source/Runtime/Slate/Private/Widgets/Layout/SScrollBox.cpp
@@ -33,21 +33,21 @@ void SScrollPanel::Construct(const FArguments& InArgs, TArray<SScrollBox::FSlot:
 	Orientation = InArgs._Orientation;
 	BackPadScrolling = InArgs._BackPadScrolling;
 	FrontPadScrolling = InArgs._FrontPadScrolling;
+	bAllowContentToShrink = InArgs._AllowContentToShrink;
 }
 
 void SScrollPanel::OnArrangeChildren(const FGeometry& AllottedGeometry, FArrangedChildren& ArrangedChildren) const
 {
 	const float ScrollPadding = Orientation == Orient_Vertical ? AllottedGeometry.GetLocalSize().Y : AllottedGeometry.GetLocalSize().X;
 	const float ChildrenOffset = -PhysicalOffset + (BackPadScrolling ? ScrollPadding : 0);
-	const bool AllowShrink = false;
 
 	if (Orientation == EOrientation::Orient_Horizontal)
 	{
-		ArrangeChildrenInStack<EOrientation::Orient_Horizontal>(GSlateFlowDirection, this->Children, AllottedGeometry, ArrangedChildren, ChildrenOffset, AllowShrink);
+		ArrangeChildrenInStack<EOrientation::Orient_Horizontal, TPanelChildren<SScrollBox::FSlot>, SScrollBox::FSlot>(GSlateFlowDirection, this->Children, AllottedGeometry, ArrangedChildren, ChildrenOffset, bAllowContentToShrink, ArrangedSize);
 	}
 	else
 	{
-		ArrangeChildrenInStack<EOrientation::Orient_Vertical>(GSlateFlowDirection, this->Children, AllottedGeometry, ArrangedChildren, ChildrenOffset, AllowShrink);
+		ArrangeChildrenInStack<EOrientation::Orient_Vertical, TPanelChildren<SScrollBox::FSlot>, SScrollBox::FSlot>(GSlateFlowDirection, this->Children, AllottedGeometry, ArrangedChildren, ChildrenOffset, bAllowContentToShrink, ArrangedSize);
 	}
 }
 
@@ -57,19 +57,21 @@ FVector2D SScrollPanel::ComputeDesiredSize(float) const
 	for (int32 SlotIndex = 0; SlotIndex < Children.Num(); ++SlotIndex)
 	{
 		const SScrollBox::FSlot& ThisSlot = Children[SlotIndex];
-		if (ThisSlot.GetWidget()->GetVisibility() != EVisibility::Collapsed)
+		if (ThisSlot.GetWidget()->GetVisibility() == EVisibility::Collapsed)
 		{
-			const FVector2D ChildDesiredSize = ThisSlot.GetWidget()->GetDesiredSize();
-			if (Orientation == Orient_Vertical)
-			{
-				ThisDesiredSize.X = FMath::Max(ChildDesiredSize.X + ThisSlot.GetPadding().GetTotalSpaceAlong<Orient_Horizontal>(), ThisDesiredSize.X);
-				ThisDesiredSize.Y += ChildDesiredSize.Y + ThisSlot.GetPadding().GetTotalSpaceAlong<Orient_Vertical>();
-			}
-			else
-			{
-				ThisDesiredSize.X += ChildDesiredSize.X + ThisSlot.GetPadding().GetTotalSpaceAlong<Orient_Horizontal>();
-				ThisDesiredSize.Y = FMath::Max(ChildDesiredSize.Y + ThisSlot.GetPadding().GetTotalSpaceAlong<Orient_Vertical>(), ThisDesiredSize.Y);
-			}
+			continue;
+		}
+
+		const FVector2D ChildDesiredSize = ThisSlot.GetWidget()->GetDesiredSize();
+		if (Orientation == Orient_Vertical)
+		{
+			ThisDesiredSize.X = FMath::Max(ChildDesiredSize.X + ThisSlot.GetPadding().GetTotalSpaceAlong<Orient_Horizontal>(), ThisDesiredSize.X);
+			ThisDesiredSize.Y += ChildDesiredSize.Y + ThisSlot.GetPadding().GetTotalSpaceAlong<Orient_Vertical>();
+		}
+		else
+		{
+			ThisDesiredSize.X += ChildDesiredSize.X + ThisSlot.GetPadding().GetTotalSpaceAlong<Orient_Horizontal>();
+			ThisDesiredSize.Y = FMath::Max(ChildDesiredSize.Y + ThisSlot.GetPadding().GetTotalSpaceAlong<Orient_Vertical>(), ThisDesiredSize.Y);
 		}
 	}
 
@@ -126,6 +128,7 @@ void SScrollBox::Construct( const FArguments& InArgs )
 	bTouchPanningCapture = false;
 	bIsFocusable = false;
 	bVolatilityAlwaysInvalidatesPrepass = true;
+	bAllowContentToShrink = InArgs._AllowContentToShrink;
 
 	if (InArgs._ExternalScrollbar.IsValid())
 	{
@@ -341,6 +344,18 @@ float SScrollBox::GetScrollOffset() const
 	return DesiredScrollOffset;
 }
 
+UE::Slate::FDeprecateVector2DResult SScrollBox::GetScrollPanelContentSize() const
+{
+	UE::Slate::FDeprecateVector2DResult Size;
+	Size = ScrollPanel->GetArrangedSize();
+	if (!bAllowContentToShrink || Size.IsNearlyZero())
+	{
+		Size = ScrollPanel->GetDesiredSize();
+	}
+
+	return Size;
+}
+
 float SScrollBox::GetOverscrollOffset() const
 {
 	if (AllowOverscroll != EAllowOverscroll::Yes)
@@ -355,7 +370,7 @@ float SScrollBox::GetOverscrollOffset() const
 		return ScrollOffset;
 	}
 
-	const float ContentSize = GetScrollComponentFromVector(ScrollPanel->GetDesiredSize());
+	const float ContentSize = GetScrollComponentFromVector(GetScrollPanelContentSize());
 	const float ViewSize = GetScrollComponentFromVector(CachedGeometry.GetLocalSize());
 	if (ScrollOffset > ContentSize - ViewSize)
 	{
@@ -388,7 +403,7 @@ float SScrollBox::GetOverscrollPercentage() const
 		return (ScrollOffset / ViewSize) * 100.0f;
 	}
 
-	const float ContentSize = GetScrollComponentFromVector(ScrollPanel->GetDesiredSize());
+	const float ContentSize = GetScrollComponentFromVector(GetScrollPanelContentSize());
 	if (ScrollOffset > ContentSize - ViewSize)
 	{
 		// Positive overscroll (past the end of the ScrollBox)
@@ -403,14 +418,14 @@ float SScrollBox::GetOverscrollPercentage() const
 float SScrollBox::GetScrollOffsetOfEnd() const
 {
 	const FGeometry ScrollPanelGeometry = FindChildGeometry(CachedGeometry, ScrollPanel.ToSharedRef());
-	const float ContentSize = GetScrollComponentFromVector(ScrollPanel->GetDesiredSize());
+	const float ContentSize = GetScrollComponentFromVector(GetScrollPanelContentSize());
 	return FMath::Max(ContentSize - GetScrollComponentFromVector(ScrollPanelGeometry.Size), 0.0f);
 }
 
 float SScrollBox::GetViewFraction() const
 {
 	const FGeometry ScrollPanelGeometry = FindChildGeometry(CachedGeometry, ScrollPanel.ToSharedRef());
-	const float ContentSize = GetScrollComponentFromVector(ScrollPanel->GetDesiredSize());
+	const float ContentSize = GetScrollComponentFromVector(GetScrollPanelContentSize());
 
 	return FMath::Clamp<float>(GetScrollComponentFromVector(CachedGeometry.GetLocalSize()) > 0 ? GetScrollComponentFromVector(ScrollPanelGeometry.Size) / ContentSize : 1, 0.0f, 1.0f);
 }
@@ -418,7 +433,7 @@ float SScrollBox::GetViewFraction() const
 float SScrollBox::GetViewOffsetFraction() const
 {
 	const FGeometry ScrollPanelGeometry = FindChildGeometry(CachedGeometry, ScrollPanel.ToSharedRef());
-	const float ContentSize = GetScrollComponentFromVector(ScrollPanel->GetDesiredSize());
+	const float ContentSize = GetScrollComponentFromVector(GetScrollPanelContentSize());
 
 	const float ViewFraction = GetViewFraction();
 	return FMath::Clamp( DesiredScrollOffset/ContentSize, 0.f, 1.f - ViewFraction );
@@ -707,7 +722,10 @@ EActiveTimerReturnType SScrollBox::UpdateInertialScroll(double InCurrentTime, fl
 
 void SScrollBox::Tick( const FGeometry& AllottedGeometry, const double InCurrentTime, const float InDeltaTime )
 {
-	CachedGeometry = AllottedGeometry;
+	if (!bAllowContentToShrink)
+	{
+		CachedGeometry = AllottedGeometry;
+	}
 
 	if ( bTouchPanningCapture && (FSlateApplication::Get().GetCurrentTime() - LastScrollTime) > 0.10 )
 	{
@@ -717,19 +735,19 @@ void SScrollBox::Tick( const FGeometry& AllottedGeometry, const double InCurrent
 	// If we needed a widget to be scrolled into view, make that happen.
 	if ( ScrollIntoViewRequest )
 	{
-		ScrollIntoViewRequest(AllottedGeometry);
+		ScrollIntoViewRequest(CachedGeometry);
 		ScrollIntoViewRequest = nullptr;
 	}
-
-	const FGeometry ScrollPanelGeometry = FindChildGeometry( AllottedGeometry, ScrollPanel.ToSharedRef() );
-	const float ContentSize = GetScrollComponentFromVector(ScrollPanel->GetDesiredSize());
+	
+	const FGeometry ScrollPanelGeometry = FindChildGeometry(CachedGeometry, ScrollPanel.ToSharedRef() );
+	const float ContentSize = GetScrollComponentFromVector(GetScrollPanelContentSize());
 
 	if ( bScrollToEnd )
 	{
 		DesiredScrollOffset = FMath::Max(ContentSize - GetScrollComponentFromVector(ScrollPanelGeometry.GetLocalSize()), 0.0f);
 		bScrollToEnd = false;
 	}
-
+	
 	// If this scroll box has no size, do not compute a view fraction because it will be wrong and causes pop in when the size is available
 	const float ViewFraction = GetViewFraction();
 	const float TargetViewOffset = GetViewOffsetFraction();
@@ -737,10 +755,10 @@ void SScrollBox::Tick( const FGeometry& AllottedGeometry, const double InCurrent
 	const float CurrentViewOffset = bAnimateScroll ? FMath::FInterpTo(ScrollBar->DistanceFromTop(), TargetViewOffset, InDeltaTime, ScrollingAnimationInterpolationSpeed) : TargetViewOffset;
 
 	// Update the scrollbar with the clamped version of the offset
-	float NewPhysicalOffset = GetScrollComponentFromVector(CurrentViewOffset * ScrollPanel->GetDesiredSize());
+	float NewPhysicalOffset = GetScrollComponentFromVector(CurrentViewOffset * GetScrollPanelContentSize());
 	if ( AllowOverscroll == EAllowOverscroll::Yes )
 	{
-		NewPhysicalOffset += Overscroll.GetOverscroll(AllottedGeometry);
+		NewPhysicalOffset += Overscroll.GetOverscroll(CachedGeometry);
 	}
 
 	const bool bWasScrolling = bIsScrolling;
@@ -759,6 +777,11 @@ void SScrollBox::Tick( const FGeometry& AllottedGeometry, const double InCurrent
 		// We cannot scroll, so ensure that there is no offset.
 		ScrollPanel->PhysicalOffset = 0.0f;
 	}
+
+	if (bAllowContentToShrink)
+	{
+		CachedGeometry = AllottedGeometry;
+	}
 }
 
 bool SScrollBox::ComputeVolatility() const
@@ -1027,7 +1050,7 @@ bool SScrollBox::ScrollBy(const FGeometry& AllottedGeometry, float LocalScrollAm
 
 	bAnimateScroll = InAnimateScroll;
 
-	const float ContentSize = GetScrollComponentFromVector(ScrollPanel->GetDesiredSize());
+	const float ContentSize = GetScrollComponentFromVector(GetScrollPanelContentSize());
 	const FGeometry ScrollPanelGeometry = FindChildGeometry( AllottedGeometry, ScrollPanel.ToSharedRef() );
 
 	const float PreviousScrollOffset = DesiredScrollOffset;
@@ -1252,7 +1275,7 @@ void SScrollBox::ScrollBar_OnUserScrolled( float InScrollOffsetFraction )
 {
 	bAnimateScroll = false;
 
-	const float ContentSize = GetScrollComponentFromVector(ScrollPanel->GetDesiredSize());
+	const float ContentSize = GetScrollComponentFromVector(GetScrollPanelContentSize());
 	const FGeometry ScrollPanelGeometry = FindChildGeometry(CachedGeometry, ScrollPanel.ToSharedRef());
 
 	// Clamp to max scroll offset
@@ -1280,7 +1303,7 @@ FSlateColor SScrollBox::GetStartShadowOpacity() const
 FSlateColor SScrollBox::GetEndShadowOpacity() const
 {
 	// The shadow should only be visible when the user needs a hint that they can scroll down.
-	const float ShadowOpacity = (ScrollBar->DistanceFromBottom() * GetScrollComponentFromVector(ScrollPanel->GetDesiredSize()) / ShadowFadeDistance);
+	const float ShadowOpacity = (ScrollBar->DistanceFromBottom() * GetScrollComponentFromVector(GetScrollPanelContentSize()) / ShadowFadeDistance);
 	
 	return FLinearColor(1.0f, 1.0f, 1.0f, ShadowOpacity);
 }
diff --git a/Engine/Source/Runtime/Slate/Public/Widgets/Layout/SScrollBox.h b/Engine/Source/Runtime/Slate/Public/Widgets/Layout/SScrollBox.h
--- a/Engine/Source/Runtime/Slate/Public/Widgets/Layout/SScrollBox.h
+++ b/Engine/Source/Runtime/Slate/Public/Widgets/Layout/SScrollBox.h
@@ -171,6 +171,7 @@ public:
 		, _OnFocusReceived()
 		, _OnFocusLost()
 		, _ConsumeMouseWheel(EConsumeMouseWheel::WhenScrollingPossible)
+		, _AllowContentToShrink(false)
 		{
 			_Clipping = EWidgetClipping::ClipToBounds;
 		}
@@ -240,6 +241,9 @@ public:
 
 		SLATE_ARGUMENT(EConsumeMouseWheel, ConsumeMouseWheel);
 
+		/** If enabled, content in the underlying panel will be allowed to shrink with the overall size of the ScrollBox instead of always taking up maximum space */
+		SLATE_ARGUMENT(bool, AllowContentToShrink);
+
 	SLATE_END_ARGS()
 
 	SLATE_API SScrollBox();
@@ -292,6 +296,8 @@ public:
 	SLATE_API void SetScrollWhenFocusChanges(EScrollWhenFocusChanges NewScrollWhenFocusChanges);
 
 	SLATE_API float GetScrollOffset() const;
+
+	SLATE_API UE::Slate::FDeprecateVector2DResult GetScrollPanelContentSize() const;
 	
 	SLATE_API float GetOverscrollOffset() const;
 	
@@ -369,7 +375,7 @@ public:
 	SLATE_API virtual void OnMouseEnter( const FGeometry& MyGeometry, const FPointerEvent& MouseEvent ) override;
 	SLATE_API virtual void OnMouseLeave( const FPointerEvent& MouseEvent ) override;
 	SLATE_API virtual FReply OnMouseWheel( const FGeometry& MyGeometry, const FPointerEvent& MouseEvent ) override;
-	SLATE_API virtual FReply OnAnalogValueChanged( const FGeometry& MyGeometry, const FAnalogInputEvent& InAnalogInputEvent);
+	SLATE_API virtual FReply OnAnalogValueChanged( const FGeometry& MyGeometry, const FAnalogInputEvent& InAnalogInputEvent) override;
 	SLATE_API virtual FCursorReply OnCursorQuery( const FGeometry& MyGeometry, const FPointerEvent& CursorEvent ) const override;
 	SLATE_API virtual int32 OnPaint( const FPaintArgs& Args, const FGeometry& AllottedGeometry, const FSlateRect& MyCullingRect, FSlateWindowElementList& OutDrawElements, int32 LayerId, const FWidgetStyle& InWidgetStyle, bool bParentEnabled ) const override;
 	SLATE_API virtual FReply OnTouchEnded(const FGeometry& MyGeometry, const FPointerEvent& InTouchEvent) override;
@@ -421,7 +427,8 @@ private:
 	 * Scroll the view by ScrollAmount given its currently AllottedGeometry.
 	 *
 	 * @param AllottedGeometry  The geometry allotted for this SScrollBox by the parent
-	 * @param ScrollAmount      
+	 * @param LocalScrollAmount
+	 * @param Overscroll      
 	 * @param InAnimateScroll	Whether or not to animate the scroll
 	 * @return Whether or not the scroll was fully handled
 	 */
@@ -586,6 +593,9 @@ protected:
 	bool bTouchPanningCapture : 1;
 
 	bool bIsFocusable : 1;
+
+	/** Whether content in the underlying panel will be allowed to shrink with the overall size of the ScrollBox instead of always taking up maximum space */
+	bool bAllowContentToShrink = false;
 };
 
 class SScrollPanel : public SPanel
@@ -600,6 +610,7 @@ public:
 	SLATE_ARGUMENT(EOrientation, Orientation)
 	SLATE_ARGUMENT(bool, BackPadScrolling)
 	SLATE_ARGUMENT(bool, FrontPadScrolling)
+	SLATE_ARGUMENT(bool, AllowContentToShrink)
 
 	SLATE_END_ARGS()
 
@@ -629,6 +640,11 @@ public:
 		return &Children;
 	}
 
+	const FVector2D& GetArrangedSize() const
+	{
+		return ArrangedSize;
+	}
+
 	float PhysicalOffset;
 	TPanelChildren<SScrollBox::FSlot> Children;
 
@@ -642,4 +658,6 @@ private:
 	EOrientation Orientation;
 	bool BackPadScrolling;
 	bool FrontPadScrolling;
+	bool bAllowContentToShrink;
+	mutable FVector2D ArrangedSize;
 };
