Adding Shadow Box Styles into List View similar to Scroll Box

Added a bEnableShadowBrush to enable the ability to Shadow Boxes to the Top/Bottom or Left/Right of a list view
-

diff --git a/Engine/Source/Runtime/Slate/Private/Widgets/Views/STableViewBase.cpp b/Engine/Source/Runtime/Slate/Private/Widgets/Views/STableViewBase.cpp
--- a/Engine/Source/Runtime/Slate/Private/Widgets/Views/STableViewBase.cpp
+++ b/Engine/Source/Runtime/Slate/Private/Widgets/Views/STableViewBase.cpp
@@ -16,6 +16,7 @@
 #include "Framework/Layout/Overscroll.h"
 #include "Widgets/Views/STableRow.h"
 #include "Widgets/Views/SListPanel.h"
+#include "Widgets/Images/SImage.h"
 
 #include UE_INLINE_GENERATED_CPP_BY_NAME(STableViewBase)
 
@@ -66,6 +67,68 @@ FTableViewDimensions::FTableViewDimensions(EOrientation InOrientation, const UE:
 	}
 }
 
+TSharedRef<SWidget> STableViewBase::ConstructVerticalShadowBox(TSharedRef<SWidget> Content)
+{
+	return SNew(SOverlay)
+		+ SOverlay::Slot()
+		.Padding(0.0)
+		[
+			Content
+		]
+		+ SOverlay::Slot()
+		.HAlign(HAlign_Fill)
+		.VAlign(VAlign_Top)
+		[
+			// Shadow: Hint to scroll up
+			SNew(SImage)
+			.Visibility(EVisibility::HitTestInvisible)
+			.ColorAndOpacity(this, &STableViewBase::GetStartShadowOpacity)
+			.Image(&ShadowBoxStyle->TopShadowBrush)
+		]
+		+ SOverlay::Slot()
+		.HAlign(HAlign_Fill)
+		.VAlign(VAlign_Bottom)
+		[
+			// Shadow: a hint to scroll down
+			SNew(SImage)
+			.Visibility(EVisibility::HitTestInvisible)
+			.ColorAndOpacity(this, &STableViewBase::GetEndShadowOpacity)
+			.Image(&ShadowBoxStyle->BottomShadowBrush)
+		];
+}
+
+TSharedRef<SWidget> STableViewBase::ConstructHorizontalShadowBox(TSharedRef<SWidget> Content)
+{
+	return SNew(SOverlay)
+		+ SOverlay::Slot()
+		.Padding(0.0)
+		[
+			Content
+		]
+		+ SOverlay::Slot()
+		.HAlign(HAlign_Left)
+		.VAlign(VAlign_Fill)
+		[
+			// Shadow: Hint to left
+			SNew(SImage)
+			.Visibility(EVisibility::HitTestInvisible)
+			.ColorAndOpacity(this, &STableViewBase::GetStartShadowOpacity)
+			.Image(&ShadowBoxStyle->LeftShadowBrush)
+		]
+		+ SOverlay::Slot()
+		.HAlign(HAlign_Right)
+		.VAlign(VAlign_Fill)
+		[
+			// Shadow: a hint to scroll right
+			SNew(SImage)
+			.Visibility(EVisibility::HitTestInvisible)
+			.ColorAndOpacity(this, &STableViewBase::GetEndShadowOpacity)
+			.Image(&ShadowBoxStyle->RightShadowBrush)
+		];
+}
+
+
+
 void STableViewBase::ConstructChildren( const TAttribute<float>& InItemWidth, const TAttribute<float>& InItemHeight, const TAttribute<EListItemAlignment>& InItemAlignment, const TSharedPtr<SHeaderRow>& InHeaderRow, const TSharedPtr<SScrollBar>& InScrollBar, EOrientation InScrollOrientation, const FOnTableViewScrolled& InOnTableViewScrolled, const FScrollBarStyle* InScrollBarStyle, const bool bInPreventThrottling )
 {
 	bItemsNeedRefresh = true;
@@ -113,6 +176,32 @@ void STableViewBase::ConstructChildren( const TAttribute<float>& InItemWidth, co
 
 		const FOptionalSize ScrollBarSize(InScrollBarStyle ? InScrollBarStyle->Thickness + (SScrollBar::DefaultUniformPadding * 2.0f) : 16.f);
 
+		// We create the itemsPanel container
+		TSharedRef<SWidget> PanelContainer = SNew(SVerticalBox)
+			+SVerticalBox::Slot()
+			.AutoHeight()
+			[
+				PinnedItemsPanel.ToSharedRef()
+			]
+			+SVerticalBox::Slot()
+			.FillHeight(1)
+			[
+				ItemsPanel.ToSharedRef()
+			];		
+
+		// If we use Shadow Box Style, wrap the Container in a Shadow Box.
+		if (bShouldUseShadowBoxStyle && ShadowBoxStyle)
+		{
+			if (Orientation == Orient_Vertical)
+			{
+				PanelContainer = ConstructVerticalShadowBox(PanelContainer);
+			}
+			else if (Orientation == Orient_Horizontal)
+			{
+				PanelContainer = ConstructHorizontalShadowBox(PanelContainer);
+			}
+		}
+		
 		if (Orientation == Orient_Vertical)
 		{
 			VerticalScrollBarSlot = nullptr;
@@ -120,17 +209,7 @@ void STableViewBase::ConstructChildren( const TAttribute<float>& InItemWidth, co
 				+SHorizontalBox::Slot()
 				.FillWidth(1)
 				[
-					SNew(SVerticalBox)
-					+SVerticalBox::Slot()
-					.AutoHeight()
-					[
-						PinnedItemsPanel.ToSharedRef()
-					]
-					+SVerticalBox::Slot()
-					.FillHeight(1)
-					[
-						ItemsPanel.ToSharedRef()
-					]
+					PanelContainer
 				]
 				+SHorizontalBox::Slot()
 				.AutoWidth()
@@ -151,17 +230,7 @@ void STableViewBase::ConstructChildren( const TAttribute<float>& InItemWidth, co
 				+SVerticalBox::Slot()
 				.FillHeight(1)
 				[
-					SNew(SVerticalBox)
-					+SVerticalBox::Slot()
-					.AutoHeight()
-					[
-						PinnedItemsPanel.ToSharedRef()
-					]
-					+SVerticalBox::Slot()
-					.FillHeight(1)
-					[
-						ItemsPanel.ToSharedRef()
-					]
+					PanelContainer
 				]
 				+SVerticalBox::Slot()
 				.AutoHeight()
@@ -211,14 +280,12 @@ void STableViewBase::ConstructChildren( const TAttribute<float>& InItemWidth, co
 	
 }
 
-
 bool STableViewBase::SupportsKeyboardFocus() const
 {
 	// The ListView is focusable.
 	return true;
 }
 
-
 void STableViewBase::OnFocusLost( const FFocusEvent& InFocusEvent )
 {
 	bShowSoftwareCursor = false;
@@ -834,6 +901,8 @@ STableViewBase::STableViewBase( ETableViewMode::Type InTableViewMode )
 	, SoftwareCursorPosition( ForceInitToZero )
 	, bShowSoftwareCursor( false )
 	, WheelScrollMultiplier(GetGlobalScrollAmount())
+	, bShouldUseShadowBoxStyle(false)
+	, ShadowBoxStyle( &FAppStyle::Get().GetWidgetStyle<FScrollBoxStyle>("ScrollBox"))
 	, BackgroundBrush(FStyleDefaults::GetNoBrush())
 	, bIsScrollingActiveTimerRegistered( false )
 	, Overscroll()
@@ -1293,6 +1362,24 @@ void STableViewBase::OnClippingChanged()
 	ItemsPanel->SetClipping(GetClipping());
 }
 
+const float ShadowFadeDistance = 32.0f;
+
+FSlateColor STableViewBase::GetStartShadowOpacity() const
+{
+	// The shadow should only be visible when the user needs a hint that they can scroll up.
+	const float ShadowOpacity = FMath::Clamp( CurrentScrollOffset/ShadowFadeDistance, 0.0f, 1.0f);
+	
+	return FLinearColor(1.0f, 1.0f, 1.0f, ShadowOpacity);
+}
+
+FSlateColor STableViewBase::GetEndShadowOpacity() const
+{	
+	// The shadow should only be visible when the user needs a hint that they can scroll down.
+	const float ShadowOpacity = (ScrollBar->DistanceFromBottom() * GetScrollComponentFromVector(ItemsPanel->GetDesiredSize()) / ShadowFadeDistance);
+	
+	return FLinearColor(1.0f, 1.0f, 1.0f, ShadowOpacity);
+}
+
 bool STableViewBase::CanUseInertialScroll( float ScrollAmount ) const
 {
 	const auto CurrentOverscroll = Overscroll.GetOverscroll(GetTickSpaceGeometry());
diff --git a/Engine/Source/Runtime/Slate/Public/Widgets/Views/SListView.h b/Engine/Source/Runtime/Slate/Public/Widgets/Views/SListView.h
--- a/Engine/Source/Runtime/Slate/Public/Widgets/Views/SListView.h
+++ b/Engine/Source/Runtime/Slate/Public/Widgets/Views/SListView.h
@@ -88,6 +88,7 @@ public:
 public:
 	SLATE_BEGIN_ARGS(SListView<ItemType>)
 		: _ListViewStyle(&FAppStyle::Get().GetWidgetStyle<FTableViewStyle>("ListView"))
+		, _bEnableShadowBoxStyle(false)
 		, _OnGenerateRow()
 		, _OnGeneratePinnedRow()
 		, _OnRefreshRow()
@@ -126,6 +127,10 @@ public:
 		}
 
 		SLATE_STYLE_ARGUMENT( FTableViewStyle, ListViewStyle )
+		
+		SLATE_STYLE_ARGUMENT( FScrollBoxStyle, ShadowBoxStyle )
+
+		SLATE_ARGUMENT( bool, bEnableShadowBoxStyle)
 
 		SLATE_EVENT( FOnGenerateRow, OnGenerateRow )
 
@@ -272,6 +277,10 @@ public:
 
 		this->SetStyle(InArgs._ListViewStyle);
 
+		this->bShouldUseShadowBoxStyle = InArgs._bEnableShadowBoxStyle;
+
+		this->SetShadowBoxStyle(InArgs._ShadowBoxStyle);
+			
 		this->MaxPinnedItems = InArgs._MaxPinnedItems;
 		this->DefaultMaxPinnedItems = InArgs._MaxPinnedItems;
 		this->ScrollBarSlotPadding = InArgs._ScrollBarPadding;
@@ -1320,6 +1329,11 @@ public:
 		SetBackgroundBrush( Style != nullptr ? &Style->BackgroundBrush : FStyleDefaults::GetNoBrush() );
 	}
 
+	void SetShadowBoxStyle(const FScrollBoxStyle* InScrollBoxStyle)
+	{
+		ShadowBoxStyle = InScrollBoxStyle;
+	}
+
 	/** Sets the OnEntryInitializer delegate. This delegate is invoked after initializing an entry being generated, before it may be added to the actual widget hierarchy. */
 	void SetOnEntryInitialized(const FOnEntryInitialized& Delegate)
 	{
diff --git a/Engine/Source/Runtime/Slate/Public/Widgets/Views/STableViewBase.h b/Engine/Source/Runtime/Slate/Public/Widgets/Views/STableViewBase.h
--- a/Engine/Source/Runtime/Slate/Public/Widgets/Views/STableViewBase.h
+++ b/Engine/Source/Runtime/Slate/Public/Widgets/Views/STableViewBase.h
@@ -448,6 +448,12 @@ protected:
 	/** Delegate to call when the table view is scrolled */
 	FOnTableViewScrolled OnTableViewScrolled;
 
+	/** Gets the component of a vector in the direction of scrolling based on the Orientation property. */
+	FORCEINLINE float GetScrollComponentFromVector(FVector2f Vector) const
+	{
+		return float(Orientation == Orient_Vertical ? Vector.Y : Vector.X);
+	}
+
 	/** 
 	 * The fixed offset in lines to always apply to the top/left (depending on orientation) of the list.
 	 *
@@ -546,8 +552,20 @@ protected:
 	/** Passing over the clipping to SListPanel */
 	SLATE_API virtual void OnClippingChanged() override;
 
+	/** Should Use ShdadowBox Styling */
+	bool bShouldUseShadowBoxStyle = false;
+
+	/** Style resource for the Shadow Box when scrolling */
+	const FScrollBoxStyle* ShadowBoxStyle;
+
 	/** Brush resource representing the background area of the view */
 	FInvalidatableBrushAttribute BackgroundBrush;
+	
+	/** Does the user need a hint that they can scroll to the start of the list? */
+	FSlateColor GetStartShadowOpacity() const;
+	
+	/** Does the user need a hint that they can scroll to the end of the list? */
+	FSlateColor GetEndShadowOpacity() const;
 
 protected:
 
@@ -585,6 +603,12 @@ private:
 
 	/** Whether or not the current scroll is inertial */
 	bool bIsInertialScroll = false;
+	
+	/** Construct the vertical shadow box when enabled*/
+	TSharedRef<SWidget> ConstructVerticalShadowBox(TSharedRef<SWidget> Content);
+	
+	/** Construct the horizontal shadow box when enabled */
+	TSharedRef<SWidget> ConstructHorizontalShadowBox(TSharedRef<SWidget> Content);
 
 	TArray<float> DeltaTimeQueue;
 	int32 DeltaTimeCircularIndex = 0;
diff --git a/Engine/Source/Runtime/UMG/Private/Components/ListView.cpp b/Engine/Source/Runtime/UMG/Private/Components/ListView.cpp
--- a/Engine/Source/Runtime/UMG/Private/Components/ListView.cpp
+++ b/Engine/Source/Runtime/UMG/Private/Components/ListView.cpp
@@ -503,6 +503,16 @@ void UListView::SetHorizontalEntrySpacing(float InHorizontalEntrySpacing)
 	}
 }
 
+void UListView::SetShadowBrushStyle(const FScrollBoxStyle& NewScrollBoxStyle)
+{
+	ShadowBrushStyle = NewScrollBoxStyle;
+}
+
+const FScrollBoxStyle& UListView::GetShadowBrushStyle() const
+{
+	return ShadowBrushStyle;
+}
+
 void UListView::InitHorizontalEntrySpacing(float InHorizontalEntrySpacing)
 {
 	SetHorizontalEntrySpacing(InHorizontalEntrySpacing);
diff --git a/Engine/Source/Runtime/UMG/Public/Components/ListView.h b/Engine/Source/Runtime/UMG/Public/Components/ListView.h
--- a/Engine/Source/Runtime/UMG/Public/Components/ListView.h
+++ b/Engine/Source/Runtime/UMG/Public/Components/ListView.h
@@ -196,6 +196,8 @@ protected:
 		Args.Orientation = Orientation;
 		Args.ListViewStyle = &WidgetStyle;
 		Args.ScrollBarStyle = &ScrollBarStyle;
+		Args.ShadowBoxStyle = &ShadowBrushStyle;
+		Args.bEnableShadowBoxStyle = bEnableShadowBrush;
 		Args.ScrollBarPadding = ScrollBarPadding;
 		Args.ScrollIntoViewAlignment = ScrollIntoViewAlignment;
 		MyListView = ITypedUMGListView<UObject*>::ConstructListView<ListViewT>(this, ListItems, Args);
@@ -212,6 +214,12 @@ protected:
 	UPROPERTY(EditAnywhere, BlueprintReadOnly, Category = ListView)
 	FScrollBarStyle ScrollBarStyle;
 
+	UPROPERTY(EditAnywhere, BlueprintReadWrite, Category="Style", meta=( DisplayName="Enable Shadow Brush Style" ))
+	bool bEnableShadowBrush = false;
+	
+	UPROPERTY(EditAnywhere, BlueprintReadWrite, Getter, Setter, Category="Style", meta=( DisplayName="Shadow Brush Style", EditCondition = bEnableShadowBrush, EditConditionHides ))
+	FScrollBoxStyle ShadowBrushStyle;
+
 	/** 
 	 * The scroll & layout orientation of the list. ListView and TileView only. 
 	 * Vertical will scroll vertically and arrange tiles into rows.
@@ -286,6 +294,12 @@ public:
 	/** Set the horizontal spacing between entries. */
 	UMG_API void SetHorizontalEntrySpacing(float InHorizontalEntrySpacing);
 
+	/** Sets the Shadow Brush Style when scrolling */
+	UMG_API void SetShadowBrushStyle(const FScrollBoxStyle& NewWidgetStyle);
+
+	/** Gets the Shadow Brush Style when scrolling */
+	UMG_API const FScrollBoxStyle& GetShadowBrushStyle() const;
+
 protected:
 	UE_DEPRECATED(5.6, "InitHorizontalEntrySpacing is deprecated, please use SetHorizontalEntrySpacing instead.")
 	UMG_API void InitHorizontalEntrySpacing(float InHorizontalEntrySpacing);
diff --git a/Engine/Source/Runtime/UMG/Public/Components/ListViewBase.h b/Engine/Source/Runtime/UMG/Public/Components/ListViewBase.h
--- a/Engine/Source/Runtime/UMG/Public/Components/ListViewBase.h
+++ b/Engine/Source/Runtime/UMG/Public/Components/ListViewBase.h
@@ -258,6 +258,8 @@ protected:
 		EScrollIntoViewAlignment ScrollIntoViewAlignment = EScrollIntoViewAlignment::CenterAligned;
 		const FTableViewStyle* ListViewStyle = &FUMGCoreStyle::Get().GetWidgetStyle<FTableViewStyle>("ListView");
 		const FScrollBarStyle* ScrollBarStyle = &FUMGCoreStyle::Get().GetWidgetStyle<FScrollBarStyle>("ScrollBar");
+		const FScrollBoxStyle* ShadowBoxStyle = &FUMGCoreStyle::Get().GetWidgetStyle<FScrollBoxStyle>("ScrollBox");
+		bool bEnableShadowBoxStyle = false;
 		FMargin ScrollBarPadding = FMargin(0.0f);
 		bool bPreventThrottling = false;
 	};
@@ -279,7 +281,9 @@ protected:
 			.ClearScrollVelocityOnSelection(Args.bClearScrollVelocityOnSelection)
 			.Orientation(Args.Orientation)
 			.ScrollIntoViewAlignment(Args.ScrollIntoViewAlignment)
+			.bEnableShadowBoxStyle(Args.bEnableShadowBoxStyle)
 			.ListViewStyle(Args.ListViewStyle)
+			.ShadowBoxStyle(Args.ShadowBoxStyle)
 			.ScrollBarStyle(Args.ScrollBarStyle)
 			.ScrollBarPadding(Args.ScrollBarPadding)
 			.PreventThrottling(Args.bPreventThrottling)
