UMG: Clockwise Ordering option for DynamicEntryBox

[RN] minor creator/ui

diff --git a/Engine/Source/Runtime/Slate/Private/Widgets/Layout/SRadialBox.cpp b/Engine/Source/Runtime/Slate/Private/Widgets/Layout/SRadialBox.cpp
--- a/Engine/Source/Runtime/Slate/Private/Widgets/Layout/SRadialBox.cpp
+++ b/Engine/Source/Runtime/Slate/Private/Widgets/Layout/SRadialBox.cpp
@@ -30,6 +30,7 @@ void SRadialBox::Construct( const FArguments& InArgs )
 	bUseAllottedWidth = InArgs._UseAllottedWidth;
 	StartingAngle = InArgs._StartingAngle;
 	bDistributeItemsEvenly = InArgs._bDistributeItemsEvenly;
+	bClockwiseOrder = InArgs._bClockwiseOrder;
 	SectorCentralAngle = InternalNormalizeAngle(InArgs._SectorCentralAngle);
 	AngleBetweenItems = InArgs._AngleBetweenItems;
 
@@ -112,7 +113,8 @@ void SRadialBox::FChildArranger::Arrange()
 		ArrangementData.SlotSize.X = DesiredSizeOfSlot.X;
 		ArrangementData.SlotSize.Y = DesiredSizeOfSlot.Y;
 
-		DegreeOffset -= DegreeIncrements;
+		DegreeOffset += RadialBox.bClockwiseOrder ? DegreeIncrements : (-1 * DegreeIncrements);
+
 
 		OnSlotArranged(Slot, ArrangementData);
 	}
diff --git a/Engine/Source/Runtime/Slate/Public/Widgets/Layout/SRadialBox.h b/Engine/Source/Runtime/Slate/Public/Widgets/Layout/SRadialBox.h
--- a/Engine/Source/Runtime/Slate/Public/Widgets/Layout/SRadialBox.h
+++ b/Engine/Source/Runtime/Slate/Public/Widgets/Layout/SRadialBox.h
@@ -40,6 +40,7 @@ public:
 		, _bDistributeItemsEvenly(true)
 		, _AngleBetweenItems(0.f)
 		, _SectorCentralAngle(360.f)
+		, _bClockwiseOrder(false)
 		{
 			_Visibility = EVisibility::SelfHitTestInvisible;
 		}
@@ -65,6 +66,9 @@ public:
 		/** If we need a section of a radial (for example half-a-radial) we can define a central angle < 360 (180 in case of half-a-radial). Used when bDistributeItemsEvenly is enabled. */
 		SLATE_ARGUMENT(float, SectorCentralAngle)
 
+		/** If true children will be arranged in a clockwise order from starting angle  */
+		SLATE_ARGUMENT(bool, bClockwiseOrder)
+
 	SLATE_END_ARGS()
 
 	SLATE_API SRadialBox();
@@ -125,6 +129,9 @@ private:
 	/** If true the box will have a preferred width equal to its alloted width  */
 	bool bUseAllottedWidth;
 
+	/** If true children will be arranged in a clockwise order from starting angle  */
+	bool bClockwiseOrder;
+
 	class FChildArranger;
 	friend class SRadialBox::FChildArranger;
 };
diff --git a/Engine/Source/Runtime/UMG/Private/Components/DynamicEntryBoxBase.cpp b/Engine/Source/Runtime/UMG/Private/Components/DynamicEntryBoxBase.cpp
--- a/Engine/Source/Runtime/UMG/Private/Components/DynamicEntryBoxBase.cpp
+++ b/Engine/Source/Runtime/UMG/Private/Components/DynamicEntryBoxBase.cpp
@@ -285,6 +285,7 @@ TSharedRef<SWidget> UDynamicEntryBoxBase::RebuildWidget()
 			.UseAllottedWidth(true)
 			.StartingAngle(RadialBoxSettings.StartingAngle)
 			.bDistributeItemsEvenly(RadialBoxSettings.bDistributeItemsEvenly)
+			.bClockwiseOrder(RadialBoxSettings.bClockwiseOrder)
 			.AngleBetweenItems(RadialBoxSettings.AngleBetweenItems)
 			.SectorCentralAngle(RadialBoxSettings.SectorCentralAngle);
 		break;
diff --git a/Engine/Source/Runtime/UMG/Public/Components/RadialBoxSettings.h b/Engine/Source/Runtime/UMG/Public/Components/RadialBoxSettings.h
--- a/Engine/Source/Runtime/UMG/Public/Components/RadialBoxSettings.h
+++ b/Engine/Source/Runtime/UMG/Public/Components/RadialBoxSettings.h
@@ -19,6 +19,10 @@ struct FRadialBoxSettings
 	UPROPERTY(EditAnywhere, Category = "Items")
 	bool bDistributeItemsEvenly;
 
+	/* If enabled, Places*/
+	UPROPERTY(EditAnywhere, Category = "Items")
+	bool bClockwiseOrder;
+
 	/* Amount of Euler degrees that separate each item. Only used when bDistributeItemsEvenly is false */
 	UPROPERTY(EditAnywhere, Category = "Items", meta = (EditCondition = "!bDistributeItemsEvenly", ClampMin = 0, ClampMax = 360))
 	float AngleBetweenItems;
@@ -30,6 +34,7 @@ struct FRadialBoxSettings
 	FRadialBoxSettings()
 		: StartingAngle(0.f)
 		, bDistributeItemsEvenly(true)
+		, bClockwiseOrder(false)
 		, AngleBetweenItems(0.f)
 		, SectorCentralAngle(360.f)
 	{
