Fixed the DynamicEntryBox Unreal Motion Graphics (UMG) Widget's SectorCentralAngle value so that multiples of 360 are not incorrectly set to 0 at runtime. This allows DistributeItemsEvenly to work as expected with the widget's default values.

#tests Reproduced the problem as explained by Dan in the Slack thread related to the ticket. After making the fixes, with DistributeItemsEvenly set to true, I tried different number of entries with different values for SectorCentralAngle, and all worked as expected both in the UMG designer and at runtime
[RN] fix ui/slate

diff --git a/Engine/Source/Runtime/Slate/Private/Widgets/Layout/SRadialBox.cpp b/Engine/Source/Runtime/Slate/Private/Widgets/Layout/SRadialBox.cpp
--- a/Engine/Source/Runtime/Slate/Private/Widgets/Layout/SRadialBox.cpp
+++ b/Engine/Source/Runtime/Slate/Private/Widgets/Layout/SRadialBox.cpp
@@ -70,7 +70,6 @@ private:
 	TMap<int32, FArrangementData> OngoingArrangementDataMap;
 };
 
-
 SRadialBox::FChildArranger::FChildArranger(const SRadialBox& InRadialBox, const FOnSlotArranged& InOnSlotArranged)
 	: RadialBox(InRadialBox)
 	, OnSlotArranged(InOnSlotArranged)
@@ -157,6 +156,11 @@ FChildren* SRadialBox::GetChildren()
 	return &Slots;	
 }
 
+void SRadialBox::SetSectorCentralAngle(float InSectorCentralAngle)
+{
+	SectorCentralAngle = InternalNormalizeAngle(InSectorCentralAngle);
+}
+
 void SRadialBox::SetUseAllottedWidth(bool bInUseAllottedWidth)
 {
 	bUseAllottedWidth = bInUseAllottedWidth;
@@ -164,6 +168,10 @@ void SRadialBox::SetUseAllottedWidth(bool bInUseAllottedWidth)
 
 float SRadialBox::InternalNormalizeAngle(float Angle) const
 {
+	if (FMath::IsNearlyZero(Angle))
+	{
+		return 0.f;
+	}
 	float NormalizedAngle = FMath::Fmod(Angle, 360.f);
-	return NormalizedAngle < 0.f ? NormalizedAngle + 360.f : NormalizedAngle;
+	return NormalizedAngle < UE_SMALL_NUMBER ? NormalizedAngle + 360.f : NormalizedAngle;
 }
diff --git a/Engine/Source/Runtime/Slate/Public/Widgets/Layout/SRadialBox.h b/Engine/Source/Runtime/Slate/Public/Widgets/Layout/SRadialBox.h
--- a/Engine/Source/Runtime/Slate/Public/Widgets/Layout/SRadialBox.h
+++ b/Engine/Source/Runtime/Slate/Public/Widgets/Layout/SRadialBox.h
@@ -100,14 +100,14 @@ public:
 	void SetStartingAngle(float InStartingAngle) { StartingAngle = InStartingAngle; }
 	void SetAngleBetweenItems(float InAngleBetweenItems) { AngleBetweenItems = InAngleBetweenItems; }
 	void SetDistributeItemsEvenly(bool bInDistributeItemsEvenly) { bDistributeItemsEvenly = bInDistributeItemsEvenly; }
-	void SetSectorCentralAngle(float InSectorCentralAngle) { SectorCentralAngle = InSectorCentralAngle; }
+	SLATE_API void SetSectorCentralAngle(float InSectorCentralAngle);
 
 	SLATE_API void SetUseAllottedWidth(bool bInUseAllottedWidth);
 
 private:
 
 	/** Mods the angle so it's between 0-360 */
-	SLATE_API float InternalNormalizeAngle(float Angle) const;
+	float InternalNormalizeAngle(float Angle) const;
 
 	/** The slots that contain this panel's children. */
 	TPanelChildren<FSlot> Slots;
diff --git a/Engine/Source/Runtime/UMG/Public/Components/DynamicEntryBoxBase.h b/Engine/Source/Runtime/UMG/Public/Components/DynamicEntryBoxBase.h
--- a/Engine/Source/Runtime/UMG/Public/Components/DynamicEntryBoxBase.h
+++ b/Engine/Source/Runtime/UMG/Public/Components/DynamicEntryBoxBase.h
@@ -74,7 +74,7 @@ protected:
 	UPROPERTY(EditAnywhere, BlueprintReadOnly, Getter, Category = EntryLayout)
 	int32 MaxElementSize = 0;
 
-	UE_DEPRECATED(5.2, "Direct access to EntrySpacing is deprecated. Please use the getter or setter.")
+	UE_DEPRECATED(5.2, "Direct access to RadialBoxSettings is deprecated. Please use the getter or setter.")
 	/** Settings only relevant to RadialBox */
 	UPROPERTY(EditAnywhere, BlueprintReadWrite, Getter, Setter = "SetRadialSettings", BlueprintSetter = "SetRadialSettings", Category = EntryLayout)
 	FRadialBoxSettings RadialBoxSettings;
diff --git a/Engine/Source/Runtime/UMG/Public/Components/RadialBoxSettings.h b/Engine/Source/Runtime/UMG/Public/Components/RadialBoxSettings.h
--- a/Engine/Source/Runtime/UMG/Public/Components/RadialBoxSettings.h
+++ b/Engine/Source/Runtime/UMG/Public/Components/RadialBoxSettings.h
@@ -19,7 +19,7 @@ struct FRadialBoxSettings
 	UPROPERTY(EditAnywhere, Category = "Items")
 	bool bDistributeItemsEvenly;
 
-	/* If enabled, Places*/
+	/* If enabled, Places the entries in clockwise order instead of counter-clockwise order */
 	UPROPERTY(EditAnywhere, Category = "Items")
 	bool bClockwiseOrder;
 
