Dragging from a tab well docked somewhere in the middle of a window (i.e. not the main tab well or the tab well for a floating window) will no longer move the window.

[RN]

diff --git a/Engine/Source/Runtime/Slate/Private/Framework/Docking/SDockingTabWell.cpp b/Engine/Source/Runtime/Slate/Private/Framework/Docking/SDockingTabWell.cpp
--- a/Engine/Source/Runtime/Slate/Private/Framework/Docking/SDockingTabWell.cpp
+++ b/Engine/Source/Runtime/Slate/Private/Framework/Docking/SDockingTabWell.cpp
@@ -480,8 +480,25 @@ FReply SDockingTabWell::OnDrop( const FGeometry& MyGeometry, const FDragDropEven
 
 EWindowZone::Type SDockingTabWell::GetWindowZoneOverride() const
 {
-	// Pretend we are a title bar so the user can grab the area to move the window around
-	return EWindowZone::TitleBar;
+	// If this is the tab well for the top-level tab stack of a window, then this window zone should be treated like a title bar instead of client area.
+
+	// Get the tab stack that owns this tab well
+	if (const TSharedPtr<SDockingTabStack> ParentTabStack = ParentTabStackPtr.Pin())
+	{
+		// Get the docking area for that tab stack
+		const TSharedPtr<SDockingArea> ParentDockingArea = ParentTabStack->GetDockArea();
+
+		// If the docking area is managing a window, then it's at the top level of the window (either the main window or a floating window)
+		// Docking areas that are themselves docked within another docking area won't have a parent window
+		if (ParentDockingArea && ParentDockingArea->GetParentWindow())
+		{
+			// Declare the tab well to be a title bar, allowing the user to drag the tab well to move the window
+			return EWindowZone::TitleBar;
+		}
+	}
+
+	// Otherwise, declare this to be simple client area
+	return EWindowZone::ClientArea;
 }
 
 FReply SDockingTabWell::OnMouseButtonUp( const FGeometry& MyGeometry, const FPointerEvent& MouseEvent )
