Slate Notifications - Ensure we first show notifications when arranged (Windows, Mac) and update DesiredSize when arranging on Mac

#tests Preflight + Test on Windows, Mac & Linux

diff --git a/Engine/Source/Runtime/Slate/Private/Framework/Notifications/NotificationManager.cpp b/Engine/Source/Runtime/Slate/Private/Framework/Notifications/NotificationManager.cpp
--- a/Engine/Source/Runtime/Slate/Private/Framework/Notifications/NotificationManager.cpp
+++ b/Engine/Source/Runtime/Slate/Private/Framework/Notifications/NotificationManager.cpp
@@ -56,17 +56,36 @@ void FSlateNotificationManager::FRegionalNotificationList::Arrange()
 		if( PinnedWindow.IsValid() )
 		{
 			const float StackOffset = ScaledOffset.Y * ((Notifications.Num() - 1) - ListIndex);
+
+#if PLATFORM_MAC
+			// Ensure the DesiredSize is up to date to avoid the os detecting the window as moving to right screen
+			// On Windows this worsened the animation of the scaling because the content doesn't scale fast enough
+			// On Linux the animation doesn't work at all for unknown reason likely related to SDL
+			PinnedWindow->SlatePrepass(FSlateApplicationBase::Get().GetApplicationScale() * PinnedWindow->GetDPIScaleFactor());
+#endif // PLATFORM_MAC
 			const FVector2D DesiredSize = PinnedWindow->GetDesiredSize();
+			
 			const FVector2D NewPosition(AnchorPoint.X - DesiredSize.X, AnchorPoint.Y - DesiredSize.Y - StackOffset);
 
 			if( NewPosition != PinnedWindow->GetPositionInScreen() && DesiredSize != PinnedWindow->GetSizeInScreen())
 			{
 				PinnedWindow->ReshapeWindow( NewPosition, DesiredSize );
 			}
+			
 			else if (NewPosition != PinnedWindow->GetPositionInScreen())
 			{
 				PinnedWindow->MoveWindowTo(NewPosition);
 			}
+			
+#if !PLATFORM_LINUX
+			// Only display the notification when it's ready and arranged
+			// On Linux there is a hack in SDL to already delayed the showing of notification that is not interacting well with this optimization.
+			if (!PinnedWindow->IsVisible())
+			{
+				PinnedWindow->ShowWindow();
+			}
+#endif
+			
 			AnchorPoint.Y -= DesiredSize.Y;
 		}
 	}
@@ -152,13 +171,18 @@ TSharedRef<SNotificationList> FSlateNotificationManager::CreateStackForArea(cons
 
 	// Prevent tying a notification to a window that is minimized because a minimized window isn't painted
 	// and therefore the notification wouldn't be either.
+	
+	// We don't show the NotificationWindow immediately because we want to arrange it first.
+	// On Linux there is a hack in SDL to already delayed the showing of notification that is not interacting well with this optimization.
+	constexpr bool bShowImmediately = !PLATFORM_LINUX;
+
 	if (Window.IsValid() && Window->IsVisible() && !Window->IsWindowMinimized())
 	{
-		FSlateApplication::Get().AddWindowAsNativeChild(NotificationWindow, Window.ToSharedRef());
+		FSlateApplication::Get().AddWindowAsNativeChild(NotificationWindow, Window.ToSharedRef(), bShowImmediately);
 	}
 	else
 	{
-		FSlateApplication::Get().AddWindow(NotificationWindow);
+		FSlateApplication::Get().AddWindow(NotificationWindow, bShowImmediately);
 	}
 
 	if (!FSlateApplication::Get().GetActiveModalWindow().IsValid())
