NotificationManager - Fix inversion of a condition preventing Notifications from displaying on Linux

#tests Tested on Mac + Preflight Windows, Linux

diff --git a/Engine/Source/Runtime/Slate/Private/Framework/Notifications/NotificationManager.cpp b/Engine/Source/Runtime/Slate/Private/Framework/Notifications/NotificationManager.cpp
--- a/Engine/Source/Runtime/Slate/Private/Framework/Notifications/NotificationManager.cpp
+++ b/Engine/Source/Runtime/Slate/Private/Framework/Notifications/NotificationManager.cpp
@@ -71,20 +71,16 @@ void FSlateNotificationManager::FRegionalNotificationList::Arrange()
 			{
 				PinnedWindow->ReshapeWindow( NewPosition, DesiredSize );
 			}
-			
 			else if (NewPosition != PinnedWindow->GetPositionInScreen())
 			{
 				PinnedWindow->MoveWindowTo(NewPosition);
 			}
 			
-#if !PLATFORM_LINUX
 			// Only display the notification when it's ready and arranged
-			// On Linux there is a hack in SDL to already delayed the showing of notification that is not interacting well with this optimization.
-			if (!PinnedWindow->IsVisible())
+			if (!bShowNotificationImmediately && !PinnedWindow->IsVisible())
 			{
 				PinnedWindow->ShowWindow();
 			}
-#endif
 			
 			AnchorPoint.Y -= DesiredSize.Y;
 		}
@@ -171,18 +167,14 @@ TSharedRef<SNotificationList> FSlateNotificationManager::CreateStackForArea(cons
 
 	// Prevent tying a notification to a window that is minimized because a minimized window isn't painted
 	// and therefore the notification wouldn't be either.
-	
-	// We don't show the NotificationWindow immediately because we want to arrange it first.
-	// On Linux there is a hack in SDL to already delayed the showing of notification that is not interacting well with this optimization.
-	constexpr bool bShowImmediately = !PLATFORM_LINUX;
 
 	if (Window.IsValid() && Window->IsVisible() && !Window->IsWindowMinimized())
 	{
-		FSlateApplication::Get().AddWindowAsNativeChild(NotificationWindow, Window.ToSharedRef(), bShowImmediately);
+		FSlateApplication::Get().AddWindowAsNativeChild(NotificationWindow, Window.ToSharedRef(), bShowNotificationImmediately);
 	}
 	else
 	{
-		FSlateApplication::Get().AddWindow(NotificationWindow, bShowImmediately);
+		FSlateApplication::Get().AddWindow(NotificationWindow, bShowNotificationImmediately);
 	}
 
 	if (!FSlateApplication::Get().GetActiveModalWindow().IsValid())
diff --git a/Engine/Source/Runtime/Slate/Public/Framework/Notifications/NotificationManager.h b/Engine/Source/Runtime/Slate/Public/Framework/Notifications/NotificationManager.h
--- a/Engine/Source/Runtime/Slate/Public/Framework/Notifications/NotificationManager.h
+++ b/Engine/Source/Runtime/Slate/Public/Framework/Notifications/NotificationManager.h
@@ -234,4 +234,7 @@ private:
 
 	/** CVar allowing us to control the display of notifications. */
 	FAutoConsoleVariableRef CVarAllowNotifications;
+	
+	// On Linux there is a hack in SDL to already delayed the showing of notification that is not interacting well with this optimization.
+	static constexpr bool bShowNotificationImmediately = PLATFORM_LINUX;
 };
