NotificationManager - Fix static analysis error with a not operator on a static constexpr

#tests Local run on Mac + Preflight on Windows / Mac

diff --git a/Engine/Source/Runtime/Slate/Private/Framework/Notifications/NotificationManager.cpp b/Engine/Source/Runtime/Slate/Private/Framework/Notifications/NotificationManager.cpp
--- a/Engine/Source/Runtime/Slate/Private/Framework/Notifications/NotificationManager.cpp
+++ b/Engine/Source/Runtime/Slate/Private/Framework/Notifications/NotificationManager.cpp
@@ -77,9 +77,12 @@ void FSlateNotificationManager::FRegionalNotificationList::Arrange()
 			}
 			
 			// Only display the notification when it's ready and arranged
-			if (!bShowNotificationImmediately && !PinnedWindow->IsVisible())
+			if constexpr (!bShowNotificationImmediately)
 			{
-				PinnedWindow->ShowWindow();
+				if (!PinnedWindow->IsVisible())
+				{
+					PinnedWindow->ShowWindow();
+				}
 			}
 			
 			AnchorPoint.Y -= DesiredSize.Y;
