Fixed SWrapBox Alignment calculations when using `Inner Slot Padding` for Center, Right, and Fill `Horizontal Alignment`s.

#tests reproduced the tests in the ticket and achieved expected behaviours.
[RN] fix ui/slate

diff --git a/Engine/Source/Runtime/Slate/Private/Widgets/Layout/SWrapBox.cpp b/Engine/Source/Runtime/Slate/Private/Widgets/Layout/SWrapBox.cpp
--- a/Engine/Source/Runtime/Slate/Private/Widgets/Layout/SWrapBox.cpp
+++ b/Engine/Source/Runtime/Slate/Private/Widgets/Layout/SWrapBox.cpp
@@ -68,6 +68,7 @@ public:
 		FVector2D SlotOffset;
 		FVector2D SlotSize;
 		float WidthOfCurrentLine;
+		float TotalLinePadding;
 		int ChildIndexRelativeToLine;
 	};
 
@@ -84,7 +85,12 @@ private:
 	const FOnSlotArranged& OnSlotArranged;
 	FVector2D Offset;
 	float MaximumSizeInCurrentLine;
+
+	// This is the total width of the line, including padding.
 	float WidthOfCurrentLine;
+
+	// This is the total width of just the padding in the current line. Not the other content.
+	float TotalLinePadding;
 	int32 IndexOfFirstChildInCurrentLine;
 	TMap<int32, FArrangementData> OngoingArrangementDataMap;
 };
@@ -96,6 +102,7 @@ SWrapBox::FChildArranger::FChildArranger(const SWrapBox& InWrapBox, const FOnSlo
 	, Offset(FVector2D::ZeroVector)
 	, MaximumSizeInCurrentLine(0.0f)
 	, WidthOfCurrentLine(0.0f)
+	, TotalLinePadding(0.0f)
 	, IndexOfFirstChildInCurrentLine(INDEX_NONE)
 {
 	OngoingArrangementDataMap.Reserve(WrapBox.Slots.Num());
@@ -152,7 +159,6 @@ void SWrapBox::FChildArranger::Arrange()
 			if (WrapBox.Orientation == EOrientation::Orient_Horizontal)
 			{
 				Offset.X += WrapBox.InnerSlotPadding.X;
-				WidthOfCurrentLine += WrapBox.InnerSlotPadding.X;
 			}
 			else
 			{
@@ -246,6 +252,13 @@ void SWrapBox::FChildArranger::Arrange()
 			MaximumSizeInCurrentLine = FMath::Max(MaximumSizeInCurrentLine, ArrangementData.SlotSize.Y);
 			WidthOfCurrentLine += ArrangementData.SlotSize.X;
 
+			// At this point, we know that this child is going on this line. If it is not the first child on this line, then we must include the inner slot padding in this line's width.
+			if (!IsFirstChildInCurrentLine())
+			{
+				WidthOfCurrentLine += WrapBox.InnerSlotPadding.X;
+				TotalLinePadding += WrapBox.InnerSlotPadding.X;
+			}
+
 			// Update offset to right bound of child.
 			Offset.X = ArrangementData.SlotOffset.X + ArrangementData.SlotSize.X;
 		}
@@ -292,6 +305,7 @@ void SWrapBox::FChildArranger::FinalizeLine(int32 IndexOfLastChildInCurrentLine)
 		FArrangementData& ArrangementData = OngoingArrangementDataMap[ChildIndex];
 		ArrangementData.ChildIndexRelativeToLine = RelativeChildIndex;
 		ArrangementData.WidthOfCurrentLine = WidthOfCurrentLine;
+		ArrangementData.TotalLinePadding = TotalLinePadding;
 
 		RelativeChildIndex++;
 
@@ -346,7 +360,8 @@ void SWrapBox::FChildArranger::FinalizeLine(int32 IndexOfLastChildInCurrentLine)
 		Offset.Y += MaximumSizeInCurrentLine + WrapBox.InnerSlotPadding.Y;
 	}
 
-	WidthOfCurrentLine = 0;
+	WidthOfCurrentLine = 0.0f;
+	TotalLinePadding = 0.0f;
 	MaximumSizeInCurrentLine = 0.0f;
 	IndexOfFirstChildInCurrentLine = INDEX_NONE;
 }
@@ -382,7 +397,8 @@ void SWrapBox::OnArrangeChildren(const FGeometry& AllottedGeometry, FArrangedChi
 				LeftStop = FMath::FloorToFloat((AllottedGeometry.GetLocalSize().X - ArrangementData.WidthOfCurrentLine) / 2.0f);
 				break;
 			case HAlign_Fill:
-				const float NewSlotSize = AllottedSlotSize.X / ArrangementData.WidthOfCurrentLine * AllottedGeometry.GetLocalSize().X;
+				// We must make sure to not include the size of the line's padding when calculating how much to stretch the slot by.
+				const float NewSlotSize = AllottedSlotSize.X / (ArrangementData.WidthOfCurrentLine - ArrangementData.TotalLinePadding) * (AllottedGeometry.GetLocalSize().X - ArrangementData.TotalLinePadding);
 				SlotSizeAdjustment = NewSlotSize - AllottedSlotSize.X;
 				AllottedSlotSize.X = NewSlotSize;
 				break;
diff --git a/Engine/Source/Runtime/UMG/Public/Components/WrapBox.h b/Engine/Source/Runtime/UMG/Public/Components/WrapBox.h
--- a/Engine/Source/Runtime/UMG/Public/Components/WrapBox.h
+++ b/Engine/Source/Runtime/UMG/Public/Components/WrapBox.h
@@ -41,7 +41,7 @@ public:
 
 	UE_DEPRECATED(5.1, "Direct access to HorizontalAlignment is deprecated. Please use the getter or setter.")
 	/** The alignment of each line of wrapped content. */
-	UPROPERTY(EditAnywhere, BlueprintReadWrite, Getter, Setter, BlueprintSetter = "SetHorizontalAlignment", Category = "Content Layout")
+	UPROPERTY(EditAnywhere, BlueprintReadWrite, Getter, Setter, BlueprintSetter = "SetHorizontalAlignment", Category = "Content Layout", meta = (EditCondition = "Orientation == EOrientation::Orient_Horizontal"))
 	TEnumAsByte<EHorizontalAlignment> HorizontalAlignment;
 
 	UE_DEPRECATED(5.1, "Direct access to Orientation is deprecated. Please use the getter or setter.")
