UE-296848 division by zero fix

#rnx
#tests Tested the reproduction steps in the Jira. All behaviour worked as expected.

diff --git a/Engine/Source/Runtime/Slate/Private/Widgets/Layout/SWrapBox.cpp b/Engine/Source/Runtime/Slate/Private/Widgets/Layout/SWrapBox.cpp
--- a/Engine/Source/Runtime/Slate/Private/Widgets/Layout/SWrapBox.cpp
+++ b/Engine/Source/Runtime/Slate/Private/Widgets/Layout/SWrapBox.cpp
@@ -398,7 +398,12 @@ void SWrapBox::OnArrangeChildren(const FGeometry& AllottedGeometry, FArrangedChi
 				break;
 			case HAlign_Fill:
 				// We must make sure to not include the size of the line's padding when calculating how much to stretch the slot by.
-				const float NewSlotSize = AllottedSlotSize.X / (ArrangementData.WidthOfCurrentLine - ArrangementData.TotalLinePadding) * (AllottedGeometry.GetLocalSize().X - ArrangementData.TotalLinePadding);
+				// We should be careful to not have any division by zero issues here.
+				float NewSlotSize = 0.0f;
+				if (!FMath::IsNearlyZero(ArrangementData.WidthOfCurrentLine - ArrangementData.TotalLinePadding))
+				{
+					NewSlotSize = AllottedSlotSize.X / (ArrangementData.WidthOfCurrentLine - ArrangementData.TotalLinePadding) * (AllottedGeometry.GetLocalSize().X - ArrangementData.TotalLinePadding);
+				}
 				SlotSizeAdjustment = NewSlotSize - AllottedSlotSize.X;
 				AllottedSlotSize.X = NewSlotSize;
 				break;
