Slate: CVar `Slate.UseSharedBreakIterator` to support shared ICU Break Iterators to reduce CPU usage.
- TextLayout, SEditableTextLayout, and SlateTextShaper leverage the shared access.

[RN] minor creator/ui

diff --git a/Engine/Source/Runtime/Core/Private/Internationalization/ICUCharacterBoundaryIterator.cpp b/Engine/Source/Runtime/Core/Private/Internationalization/ICUCharacterBoundaryIterator.cpp
--- a/Engine/Source/Runtime/Core/Private/Internationalization/ICUCharacterBoundaryIterator.cpp
+++ b/Engine/Source/Runtime/Core/Private/Internationalization/ICUCharacterBoundaryIterator.cpp
@@ -13,4 +13,16 @@ TSharedRef<IBreakIterator> FBreakIterator::CreateCharacterBoundaryIterator()
 	return MakeShareable(new FICUBreakIterator(FICUBreakIteratorManager::Get().CreateCharacterBoundaryIterator()));
 }
 
+TSharedRef<IBreakIterator> FBreakIterator::GetCharacterBoundaryIterator()
+{
+	static thread_local TSharedPtr<IBreakIterator> CachedCharacterBoundaryIterator;
+	if (!CachedCharacterBoundaryIterator.IsValid())
+	{
+		CachedCharacterBoundaryIterator = MakeShareable(new FICUBreakIterator(FICUBreakIteratorManager::Get().CreateCharacterBoundaryIterator()));
+	}
+
+	check(CachedCharacterBoundaryIterator.IsValid());
+	return CachedCharacterBoundaryIterator.ToSharedRef();
+}
+
 #endif
diff --git a/Engine/Source/Runtime/Core/Private/Internationalization/ICULineBreakIterator.cpp b/Engine/Source/Runtime/Core/Private/Internationalization/ICULineBreakIterator.cpp
--- a/Engine/Source/Runtime/Core/Private/Internationalization/ICULineBreakIterator.cpp
+++ b/Engine/Source/Runtime/Core/Private/Internationalization/ICULineBreakIterator.cpp
@@ -211,4 +211,17 @@ TSharedRef<IBreakIterator> FBreakIterator::CreateLineBreakIterator()
 	return MakeShareable(new FICULineBreakIterator());
 }
 
+
+TSharedRef<IBreakIterator> FBreakIterator::GetLineBreakIterator()
+{
+	static thread_local TSharedPtr<IBreakIterator> CachedLineBreakIterator;
+	if (!CachedLineBreakIterator.IsValid())
+	{
+		CachedLineBreakIterator = MakeShareable(new FICULineBreakIterator());
+	}
+
+	check(CachedLineBreakIterator.IsValid());
+	return CachedLineBreakIterator.ToSharedRef();
+}
+
 #endif
diff --git a/Engine/Source/Runtime/Core/Private/Internationalization/ICUWordBreakIterator.cpp b/Engine/Source/Runtime/Core/Private/Internationalization/ICUWordBreakIterator.cpp
--- a/Engine/Source/Runtime/Core/Private/Internationalization/ICUWordBreakIterator.cpp
+++ b/Engine/Source/Runtime/Core/Private/Internationalization/ICUWordBreakIterator.cpp
@@ -13,4 +13,16 @@ TSharedRef<IBreakIterator> FBreakIterator::CreateWordBreakIterator()
 	return MakeShareable(new FICUBreakIterator(FICUBreakIteratorManager::Get().CreateWordBreakIterator()));
 }
 
+TSharedRef<IBreakIterator> FBreakIterator::GetWordBreakIterator()
+{
+	static thread_local TSharedPtr<IBreakIterator> CachedWordBreakIterator;
+	if (!CachedWordBreakIterator.IsValid())
+	{
+		CachedWordBreakIterator = MakeShareable(new FICUBreakIterator(FICUBreakIteratorManager::Get().CreateWordBreakIterator()));
+	}
+
+	check(CachedWordBreakIterator.IsValid());
+	return CachedWordBreakIterator.ToSharedRef();
+}
+
 #endif
diff --git a/Engine/Source/Runtime/Core/Private/Internationalization/LegacyCharacterBoundaryIterator.cpp b/Engine/Source/Runtime/Core/Private/Internationalization/LegacyCharacterBoundaryIterator.cpp
--- a/Engine/Source/Runtime/Core/Private/Internationalization/LegacyCharacterBoundaryIterator.cpp
+++ b/Engine/Source/Runtime/Core/Private/Internationalization/LegacyCharacterBoundaryIterator.cpp
@@ -95,4 +95,16 @@ TSharedRef<IBreakIterator> FBreakIterator::CreateCharacterBoundaryIterator()
 	return MakeShareable(new FLegacyCharacterBoundaryIterator());
 }
 
+TSharedRef<IBreakIterator> FBreakIterator::GetCharacterBoundaryIterator()
+{
+	static thread_local TSharedPtr<IBreakIterator> CachedCharacterBoundaryIterator;
+	if (!CachedCharacterBoundaryIterator.IsValid())
+	{
+		CachedCharacterBoundaryIterator = MakeShareable(new FLegacyCharacterBoundaryIterator());
+	}
+
+	check(CachedCharacterBoundaryIterator.IsValid());
+	return CachedCharacterBoundaryIterator.ToSharedRef();
+}
+
 #endif
diff --git a/Engine/Source/Runtime/Core/Private/Internationalization/LegacyLineBreakIterator.cpp b/Engine/Source/Runtime/Core/Private/Internationalization/LegacyLineBreakIterator.cpp
--- a/Engine/Source/Runtime/Core/Private/Internationalization/LegacyLineBreakIterator.cpp
+++ b/Engine/Source/Runtime/Core/Private/Internationalization/LegacyLineBreakIterator.cpp
@@ -115,4 +115,16 @@ TSharedRef<IBreakIterator> FBreakIterator::CreateLineBreakIterator()
 	return MakeShareable(new FLegacyLineBreakIterator());
 }
 
+TSharedRef<IBreakIterator> FBreakIterator::GetLineBreakIterator()
+{
+	static thread_local TSharedPtr<IBreakIterator> CachedLineBreakIterator;
+	if (!CachedLineBreakIterator.IsValid())
+	{
+		CachedLineBreakIterator = MakeShareable(new FLegacyLineBreakIterator());
+	}
+
+	check(CachedLineBreakIterator.IsValid());
+	return CachedLineBreakIterator.ToSharedRef();
+}
+
 #endif
diff --git a/Engine/Source/Runtime/Core/Private/Internationalization/LegacyWordBreakIterator.cpp b/Engine/Source/Runtime/Core/Private/Internationalization/LegacyWordBreakIterator.cpp
--- a/Engine/Source/Runtime/Core/Private/Internationalization/LegacyWordBreakIterator.cpp
+++ b/Engine/Source/Runtime/Core/Private/Internationalization/LegacyWordBreakIterator.cpp
@@ -115,4 +115,17 @@ TSharedRef<IBreakIterator> FBreakIterator::CreateWordBreakIterator()
 	return MakeShareable(new FLegacyWordBreakIterator());
 }
 
+
+TSharedRef<IBreakIterator> FBreakIterator::GetWordBreakIterator()
+{
+	static thread_local TSharedPtr<IBreakIterator> CachedWordBreakIterator;
+	if (!CachedWordBreakIterator.IsValid())
+	{
+		CachedWordBreakIterator = MakeShareable(new FLegacyWordBreakIterator());
+	}
+
+	check(CachedWordBreakIterator.IsValid());
+	return CachedWordBreakIterator.ToSharedRef();
+}
+
 #endif
diff --git a/Engine/Source/Runtime/Core/Public/Internationalization/BreakIterator.h b/Engine/Source/Runtime/Core/Public/Internationalization/BreakIterator.h
--- a/Engine/Source/Runtime/Core/Public/Internationalization/BreakIterator.h
+++ b/Engine/Source/Runtime/Core/Public/Internationalization/BreakIterator.h
@@ -28,4 +28,19 @@ struct FBreakIterator
 	 * This may produce odd results when used with languages which don't have a concept of cases for characters, however should be fine for parsing variable or asset names
 	 */
 	static CORE_API TSharedRef<IBreakIterator> CreateCamelCaseBreakIterator();
+
+	/**
+	 * Get a shared instance of a character boundary break iterator
+	 */
+	static CORE_API TSharedRef<IBreakIterator> GetCharacterBoundaryIterator();
+
+	/**
+	 * Get a shared instance of a work break iterator
+	 */
+	static CORE_API TSharedRef<IBreakIterator> GetWordBreakIterator();
+
+	/**
+	 * Get a new instance of a line break iterator 
+	 */
+	static CORE_API TSharedRef<IBreakIterator> GetLineBreakIterator();
 };
diff --git a/Engine/Source/Runtime/Slate/Private/Framework/Text/TextLayout.cpp b/Engine/Source/Runtime/Slate/Private/Framework/Text/TextLayout.cpp
--- a/Engine/Source/Runtime/Slate/Private/Framework/Text/TextLayout.cpp
+++ b/Engine/Source/Runtime/Slate/Private/Framework/Text/TextLayout.cpp
@@ -1156,7 +1156,7 @@ void FTextLayout::CreateLineWrappingCache(FLineModel& LineModel)
 	// If we've not yet been provided with a custom line break iterator, then just use the default one
 	if (!LineBreakIterator.IsValid())
 	{
-		LineBreakIterator = FBreakIterator::CreateLineBreakIterator();
+		LineBreakIterator = GSlateUseSharedBreakIterator  ? FBreakIterator::GetLineBreakIterator() : FBreakIterator::CreateLineBreakIterator();
 	}
 
 	LineBreakIterator->SetStringRef(&LineModel.Text.Get());
@@ -1221,8 +1221,8 @@ FTextLayout::FTextLayout()
 	, ViewSize( ForceInitToZero )
 	, ScrollOffset( ForceInitToZero )
 	, LineBreakIterator() // Initialized in FTextLayout::CreateWrappingCache if no custom iterator is provided
-	, GraphemeBreakIterator(FBreakIterator::CreateCharacterBoundaryIterator())
-	, WordBreakIterator(FBreakIterator::CreateWordBreakIterator())
+	, GraphemeBreakIterator(GSlateUseSharedBreakIterator  ? FBreakIterator::GetCharacterBoundaryIterator() : FBreakIterator::CreateCharacterBoundaryIterator())
+	, WordBreakIterator(GSlateUseSharedBreakIterator  ? FBreakIterator::GetWordBreakIterator() : FBreakIterator::CreateWordBreakIterator())
 	, TextBiDiDetection(TextBiDi::CreateTextBiDi())
 	, TextOverflowPolicyOverride()
 {
diff --git a/Engine/Source/Runtime/Slate/Private/Widgets/Text/SlateEditableTextLayout.cpp b/Engine/Source/Runtime/Slate/Private/Widgets/Text/SlateEditableTextLayout.cpp
--- a/Engine/Source/Runtime/Slate/Private/Widgets/Text/SlateEditableTextLayout.cpp
+++ b/Engine/Source/Runtime/Slate/Private/Widgets/Text/SlateEditableTextLayout.cpp
@@ -94,7 +94,7 @@ FSlateEditableTextLayout::FSlateEditableTextLayout(ISlateEditableTextWidget& InO
 
 	SearchCase = ESearchCase::IgnoreCase;
 
-	GraphemeBreakIterator = FBreakIterator::CreateCharacterBoundaryIterator();
+	GraphemeBreakIterator = GSlateUseSharedBreakIterator ? FBreakIterator::GetCharacterBoundaryIterator() : FBreakIterator::CreateCharacterBoundaryIterator();
 
 	BoundText = InInitialText;
 
diff --git a/Engine/Source/Runtime/SlateCore/Private/Fonts/SlateTextShaper.cpp b/Engine/Source/Runtime/SlateCore/Private/Fonts/SlateTextShaper.cpp
--- a/Engine/Source/Runtime/SlateCore/Private/Fonts/SlateTextShaper.cpp
+++ b/Engine/Source/Runtime/SlateCore/Private/Fonts/SlateTextShaper.cpp
@@ -96,7 +96,7 @@ FSlateTextShaper::FSlateTextShaper(FFreeTypeCacheDirectory* InFTCacheDirectory,
 	, FontRenderer(InFontRenderer)
 	, FontCache(InFontCache)
 	, TextBiDiDetection(TextBiDi::CreateTextBiDi())
-	, GraphemeBreakIterator(FBreakIterator::CreateCharacterBoundaryIterator())
+	, GraphemeBreakIterator(GSlateUseSharedBreakIterator ? FBreakIterator::GetCharacterBoundaryIterator() : FBreakIterator::CreateCharacterBoundaryIterator())
 #if WITH_HARFBUZZ
 	, HarfBuzzFontFactory(FTCacheDirectory)
 #endif // WITH_HARFBUZZ
diff --git a/Engine/Source/Runtime/SlateCore/Private/SlateCoreClasses.cpp b/Engine/Source/Runtime/SlateCore/Private/SlateCoreClasses.cpp
--- a/Engine/Source/Runtime/SlateCore/Private/SlateCoreClasses.cpp
+++ b/Engine/Source/Runtime/SlateCore/Private/SlateCoreClasses.cpp
@@ -108,6 +108,14 @@ FAutoConsoleVariableRef CVarNavigationDebugging(
 
 #endif
 
+bool GSlateUseSharedBreakIterator = false;
+static FAutoConsoleVariableRef CVarUseSharedBreakIterator(
+	TEXT("Slate.UseSharedBreakIterator"),
+	GSlateUseSharedBreakIterator,
+	TEXT("If true, reuse cached break iterators to reduce memory overhead")
+);
+
+
 FSlateWidgetStyle::FSlateWidgetStyle() = default;
 FSlateWidgetStyle::~FSlateWidgetStyle() = default;
 
diff --git a/Engine/Source/Runtime/SlateCore/Public/SlateGlobals.h b/Engine/Source/Runtime/SlateCore/Public/SlateGlobals.h
--- a/Engine/Source/Runtime/SlateCore/Public/SlateGlobals.h
+++ b/Engine/Source/Runtime/SlateCore/Public/SlateGlobals.h
@@ -95,6 +95,7 @@ extern SLATECORE_API bool GSlateNavigationDebugging;
 class FActiveTimerHandle;
 enum class EActiveTimerReturnType : uint8;
 
+extern SLATECORE_API bool GSlateUseSharedBreakIterator;
 
 /** Used to guard access across slate to specific threads */
 #define SLATE_CROSS_THREAD_CHECK() checkf(IsInGameThread() || IsInSlateThread(), TEXT("Slate can only be accessed from the GameThread or the SlateLoadingThread!"));
\ No newline at end of file
