Fix modal loop menus not updating widgets when clicking on checkboxes

#tests tested with metasound sources (mute/solo are menus in a non modal mode), tested with niagara system where the menu is within a modal loop
#rnx

diff --git a/Engine/Source/Runtime/Slate/Private/Framework/Application/SlateApplication.cpp b/Engine/Source/Runtime/Slate/Private/Framework/Application/SlateApplication.cpp
--- a/Engine/Source/Runtime/Slate/Private/Framework/Application/SlateApplication.cpp
+++ b/Engine/Source/Runtime/Slate/Private/Framework/Application/SlateApplication.cpp
@@ -812,6 +812,7 @@ FSlateApplication::FSlateApplication()
 	, Scale( 1.0f )
 	, DragTriggerDistance( 0 )
 	, CursorRadius( 0.0f )
+	, DrawId(0)
 	, LastUserInteractionTime( 0.0 )
 	, LastUserInteractionTimeForThrottling( 0.0 )
 	, LastMouseMoveTime( 0.0 )
@@ -1745,6 +1746,8 @@ void FSlateApplication::TickAndDrawWidgets(float DeltaTime)
 	}
 #endif
 
+	++DrawId;
+
 	UE_TRACE_SLATE_APPLICATION_TICK_AND_DRAW_WIDGETS(GetDeltaTime());
 }
 
@@ -2821,6 +2824,11 @@ const void* FSlateApplication::FindBestParentWindowHandleForDialogs(const TShare
 	return ParentWindowWindowHandle;
 }
 
+uint32 FSlateApplication::GetCurrentDrawId() const
+{
+	return DrawId;
+}
+
 const TSet<FKey>& FSlateApplication::GetPressedMouseButtons() const
 {
 	return PressedMouseButtons;
diff --git a/Engine/Source/Runtime/Slate/Public/Framework/Application/SlateApplication.h b/Engine/Source/Runtime/Slate/Public/Framework/Application/SlateApplication.h
--- a/Engine/Source/Runtime/Slate/Public/Framework/Application/SlateApplication.h
+++ b/Engine/Source/Runtime/Slate/Public/Framework/Application/SlateApplication.h
@@ -1848,6 +1848,11 @@ public:
 	*/
 	SLATE_API const void* FindBestParentWindowHandleForDialogs(const TSharedPtr<SWidget>& InWidget, const ESlateParentWindowSearchMethod InParentWindowSearchMethod = ESlateParentWindowSearchMethod::ActiveWindow);
 
+	/**
+	* @return Returns the current Id of the application drawing phase. The value returned can be used for == and != operations only.
+	*/
+	SLATE_API uint32 GetCurrentDrawId() const;
+
 public:
 #if WITH_EDITORONLY_DATA
 	FDragDropCheckingOverride OnDragDropCheckOverride;
@@ -1937,6 +1942,9 @@ private:
 	/** The hit-test radius of the cursor. Default value is 0. */
 	float CursorRadius;
 
+	/** The current DrawId. Value changes whenever we do a new pass*/
+	uint32 DrawId;
+
 	/**
 	 * All users currently registered with Slate.
 	 * Normally there's just one, but any case where multiple users can provide input to the same application will register multiple users.
diff --git a/Engine/Source/Runtime/Slate/Public/Widgets/Layout/SLinkedBox.cpp b/Engine/Source/Runtime/Slate/Public/Widgets/Layout/SLinkedBox.cpp
--- a/Engine/Source/Runtime/Slate/Public/Widgets/Layout/SLinkedBox.cpp
+++ b/Engine/Source/Runtime/Slate/Public/Widgets/Layout/SLinkedBox.cpp
@@ -6,6 +6,8 @@
 #include "Layout/ArrangedChildren.h"
 #include "Layout/LayoutUtils.h"
 #include "Widgets/Layout/SBorder.h"
+#include "Framework/Application/SlateApplication.h"
+
 
 FLinkedBoxManager::FLinkedBoxManager()
 : CachedUniformSize(FVector2D::ZeroVector)
@@ -28,7 +30,9 @@ void FLinkedBoxManager::UnregisterLinkedBox(SLinkedBox* InBox)
 
 FVector2D FLinkedBoxManager::GetUniformCellSize() const
 {
-	if (FrameCounterLastCached == GFrameCounter)
+	const uint32 CurrentDrawId = FSlateApplication::Get().GetCurrentDrawId();
+
+	if (SlateApplicationDrawIdLastCached == CurrentDrawId)
 	{
 		return CachedUniformSize;
 	}
@@ -48,7 +52,7 @@ FVector2D FLinkedBoxManager::GetUniformCellSize() const
 		}
 	}
 
-	FrameCounterLastCached = GFrameCounter;
+	SlateApplicationDrawIdLastCached = CurrentDrawId;
 	CachedUniformSize = MaxDesiredSize;
 	return CachedUniformSize;
 }
diff --git a/Engine/Source/Runtime/Slate/Public/Widgets/Layout/SLinkedBox.h b/Engine/Source/Runtime/Slate/Public/Widgets/Layout/SLinkedBox.h
--- a/Engine/Source/Runtime/Slate/Public/Widgets/Layout/SLinkedBox.h
+++ b/Engine/Source/Runtime/Slate/Public/Widgets/Layout/SLinkedBox.h
@@ -43,7 +43,7 @@ public:
 
 protected:
 	
-	mutable uint64 FrameCounterLastCached = 0;
+	mutable uint32 SlateApplicationDrawIdLastCached = 0;
 	mutable FVector2D CachedUniformSize;
 
 private:
