Fix a regression where Toolbar buttons with a combined options menu did not display that options menu when the toolbar button was clipped.

[RN] Editor, Minor

diff --git a/Engine/Source/Runtime/Slate/Private/Framework/MultiBox/MultiBox.cpp b/Engine/Source/Runtime/Slate/Private/Framework/MultiBox/MultiBox.cpp
--- a/Engine/Source/Runtime/Slate/Private/Framework/MultiBox/MultiBox.cpp
+++ b/Engine/Source/Runtime/Slate/Private/Framework/MultiBox/MultiBox.cpp
@@ -1554,6 +1554,24 @@ TSharedRef<SWidget> SMultiBoxWidget::OnWrapButtonClicked()
 			if (ClippedWidget == Pair.Key)
 			{
 				Pair.Value->CreateMenuEntry(MenuBuilder);
+				
+				// Entries that are followed by simple combo boxes have both widgets merged together to create a single button with an "options" dropdown.
+				// When overflowing, the primary button of this pair needs to handle placing this consumed child.
+				const TArray<TSharedRef<const FMultiBlock>>& Blocks = GetMultiBox()->GetBlocks();
+				int32 NextBlockIndex = Blocks.Find(Pair.Value.ToSharedRef()) + 1;
+				if (NextBlockIndex < Blocks.Num())
+				{
+					const TSharedRef<const FMultiBlock>& NextBlock = Blocks[NextBlockIndex];
+					if (NextBlock->GetType() == EMultiBlockType::ToolBarComboButton)
+					{
+						TSharedRef<const FToolBarComboButtonBlock> NextToolBarComboButtonBlock = StaticCastSharedRef<const FToolBarComboButtonBlock>(NextBlock);
+						if (NextToolBarComboButtonBlock->IsSimpleComboBox())
+						{
+							NextBlock->CreateMenuEntry(MenuBuilder);
+						}
+					}
+				}
+				
 				break;
 			}
 		}
