[Editor] Added respect for TextOverflowPolicy to the HyperLink widget

#rnx

diff --git a/Engine/Source/Runtime/Slate/Public/Widgets/Input/SHyperlink.h b/Engine/Source/Runtime/Slate/Public/Widgets/Input/SHyperlink.h
--- a/Engine/Source/Runtime/Slate/Public/Widgets/Input/SHyperlink.h
+++ b/Engine/Source/Runtime/Slate/Public/Widgets/Input/SHyperlink.h
@@ -41,6 +41,7 @@ public:
 		, _HighlightColor()
 		, _HighlightShape()
 		, _HighlightText()
+		, _OverflowPolicy()
 		{}
 
 		SLATE_ATTRIBUTE( FText, Text )
@@ -54,6 +55,7 @@ public:
 		SLATE_ATTRIBUTE( FLinearColor, HighlightColor )
 		SLATE_ATTRIBUTE( const FSlateBrush*, HighlightShape )
 		SLATE_ATTRIBUTE( FText, HighlightText )
+		SLATE_ARGUMENT(TOptional<ETextOverflowPolicy>, OverflowPolicy)
 	SLATE_END_ARGS()
 
 	/**
@@ -85,6 +87,7 @@ public:
 				.HighlightColor(InArgs._HighlightColor)
 				.HighlightShape(InArgs._HighlightShape)
 				.HighlightText(InArgs._HighlightText)
+				.OverflowPolicy(InArgs._OverflowPolicy)
 			]
 		);
 	}
