Slate: SSpinBox no longer fires OnValueChanged after OnValueCommitted when committing the value via typing in a number and hitting enter. This makes committing a value via the keyboard match the function callback order for when committing a value via dragging the spinbox.

When dragging the spinbox via the mouse, OnMouseMove was calling CommitValue with a CommitMethod CommitedViaSpin, which would call OnValueChanged, but would not call OnValueCommitted as that is only called when the CommitMethod is set to CommittedViaTypeIn or CommittedViaArrowKey. Then, in OnMouseButtonUp, if we were dragging (which we are), it directly calls NotifyValueCommitted (instead of CommitValue) and that specifically calls OnValueCommitted, which results in the order of "OnValueChanged (Repeat) -> OnValueCommitted" when using spin dragging, which is the logical order of operations.

When using the keyboard to type in a number and press enter, OnValueCommitted is never called and instead CommitValue is called directly with the CommitMethod set to CommittedViaTypeIn, meaning both OnValueChanged and OnValueComitted are broadcast from the one event. Unfortunately due to the implementation order these are actually currently called in the order of "OnValueCommitted -> OnValueChanged" which is the opposite of mouse dragging which makes it hard to implement a customization that uses both OnValueChanged and OnValueCommitted since the order of operations isn't guranteed.

diff --git a/Engine/Source/Editor/DetailCustomizations/Private/ComponentTransformDetails.cpp b/Engine/Source/Editor/DetailCustomizations/Private/ComponentTransformDetails.cpp
--- a/Engine/Source/Editor/DetailCustomizations/Private/ComponentTransformDetails.cpp
+++ b/Engine/Source/Editor/DetailCustomizations/Private/ComponentTransformDetails.cpp
@@ -1717,9 +1717,13 @@ void FComponentTransformDetails::OnSetTransform(ETransformField::Type TransformF
 		}
 	}
 
-	if (bCommitted && bBeganTransaction)
+	if (bCommitted)
 	{
-		GEditor->EndTransaction();
+		if (bBeganTransaction)
+		{
+			GEditor->EndTransaction();
+		}
+		
 		CacheDetails();
 	}
 
diff --git a/Engine/Source/Runtime/Slate/Private/Widgets/Input/SSpinBox.cpp b/Engine/Source/Runtime/Slate/Private/Widgets/Input/SSpinBox.cpp
--- a/Engine/Source/Runtime/Slate/Private/Widgets/Input/SSpinBox.cpp
+++ b/Engine/Source/Runtime/Slate/Private/Widgets/Input/SSpinBox.cpp
@@ -10,8 +10,19 @@ namespace SpinBoxPrivate
 			bUseSpinBoxMouseMoveOptimization,
 			TEXT("")
 		);
+		
+	TAutoConsoleVariable<bool> CVarSlateSpinboxUseWrongChangedCommitOrder(
+	TEXT("Slate.SpinBox.UseIncorrectChangeThenCommitEventOrder"),
+	false,
+	TEXT("If true, restores the < UE 5.7 behavior for SSpinBox widgets where, when dragging the mouse to change the value the"
+		" Spinbox fired the OnValueChanged event and then OnValueCommitted as the last event of the operation, but when entering a value"
+		" numerically via the Keyboard, it would fire OnValueCommitted and then OnValueChanged, which made it difficult to write code that"
+		" used both events due to their inconsistent calling order. This was made consistent in 5.7 (Changed then Committed in both cases)"
+		" but this cvar is provided temporarily for backwards compatibility to restore the previous behavior if it creates issues."),
+	ECVF_Default);
 }
 
+
 float SpinBoxComputeExponentSliderFraction(float FractionFilled, float StartFractionFilled, float SliderExponent)
 {
 	if (FractionFilled <= StartFractionFilled)
@@ -1014,12 +1025,24 @@ void SSpinBox<NumericType>::CommitValue(NumericType NewValue, double NewSpinValu
 		ApplySliderMinValueChanged(float(ValueAttribute.Get() - GetMinSliderValue()), true);
 	}
 
+	// OnValueChanged used to be called after OnValueCommitted when entering a value via the keyboard.
+	// This was backwards from the order they got called when dragging the mouse (which was 
+	// Changed and then Committed) so this fixes the inconsistency, but the cvar was provided to debug
+	// as there are hundreds of usages of this widget across the codebase.
+	if (!SpinBoxPrivate::CVarSlateSpinboxUseWrongChangedCommitOrder.GetValueOnAnyThread())
+	{
+		OnValueChanged.ExecuteIfBound(NewValue);
+	}
+
 	if (CommitMethod == CommittedViaTypeIn || CommitMethod == CommittedViaArrowKey)
 	{
 		OnValueCommitted.ExecuteIfBound(NewValue, OriginalCommitInfo);
 	}
 
-	OnValueChanged.ExecuteIfBound(NewValue);
+	if (SpinBoxPrivate::CVarSlateSpinboxUseWrongChangedCommitOrder.GetValueOnAnyThread())
+	{
+		OnValueChanged.ExecuteIfBound(NewValue);
+	}
 
 	if (!ValueAttribute.IsBound())
 	{
