Slate: SSpinBox now has an OnEditEditableText event, allowing customization of what text is edited when entering text mode

diff --git a/Engine/Source/Runtime/Slate/Private/Widgets/Input/SSpinBox.cpp b/Engine/Source/Runtime/Slate/Private/Widgets/Input/SSpinBox.cpp
--- a/Engine/Source/Runtime/Slate/Private/Widgets/Input/SSpinBox.cpp
+++ b/Engine/Source/Runtime/Slate/Private/Widgets/Input/SSpinBox.cpp
@@ -103,6 +103,7 @@ void SSpinBox<NumericType>::Construct(const FArguments& InArgs)
 	OnDynamicSliderMinValueChanged = InArgs._OnDynamicSliderMinValueChanged;
 	
 	OnGetDisplayValue = InArgs._OnGetDisplayValue;
+	OnEditEditableText = InArgs._OnEditEditableText;
 
 	bEnableWheel = InArgs._EnableWheel;
 	bBroadcastValueChangesPerKey = InArgs._BroadcastValueChangesPerKey;
@@ -867,6 +868,17 @@ void SSpinBox<NumericType>::EnterTextMode()
 {
 	TextBlock->SetVisibility(EVisibility::Collapsed);
 	EditableText->SetVisibility(EVisibility::Visible);
+
+	if (OnEditEditableText.IsBound())
+	{
+		if (const TOptional<FText> OptionalOverrideEditableText = OnEditEditableText.Execute(ValueAttribute.Get()))
+		{
+			const FText OverrideEditableText = OptionalOverrideEditableText.GetValue();
+			
+			EditableText->SetText(OverrideEditableText);
+			EditableText->SetEditableText(OverrideEditableText);
+		}
+	}
 }
 
 template <typename NumericType>
diff --git a/Engine/Source/Runtime/Slate/Public/Widgets/Input/SSpinBox.h b/Engine/Source/Runtime/Slate/Public/Widgets/Input/SSpinBox.h
--- a/Engine/Source/Runtime/Slate/Public/Widgets/Input/SSpinBox.h
+++ b/Engine/Source/Runtime/Slate/Public/Widgets/Input/SSpinBox.h
@@ -167,6 +167,8 @@ public:
 		SLATE_EVENT(FOnValueChanged, OnEndSliderMovement)
 		/** Called to allow customization of what text is displayed when not typing. An empty value falls back to the default behavior. */
 		SLATE_EVENT(FOnGetDisplayValue, OnGetDisplayValue)
+		/** Called to allow customization of what text is edited when entering text mode. An empty value falls back to the default behavior. */
+		SLATE_EVENT(FOnGetDisplayValue, OnEditEditableText)
 		/** Whether to clear keyboard focus when pressing enter to commit changes */
 		SLATE_ATTRIBUTE(bool, ClearKeyboardFocusOnCommit)
 		/** Whether to select all text when pressing enter to commit changes */
@@ -451,6 +453,7 @@ private:
 	FOnDynamicSliderMinMaxValueChanged OnDynamicSliderMaxValueChanged;
 	FOnDynamicSliderMinMaxValueChanged OnDynamicSliderMinValueChanged;
 	FOnGetDisplayValue OnGetDisplayValue;
+	FOnGetDisplayValue OnEditEditableText;
 
 	/** Prevents the spinbox from being smaller than desired in certain cases (e.g. when it is empty) */
 	TAttribute<float> MinDesiredWidth;
