PR #12722: Update SComboRow::OnMouseButtonDown

diff --git a/Engine/Source/Runtime/Slate/Public/Widgets/Input/SComboBox.h b/Engine/Source/Runtime/Slate/Public/Widgets/Input/SComboBox.h
--- a/Engine/Source/Runtime/Slate/Public/Widgets/Input/SComboBox.h
+++ b/Engine/Source/Runtime/Slate/Public/Widgets/Input/SComboBox.h
@@ -80,7 +80,7 @@ public:
 			if (bIsSelected)
 			{
 				// Reselect content to ensure selection is taken
-				OwnerWidget->Private_SignalSelectionChanged(ESelectInfo::Direct);
+				OwnerWidget->Private_SignalSelectionChanged(ESelectInfo::OnMouseClick);
 				return FReply::Handled();
 			}
 		}
