Slate: Adds support to the SHeaderRow to allow forcing sorting and sort direction, taking the choice away from manual user control.

diff --git a/Engine/Source/Runtime/Slate/Private/Widgets/Views/SHeaderRow.cpp b/Engine/Source/Runtime/Slate/Private/Widgets/Views/SHeaderRow.cpp
--- a/Engine/Source/Runtime/Slate/Private/Widgets/Views/SHeaderRow.cpp
+++ b/Engine/Source/Runtime/Slate/Private/Widgets/Views/SHeaderRow.cpp
@@ -38,6 +38,7 @@ public:
 		, SortPriority( EColumnSortPriority::Primary )
 		, IsSorting( false )
 		, OnSortModeChanged()
+		, CanManuallySort( true )
 		, ContextMenuContent( SNullWidget::NullWidget )
 		, ColumnId( NAME_None )
 		, Style( nullptr )
@@ -60,12 +61,15 @@ public:
 		SortMode = Column.SortMode;
 		SortPriority = Column.SortPriority;
 		IsSorting = Column.IsSorting;
+		CanManuallySort = Column.CanManuallySort;
 
 		OnSortModeChanged = Column.OnSortModeChanged;
 		ContextMenuContent = OverrideHeaderMenuContent;
 
 		ComboVisibility = Column.HeaderComboVisibility;
 
+		bHasBeenManuallySorted = false;
+
 		TAttribute< FText > LabelText = Column.DefaultText;
 		if (Column.HeaderContent.Widget == SNullWidget::NullWidget)
 		{
@@ -101,53 +105,71 @@ public:
 
 		if ( OnSortModeChanged.IsBound() )
 		{
-			//optional main button with the column's title. Used to toggle sorting modes.
-			Box->AddSlot()
-			.FillWidth( 1.0f )
-			.HAlign( HAlign_Fill )
-			[
-				SNew( SButton )
-				.ButtonStyle( FAppStyle::Get(), "NoBorder" )
-				.ForegroundColor( FSlateColor::UseForeground() )
-				.OnClicked( this, &STableColumnHeader::OnTitleClicked )
-				.ContentPadding( 0.0f )
+			TSharedRef<SWidget> SortedHeaderContent = SNew(SBox)
+				.HAlign(Column.HeaderHAlignment)
+				.VAlign(Column.HeaderVAlignment)
 				[
-					SNew(SBox)
-					.HAlign( Column.HeaderHAlignment )
-					.VAlign( Column.HeaderVAlignment )
-					[
-						SNew( SHorizontalBox )
-						+SHorizontalBox::Slot()
+					SNew(SHorizontalBox)
+						+ SHorizontalBox::Slot()
 						.AutoWidth()
-						.Padding( FMargin( 0.0f ) )
+						.Padding(FMargin(0.0f))
 						[
 							PrimaryContent
 						]
-						+SHorizontalBox::Slot()
+						+ SHorizontalBox::Slot()
 						.AutoWidth()
-						.HAlign( HAlign_Left )
-						.VAlign( VAlign_Center)
-						.Padding( FMargin( 4.0f, 0.0f, 0.0f, 0.0f) )
+						.HAlign(HAlign_Left)
+						.VAlign(VAlign_Center)
+						.Padding(FMargin(4.0f, 0.0f, 0.0f, 0.0f))
 						[
-							SNew( SOverlay )
-							+SOverlay::Slot()
-							[
-								SNew( SImage )
-								.ColorAndOpacity( FSlateColor::UseForeground() )
-								.Image( this, &STableColumnHeader::GetSortingBrush )
-								.Visibility( this, &STableColumnHeader::GetSortModeVisibility )
-							]
-							+SOverlay::Slot()
+							SNew(SOverlay)
+								+ SOverlay::Slot()
+								[
+									SNew(SImage)
+										.ColorAndOpacity(FSlateColor::UseForeground())
+										.Image(this, &STableColumnHeader::GetSortingBrush)
+										.Visibility(this, &STableColumnHeader::GetSortModeVisibility)
+								]
+								+ SOverlay::Slot()
+								[
+									SNew(SCircularThrobber)
+										.Radius(10.0f)
+										.PieceImage(FAppStyle::Get().GetBrush("Throbber.CircleChunk.Small"))
+										.Visibility(this, &STableColumnHeader::GetSortProcessingVisibility)
+								]
+						]
+				];
+
+			if (CanManuallySort.Get())
+			{
+				//optional main button with the column's title. Used to toggle sorting modes.
+				Box->AddSlot()
+					.FillWidth(1.0f)
+					.HAlign(HAlign_Fill)
+					[
+						SNew(SButton)
+							.ButtonStyle(FAppStyle::Get(), "NoBorder")
+							.ForegroundColor(FSlateColor::UseForeground())
+							.OnClicked(this, &STableColumnHeader::OnTitleClicked)
+							.ContentPadding(0.0f)
 							[
-								SNew(SCircularThrobber)
-								.Radius(10.0f)
-								.PieceImage(FAppStyle::Get().GetBrush("Throbber.CircleChunk.Small"))
-								.Visibility( this, &STableColumnHeader::GetSortProcessingVisibility )
+								SortedHeaderContent
 							]
-						]
-					]
-				]
-			];
+					];
+			}
+			else
+			{
+				// Variation that shows sorting direction without allowing user controls over sorting
+				Box->AddSlot()
+					.FillWidth(1.0f)
+					.HAlign(HAlign_Fill)
+					[
+						SortedHeaderContent							
+					];
+                   ApplySortMode();
+			}
+
+			
 		}
 		else
 		{
@@ -249,7 +271,14 @@ public:
 	/** Gets sorting mode */
 	EColumnSortMode::Type GetSortMode() const
 	{
-		return SortMode.Get();
+		if (SortMode.Get() == EColumnSortMode::None && !CanManuallySort.Get() )
+		{
+			return GetInitialSortMode();
+		}
+		else
+		{
+			return SortMode.Get();
+		}
 	}
 
 	/** Sets sorting mode */
@@ -389,7 +418,7 @@ private:
 	{
 		EColumnSortPriority::Type ColumnSortPriority = SortPriority.Get();
 
-		return ( SortMode.Get() == EColumnSortMode::Ascending ? 
+		return (GetSortMode() == EColumnSortMode::Ascending ?
 			(ColumnSortPriority == EColumnSortPriority::Secondary ? &Style->SortSecondaryAscendingImage : &Style->SortPrimaryAscendingImage) :
 			(ColumnSortPriority == EColumnSortPriority::Secondary ? &Style->SortSecondaryDescendingImage : &Style->SortPrimaryDescendingImage) );
 	}
@@ -397,8 +426,8 @@ private:
 	/** Checks if sorting mode has been selected */
 	EVisibility GetSortModeVisibility() const
 	{
-		EColumnSortMode::Type ColumnSortMode = SortMode.Get();
-		const bool bShow = ColumnSortMode != EColumnSortMode::None && !IsSorting.Get();
+		EColumnSortMode::Type ColumnSortMode = GetSortMode();
+		const bool bShow = ColumnSortMode != EColumnSortMode::None && bHasBeenManuallySorted && !IsSorting.Get();
 		return bShow ? EVisibility::HitTestInvisible : EVisibility::Collapsed;
 	}
 
@@ -411,13 +440,20 @@ private:
 	/** Called when the column title has been clicked to change sorting mode */
 	FReply OnTitleClicked()
 	{
-		if ( OnSortModeChanged.IsBound() )
+		bHasBeenManuallySorted = true;
+		ApplySortMode();
+		return FReply::Handled();
+	}
+
+	void ApplySortMode()
+	{
+		if (OnSortModeChanged.IsBound())
 		{
 			FSlateApplication::Get().CloseToolTip();
 
 			const bool bIsShiftClicked = FSlateApplication::Get().GetModifierKeys().IsShiftDown();
 			EColumnSortPriority::Type ColumnSortPriority = SortPriority.Get();
-			EColumnSortMode::Type ColumnSortMode = SortMode.Get();
+			EColumnSortMode::Type ColumnSortMode = GetSortMode();
 			if (ColumnSortMode == EColumnSortMode::None)
 			{
 				if (bIsShiftClicked && SortPriority.IsBound())
@@ -450,8 +486,6 @@ private:
 
 			OnSortModeChanged.Execute(ColumnSortPriority, ColumnId, ColumnSortMode);
 		}
-
-		return FReply::Handled();
 	}
 
 	void OpenContextMenu(const FPointerEvent& MouseEvent)
@@ -480,10 +514,13 @@ private:
 	
 	/** Whether or not the list under the header is being sorted. */
 	TAttribute< bool > IsSorting;
-	
+
 	/** Callback triggered when sorting mode changes */
 	FOnSortModeChanged OnSortModeChanged;
 
+	/** Whether or not the sorting mode should be under manual user control. */
+	TAttribute< bool > CanManuallySort;
+
 	TSharedRef< SWidget > ContextMenuContent;
 
 	TSharedPtr< SComboButton > ComboButton;
@@ -496,6 +533,8 @@ private:
 	FName ColumnId;
 
 	const FTableColumnHeaderStyle* Style;
+
+	bool bHasBeenManuallySorted = false;
 };
 
 void SHeaderRow::Construct( const FArguments& InArgs )
diff --git a/Engine/Source/Runtime/Slate/Public/Widgets/Views/SHeaderRow.h b/Engine/Source/Runtime/Slate/Public/Widgets/Views/SHeaderRow.h
--- a/Engine/Source/Runtime/Slate/Public/Widgets/Views/SHeaderRow.h
+++ b/Engine/Source/Runtime/Slate/Public/Widgets/Views/SHeaderRow.h
@@ -121,6 +121,7 @@ public:
 			, _InitialSortMode( EColumnSortMode::Ascending )
 			, _SortMode( EColumnSortMode::None )
 			, _IsSorting( false )
+			, _CanManuallySort(true)
 			, _OnSort()
 			, _OverflowPolicy( ETextOverflowPolicy::Clip)
 			{}
@@ -158,6 +159,7 @@ public:
 			SLATE_ATTRIBUTE( EColumnSortMode::Type, SortMode )
 			SLATE_ATTRIBUTE( EColumnSortPriority::Type, SortPriority )
 			SLATE_ATTRIBUTE( bool, IsSorting )
+			SLATE_ATTRIBUTE( bool, CanManuallySort)
 			SLATE_EVENT( FOnSortModeChanged, OnSort )
 
 			SLATE_ATTRIBUTE(bool, ShouldGenerateWidget)
@@ -188,6 +190,7 @@ public:
 			, SortMode( InArgs._SortMode )
 			, SortPriority( InArgs._SortPriority )
 			, IsSorting( InArgs._IsSorting)
+			, CanManuallySort(InArgs._CanManuallySort)
 			, OnSortModeChanged( InArgs._OnSort )
 			, ShouldGenerateWidget(InArgs._ShouldGenerateWidget)
 			, bIsVisible(true)
@@ -277,6 +280,7 @@ public:
 		TAttribute< EColumnSortMode::Type > SortMode;
 		TAttribute< EColumnSortPriority::Type > SortPriority;
 		TAttribute< bool > IsSorting;
+		TAttribute< bool > CanManuallySort;
 		FOnSortModeChanged OnSortModeChanged;
 
 		TAttribute<bool> ShouldGenerateWidget;
