PR #12579 Prevent Editable Text Clearing Focus It No Longer Owns

diff --git a/Engine/Source/Runtime/Slate/Private/Widgets/Text/SlateEditableTextLayout.cpp b/Engine/Source/Runtime/Slate/Private/Widgets/Text/SlateEditableTextLayout.cpp
--- a/Engine/Source/Runtime/Slate/Private/Widgets/Text/SlateEditableTextLayout.cpp
+++ b/Engine/Source/Runtime/Slate/Private/Widgets/Text/SlateEditableTextLayout.cpp
@@ -1099,7 +1099,7 @@ FReply FSlateEditableTextLayout::HandleKeyDown(const FKeyEvent& InKeyEvent)
 	else if (Key == EKeys::Enter && !OwnerWidget->IsTextReadOnly())
 	{
 		FScopedEditableTextTransaction TextTransaction(*this);
-		HandleCarriageReturn(InKeyEvent.IsRepeat());
+		HandleCarriageReturn(InKeyEvent.IsRepeat(), InKeyEvent.GetUserIndex());
 		Reply = FReply::Handled();
 	}
 	else if (Key == EKeys::Tab && OwnerWidget->CanTypeCharacter(TEXT('\t')))
@@ -1637,7 +1637,7 @@ bool FSlateEditableTextLayout::HandleTypeChar(const TCHAR InChar)
 	return false;
 }
 
-bool FSlateEditableTextLayout::HandleCarriageReturn(bool isRepeat)
+bool FSlateEditableTextLayout::HandleCarriageReturn(bool isRepeat, int32 UserIndex)
 {
 	if (OwnerWidget->IsTextReadOnly())
 	{
@@ -1679,7 +1679,11 @@ bool FSlateEditableTextLayout::HandleCarriageReturn(bool isRepeat)
 		// Release input focus?
 		if (OwnerWidget->ShouldClearKeyboardFocusOnCommit())
 		{
-			FSlateApplication::Get().ClearKeyboardFocus(EFocusCause::Cleared);
+			TSharedPtr<SWidget> OwnerSlateWidget = OwnerWidget->GetSlateWidgetPtr();
+			if (OwnerSlateWidget && OwnerSlateWidget->HasUserFocus(UserIndex))
+			{
+				FSlateApplication::Get().ClearKeyboardFocus(EFocusCause::Cleared);
+			}
 		}
 	}
 
diff --git a/Engine/Source/Runtime/Slate/Public/Widgets/Text/SlateEditableTextLayout.h b/Engine/Source/Runtime/Slate/Public/Widgets/Text/SlateEditableTextLayout.h
--- a/Engine/Source/Runtime/Slate/Public/Widgets/Text/SlateEditableTextLayout.h
+++ b/Engine/Source/Runtime/Slate/Public/Widgets/Text/SlateEditableTextLayout.h
@@ -206,7 +206,7 @@ public:
 	SLATE_API bool HandleTypeChar(const TCHAR InChar);
 
 	/** Called to handle a carriage return action acting on the current selection or at the cursor position */
-	SLATE_API bool HandleCarriageReturn(bool isRepeat);
+	SLATE_API bool HandleCarriageReturn(bool isRepeat, int32 UserIndex = 0);
 
 	/** Are we able to delete the currently selected text? */
 	SLATE_API bool CanExecuteDelete() const;
