Reduce redundant calls to FSlateUser::QueryCursor

diff --git a/Engine/Source/Runtime/Slate/Private/Framework/Application/SlateUser.cpp b/Engine/Source/Runtime/Slate/Private/Framework/Application/SlateUser.cpp
--- a/Engine/Source/Runtime/Slate/Private/Framework/Application/SlateUser.cpp
+++ b/Engine/Source/Runtime/Slate/Private/Framework/Application/SlateUser.cpp
@@ -384,10 +384,15 @@ void FSlateUser::SetCursorVisibility(bool bDrawCursor)
 	
 	if (bCanDrawCursor)
 	{
-		RequestCursorQuery();
+		if (bNeedQueryForVisibleCursor)
+		{
+			RequestCursorQuery();
+			bNeedQueryForVisibleCursor = false;
+		}
 	}
 	else
 	{
+		bNeedQueryForVisibleCursor = true;
 		ProcessCursorReply(FCursorReply::Cursor(EMouseCursor::None));
 	}
 }
diff --git a/Engine/Source/Runtime/Slate/Public/Framework/Application/SlateUser.h b/Engine/Source/Runtime/Slate/Public/Framework/Application/SlateUser.h
--- a/Engine/Source/Runtime/Slate/Public/Framework/Application/SlateUser.h
+++ b/Engine/Source/Runtime/Slate/Public/Framework/Application/SlateUser.h
@@ -219,6 +219,9 @@ private:
 	/** Whether this user is currently drawing the cursor onscreen each frame */
 	bool bCanDrawCursor = true;
 
+	/** Whether we need to call QueryCursor because bCanDrawCursor was enabled */
+	bool bNeedQueryForVisibleCursor = true;
+
 	/** The OS or actions taken by the user may require we refresh the current state of the cursor. */
 	bool bQueryCursorRequested = false;
 
