PR #13329: Slate: Add const getters for TabManager layout iteration

diff --git a/Engine/Source/Runtime/Slate/Private/Framework/Docking/TabManager.cpp b/Engine/Source/Runtime/Slate/Private/Framework/Docking/TabManager.cpp
--- a/Engine/Source/Runtime/Slate/Private/Framework/Docking/TabManager.cpp
+++ b/Engine/Source/Runtime/Slate/Private/Framework/Docking/TabManager.cpp
@@ -434,9 +434,9 @@ TSharedRef<FJsonObject> FTabManager::FLayout::PersistToString_Helper(const TShar
 		}
 				
 		TArray< TSharedPtr<FJsonValue> > Nodes;
-		for ( int32 ChildIndex=0; ChildIndex < NodeAsArea->ChildNodes.Num(); ++ChildIndex )
+		for ( const TSharedRef<FLayoutNode>& ChildNode : NodeAsArea->GetChildNodes() )
 		{
-			Nodes.Add( MakeShareable( new FJsonValueObject( PersistToString_Helper( NodeAsArea->ChildNodes[ChildIndex] ) ) ) );
+			Nodes.Add( MakeShareable( new FJsonValueObject( PersistToString_Helper( ChildNode ) ) ) );
 		}
 		JsonObj->SetArrayField( TEXT("Nodes"), Nodes );
 	}
@@ -446,9 +446,9 @@ TSharedRef<FJsonObject> FTabManager::FLayout::PersistToString_Helper(const TShar
 		JsonObj->SetStringField( TEXT("Orientation"), (NodeAsSplitter->GetOrientation() == Orient_Horizontal) ? TEXT("Orient_Horizontal") : TEXT("Orient_Vertical")  );
 
 		TArray< TSharedPtr<FJsonValue> > Nodes;
-		for ( int32 ChildIndex=0; ChildIndex < NodeAsSplitter->ChildNodes.Num(); ++ChildIndex )
+		for ( const TSharedRef<FLayoutNode>& ChildNode : NodeAsSplitter->GetChildNodes() )
 		{
-			Nodes.Add( MakeShareable( new FJsonValueObject( PersistToString_Helper( NodeAsSplitter->ChildNodes[ChildIndex] ) ) ) );
+			Nodes.Add( MakeShareable( new FJsonValueObject( PersistToString_Helper( ChildNode ) ) ) ) ;
 		}
 		JsonObj->SetArrayField( TEXT("Nodes"), Nodes );
 	}
@@ -1908,10 +1908,8 @@ bool FTabManager::CanRestoreSplitterContent(TArray<TSharedRef<SDockingNode>>& Do
 	}
 	DockingNodes.Empty();
 	// Restore the contents of this splitter.
-	for ( int32 ChildNodeIndex = 0; ChildNodeIndex < SplitterNode->ChildNodes.Num(); ++ChildNodeIndex )
+	for ( const TSharedRef<FLayoutNode>& ThisChildNode : SplitterNode->GetChildNodes() )
 	{
-		const TSharedRef<FLayoutNode> ThisChildNode = SplitterNode->ChildNodes[ChildNodeIndex];
-
 		const bool bEmbedTitleAreaContent = false;
 		const TSharedPtr<SDockingNode> ThisChildNodeWidget = RestoreArea_Helper(ThisChildNode, ParentWindow, bEmbedTitleAreaContent, OutSidebarTabs, OutputCanBeNullptr);
 		if (ThisChildNodeWidget)
@@ -1934,10 +1932,8 @@ void FTabManager::RestoreSplitterContent( const TArray<TSharedRef<SDockingNode>>
 void FTabManager::RestoreSplitterContent(const TSharedRef<FSplitter>& SplitterNode, const TSharedRef<SDockingSplitter>& SplitterWidget, const TSharedPtr<SWindow>& ParentWindow, FSidebarTabLists& OutSidebarTabs)
 {
 	// Restore the contents of this splitter.
-	for ( int32 ChildNodeIndex = 0; ChildNodeIndex < SplitterNode->ChildNodes.Num(); ++ChildNodeIndex )
+	for ( const TSharedRef<FLayoutNode>& ThisChildNode : SplitterNode->GetChildNodes() )
 	{
-		TSharedRef<FLayoutNode> ThisChildNode = SplitterNode->ChildNodes[ChildNodeIndex];
-
 		const bool bEmbedTitleAreaContent = false;
 		TSharedPtr<SDockingNode> ThisChildNodeWidget = RestoreArea_Helper(ThisChildNode, ParentWindow, bEmbedTitleAreaContent, OutSidebarTabs);
 		check(ThisChildNodeWidget.IsValid());
@@ -2321,9 +2317,9 @@ bool FTabManager::HasAnyMatchingTabs( const TSharedRef<FTabManager::FLayoutNode>
 	{
 		ensure( AsSplitter.IsValid() );
 		// Do any of the child nodes have open tabs?
-		for (int32 ChildIndex=0; ChildIndex < AsSplitter->ChildNodes.Num(); ++ChildIndex)
+		for (const TSharedRef<FTabManager::FLayoutNode>& ChildNode : AsSplitter->GetChildNodes())
 		{
-			if ( HasAnyMatchingTabs(AsSplitter->ChildNodes[ChildIndex], Matcher) )
+			if ( HasAnyMatchingTabs(ChildNode, Matcher) )
 			{
 				return true;
 			}
@@ -2493,9 +2489,9 @@ TSharedPtr<FTabManager::FStack> FTabManager::FindTabUnderNode( const FTabMatcher
 	{
 		ensure( NodeAsSplitter.IsValid() );
 		TSharedPtr<FTabManager::FStack> StackWithTab;
-		for ( int32 ChildIndex=0; !StackWithTab.IsValid() && ChildIndex < NodeAsSplitter->ChildNodes.Num(); ++ChildIndex)
+		for ( const TSharedRef<FTabManager::FLayoutNode>& ChildNode : NodeAsSplitter->GetChildNodes() )
 		{
-			StackWithTab = FindTabUnderNode( Matcher, NodeAsSplitter->ChildNodes[ChildIndex] );
+			StackWithTab = FindTabUnderNode( Matcher, ChildNode );
 		}
 
 			return StackWithTab;
diff --git a/Engine/Source/Runtime/Slate/Public/Framework/Docking/TabManager.h b/Engine/Source/Runtime/Slate/Public/Framework/Docking/TabManager.h
--- a/Engine/Source/Runtime/Slate/Public/Framework/Docking/TabManager.h
+++ b/Engine/Source/Runtime/Slate/Public/Framework/Docking/TabManager.h
@@ -5,6 +5,7 @@
 #include "CoreMinimal.h"
 #include "SlateFwd.h"
 #include "Containers/Ticker.h"
+#include "Containers/ContainersFwd.h"
 #include "Misc/Attribute.h"
 #include "Textures/SlateIcon.h"
 #include "Widgets/SWindow.h"
@@ -537,6 +538,11 @@ class FTabManager : public TSharedFromThis<FTabManager>
 					return ExtensionId;
 				}
 
+				TConstArrayView<FTab> GetTabs() const
+				{
+					return Tabs;
+				}
+
 			protected:
 
 				FStack()
@@ -606,6 +612,11 @@ class FTabManager : public TSharedFromThis<FTabManager>
 				{
 				}
 
+				TConstArrayView<TSharedRef<FLayoutNode>> GetChildNodes() const
+				{
+					return ChildNodes;
+				}
+
 			protected:
 				FSplitter()
 				: Orientation(Orient_Horizontal)
@@ -724,6 +735,11 @@ class FTabManager : public TSharedFromThis<FTabManager>
 				{
 					return PrimaryArea;
 				}
+
+				TConstArrayView<TSharedRef<FArea>> GetAreas() const
+				{
+					return Areas;
+				}
 				
 			public:
 				static SLATE_API TSharedPtr<FTabManager::FLayout> NewFromString( const FString& LayoutAsText );
