Slate: SMultiLineEditableTextBox - Disallow horizontal scrollbars when autowrap is enabled.
Fixes an issue with jittering content within scrollboxes.
Allowing both creates a layout feedback loop within scrollboxes when the autowrapped content nearly fits with the alloted horizontal space after the parent scrollbar becomes visible.

#rnx

diff --git a/Engine/Source/Runtime/Slate/Private/Widgets/Text/SMultiLineEditableText.cpp b/Engine/Source/Runtime/Slate/Private/Widgets/Text/SMultiLineEditableText.cpp
--- a/Engine/Source/Runtime/Slate/Private/Widgets/Text/SMultiLineEditableText.cpp
+++ b/Engine/Source/Runtime/Slate/Private/Widgets/Text/SMultiLineEditableText.cpp
@@ -52,7 +52,12 @@ void SMultiLineEditableText::Construct( const FArguments& InArgs )
 	ModiferKeyForNewLine = InArgs._ModiferKeyForNewLine;
 
 	HScrollBar = InArgs._HScrollBar;
-	if (HScrollBar.IsValid())
+
+	// Horizontal scroll is unecessary when auto wrap is enabled.
+	// Otherwise if allowed, having both enabled will cause a wrap+scroll relayout feedback loop which can result in jittering
+	const bool bAutoWrapEnabled = InArgs._AutoWrapText.IsSet() && InArgs._AutoWrapText.Get();
+
+    if (HScrollBar.IsValid() && !bAutoWrapEnabled)
 	{
 		HScrollBar->SetUserVisibility(EVisibility::Collapsed);
 		HScrollBar->SetOnUserScrolled(FOnUserScrolled::CreateSP(this, &SMultiLineEditableText::OnHScrollBarMoved));
