PR #13048: Fixed WidgetComponent being blocked by SVirtualJoystick

#rnx

diff --git a/Engine/Source/Runtime/Slate/Private/Widgets/Input/SVirtualJoystick.cpp b/Engine/Source/Runtime/Slate/Private/Widgets/Input/SVirtualJoystick.cpp
--- a/Engine/Source/Runtime/Slate/Private/Widgets/Input/SVirtualJoystick.cpp
+++ b/Engine/Source/Runtime/Slate/Private/Widgets/Input/SVirtualJoystick.cpp
@@ -557,6 +557,9 @@ void SVirtualJoystick::SetJoystickVisibility(const bool bInVisible, const bool b
 	}
 
 	bVisible = bInVisible;
+
+	//Update slate visibility - we still want to tick, so we use HitTestInvisible to avoid blocking other UI
+	SetVisibility(bVisible ? EVisibility::Visible : EVisibility::HitTestInvisible);
 }
 
 void SVirtualJoystick::AddControl(const FControlInfo& Control)
