UMG - Slot Details tests added

#rnx

#virtualized

diff --git a/Engine/Source/Runtime/Slate/Private/Widgets/Input/SSpinBox.cpp b/Engine/Source/Runtime/Slate/Private/Widgets/Input/SSpinBox.cpp
--- a/Engine/Source/Runtime/Slate/Private/Widgets/Input/SSpinBox.cpp
+++ b/Engine/Source/Runtime/Slate/Private/Widgets/Input/SSpinBox.cpp
@@ -699,11 +699,11 @@ NumericType SSpinBox<NumericType>::GetValue() const
 { return ValueAttribute.Get(); }
 
 template <typename NumericType>
-void SSpinBox<NumericType>::SetValue(const TAttribute<NumericType>& InValueAttribute)
+void SSpinBox<NumericType>::SetValue(const TAttribute<NumericType>& InValueAttribute, const bool bShouldCommit)
 {
 	ValueAttribute = InValueAttribute;
 	const NumericType LocalValueAttribute = ValueAttribute.Get();
-	CommitValue(LocalValueAttribute, (double)LocalValueAttribute, ECommitMethod::CommittedViaCode, ETextCommit::Default);
+	CommitValue(LocalValueAttribute, (double)LocalValueAttribute, ECommitMethod::CommittedViaCode, ETextCommit::Default, bShouldCommit);
 }
 
 template <typename NumericType>
@@ -967,7 +967,7 @@ void SSpinBox<NumericType>::TextField_OnTextCommitted(const FText& NewText, ETex
 }
 
 template <typename NumericType>
-void SSpinBox<NumericType>::CommitValue(NumericType NewValue, double NewSpinValue, ECommitMethod CommitMethod, ETextCommit::Type OriginalCommitInfo)
+void SSpinBox<NumericType>::CommitValue(NumericType NewValue, double NewSpinValue, ECommitMethod CommitMethod, ETextCommit::Type OriginalCommitInfo, const bool bShouldCommit)
 {
 	TRACE_CPUPROFILER_EVENT_SCOPE(SSpinBox_CommitValue);
 	if (CommitMethod == CommittedViaSpin || CommitMethod == CommittedViaArrowKey)
@@ -1026,7 +1026,7 @@ void SSpinBox<NumericType>::CommitValue(NumericType NewValue, double NewSpinValu
 		ApplySliderMinValueChanged(float(ValueAttribute.Get() - GetMinSliderValue()), true);
 	}
 
-	if (CommitMethod == CommittedViaTypeIn || CommitMethod == CommittedViaArrowKey)
+	if (CommitMethod == CommittedViaTypeIn || CommitMethod == CommittedViaArrowKey || bShouldCommit)
 	{
 		OnValueCommitted.ExecuteIfBound(NewValue, OriginalCommitInfo);
 	}
diff --git a/Engine/Source/Runtime/Slate/Public/Widgets/Input/SSpinBox.h b/Engine/Source/Runtime/Slate/Public/Widgets/Input/SSpinBox.h
--- a/Engine/Source/Runtime/Slate/Public/Widgets/Input/SSpinBox.h
+++ b/Engine/Source/Runtime/Slate/Public/Widgets/Input/SSpinBox.h
@@ -263,7 +263,7 @@ public:
 	/** See the Value attribute */
 	SLATE_API NumericType GetValue() const;
 
-	SLATE_API void SetValue(const TAttribute<NumericType>& InValueAttribute);
+	SLATE_API void SetValue(const TAttribute<NumericType>& InValueAttribute, const bool bShouldCommit = false);
 
 	/** See the MinValue attribute */
 	SLATE_API NumericType GetMinValue() const;
@@ -380,8 +380,9 @@ protected:
 	 * @param NewValue               Value resulting from the user's interaction
 	 * @param CommitMethod           Did the user type in the value or spin to it.
 	 * @param OriginalCommitInfo	 If the user typed in the value, information about the source of the commit
+	 * @param bShouldCommit          Should the value be committed (despite the CommitMethod parameter)
 	 */
-	SLATE_API void CommitValue(NumericType NewValue, double NewSpinValue, ECommitMethod CommitMethod, ETextCommit::Type OriginalCommitInfo);
+	SLATE_API void CommitValue(NumericType NewValue, double NewSpinValue, ECommitMethod CommitMethod, ETextCommit::Type OriginalCommitInfo, const bool bShouldCommit = false);
 
 	SLATE_API void NotifyValueCommitted(NumericType CurrentValue) const;
 
