[Slate] Added description on why the debugging mode didn't work.

#rnx

diff --git a/Engine/Source/Runtime/Slate/Private/Framework/Application/SlateApplication.cpp b/Engine/Source/Runtime/Slate/Private/Framework/Application/SlateApplication.cpp
--- a/Engine/Source/Runtime/Slate/Private/Framework/Application/SlateApplication.cpp
+++ b/Engine/Source/Runtime/Slate/Private/Framework/Application/SlateApplication.cpp
@@ -3786,6 +3786,24 @@ void FSlateApplication::EnterDebuggingMode()
 		static volatile bool bDoDebugBreak = true;
 		if (bDoDebugBreak)
 		{
+			//The application was not able to enter debugging mode.
+			//The user will see a notification, but Blueprint breakpoint will not stop the execution.
+			//We detected that the debugging mode would prevent the editor from ticking.
+			//Entering debugging mode could make the application unresponsive or corrupt the application's message stack.
+			// 
+			//This can occur when you are debugging an editor blueprint that is executed inside a modal.
+			//ie.
+			// Engine::Tick
+			//   RenamePrompt. ShowModal that blocks execution until the user presses a button.
+			//    Inside the rename modal window, code executes that triggers a blueprint.
+			// 
+			//This can occur when you are debugging a Widget that is painting.
+			//Game widgets are sadly included in the editor paint. It is related to how Slate rendering works.
+			//ie.
+			// World::Tick (no problem here)
+			// SlateApplication::TickAndPaint
+			//  SWindow::Paint // paint the editor
+			//    SMyWidget::Paint (includes OnPaint, Tick, and other callbacks).
 			UE_DEBUG_BREAK();
 		}
 	};
