PR #13555: Expose MinSize to SHeaderRow.

#rnx

diff --git a/Engine/Source/Runtime/Slate/Private/Widgets/Views/SHeaderRow.cpp b/Engine/Source/Runtime/Slate/Private/Widgets/Views/SHeaderRow.cpp
--- a/Engine/Source/Runtime/Slate/Private/Widgets/Views/SHeaderRow.cpp
+++ b/Engine/Source/Runtime/Slate/Private/Widgets/Views/SHeaderRow.cpp
@@ -870,6 +870,7 @@ void SHeaderRow::RegenerateWidgets()
 						.Value(WidthBinding)
 						.SizeRule(SSplitter::FractionOfParent)
 						.OnSlotResized(SSplitter::FOnSlotResized::CreateRaw(&SomeColumn, &FColumn::SetWidth))
+						.MinSize(SomeColumn.MinSize)
 						[
 							NewHeader
 						];
@@ -959,6 +960,7 @@ void SHeaderRow::RegenerateWidgets()
 					Splitter->AddSlot()
 						.SizeRule(SSplitter::SizeToContent)
 						.Resizable(false)
+						.MinSize(SomeColumn.MinSize)
 						[
 							SNew(SBox)
 							.WidthOverride(WidthBinding)
@@ -993,6 +995,7 @@ void SHeaderRow::RegenerateWidgets()
 						.SizeRule(SSplitter::SizeToContent)
 						.Resizable(true)
 						.OnSlotResized(SSplitter::FOnSlotResized::CreateRaw(&SomeColumn, &FColumn::SetWidth))
+						.MinSize(SomeColumn.MinSize)
 						[
 							SNew(SBox)
 							.WidthOverride(WidthBinding)
diff --git a/Engine/Source/Runtime/Slate/Public/Widgets/Views/SHeaderRow.h b/Engine/Source/Runtime/Slate/Public/Widgets/Views/SHeaderRow.h
--- a/Engine/Source/Runtime/Slate/Public/Widgets/Views/SHeaderRow.h
+++ b/Engine/Source/Runtime/Slate/Public/Widgets/Views/SHeaderRow.h
@@ -110,6 +110,7 @@ public:
 			, _FixedWidth()
 			, _ManualWidth()
 			, _OnWidthChanged()
+			, _MinSize( 0.0f )
 			, _HeaderContent()
 			, _HAlignHeader( HAlign_Fill )
 			, _VAlignHeader( VAlign_Fill )
@@ -140,6 +141,9 @@ public:
 			SLATE_ARGUMENT(TOptional< float >, FillSized)
 			SLATE_EVENT( FOnWidthChanged, OnWidthChanged )
 
+			/** Set the minimum column size when resizing. */
+			SLATE_ARGUMENT(float, MinSize)
+
 			SLATE_DEFAULT_SLOT( FArguments, HeaderContent )
 			SLATE_ARGUMENT( EHorizontalAlignment, HAlignHeader )
 			SLATE_ARGUMENT( EVerticalAlignment, VAlignHeader )
@@ -174,6 +178,7 @@ public:
 			, Width( 1.0f )
 			, DefaultWidth( 1.0f )
 			, OnWidthChanged( InArgs._OnWidthChanged)
+			, MinSize( InArgs._MinSize )
 			, SizeRule( EColumnSizeMode::Fill )
 			, HeaderContent( InArgs._HeaderContent )
 			, HeaderMenuContent( InArgs._MenuContent )
@@ -258,6 +263,9 @@ public:
 
 		FOnWidthChanged OnWidthChanged;
 
+		/** The minimum column size when resizing. */
+		float MinSize;
+
 		EColumnSizeMode::Type SizeRule;
 
 		TAlwaysValidWidget HeaderContent;	
