TabManager: Allow to overide text overflow policy in tab titles

Allow to pass an empty optional to disable text overflow

#rnx

diff --git a/Engine/Source/Runtime/Slate/Private/Widgets/Docking/SDockTab.cpp b/Engine/Source/Runtime/Slate/Private/Widgets/Docking/SDockTab.cpp
--- a/Engine/Source/Runtime/Slate/Private/Widgets/Docking/SDockTab.cpp
+++ b/Engine/Source/Runtime/Slate/Private/Widgets/Docking/SDockTab.cpp
@@ -913,7 +913,7 @@ void SDockTab::SetTabLabelSuffix(const TAttribute<FText>& InTabLabelSuffix)
 	TabLabelSuffix = InTabLabelSuffix;
 }
 
-void SDockTab::SetTabLabelOverflowPolicy(ETextOverflowPolicy InOverflowPolicy) const
+void SDockTab::SetTabLabelOverflowPolicy(TOptional<ETextOverflowPolicy> InOverflowPolicy) const
 {
 	if (LabelWidget.IsValid())
 	{
diff --git a/Engine/Source/Runtime/Slate/Public/Widgets/Docking/SDockTab.h b/Engine/Source/Runtime/Slate/Public/Widgets/Docking/SDockTab.h
--- a/Engine/Source/Runtime/Slate/Public/Widgets/Docking/SDockTab.h
+++ b/Engine/Source/Runtime/Slate/Public/Widgets/Docking/SDockTab.h
@@ -222,7 +222,7 @@ public:
 	SLATE_API void SetTabLabelSuffix(const TAttribute<FText>& InTabLabelSuffix);
 
 	/** Set the Label overflow policy */
-	SLATE_API void SetTabLabelOverflowPolicy(ETextOverflowPolicy InOverflowPolicy) const;
+	SLATE_API void SetTabLabelOverflowPolicy(TOptional<ETextOverflowPolicy> InOverflowPolicy) const;
 
 	/** The tooltip text that appears on the tab. */
 	SLATE_API void SetTabToolTipWidget(TSharedPtr<SToolTip> InTabToolTipWidget);
