Add platform text field modular feature support

diff --git a/Engine/Source/Runtime/Slate/Private/Framework/Application/IPlatformTextField.cpp b/Engine/Source/Runtime/Slate/Private/Framework/Application/IPlatformTextField.cpp
--- a/Engine/Source/Runtime/Slate/Private/Framework/Application/IPlatformTextField.cpp
+++ b/Engine/Source/Runtime/Slate/Private/Framework/Application/IPlatformTextField.cpp
@@ -1,6 +1,7 @@
 // Copyright Epic Games, Inc. All Rights Reserved.
 
 #include "Framework/Application/IPlatformTextField.h"
+#include "Features/IModularFeatures.h"
 #include "Misc/ConfigCacheIni.h"
 #include "HAL/Platform.h"
 #include "Internationalization/Culture.h"
@@ -87,4 +88,25 @@ bool IPlatformTextField::ShouldUseVirtualKeyboardAutocorrect(TSharedPtr<IVirtual
 	}
 
 	return bUseAutocorrect;
-}
\ No newline at end of file
+}
+
+
+
+FName IPlatformTextFieldFactory::FeatureName = FName("PlatformTextField");
+
+TUniquePtr<IPlatformTextField> IPlatformTextFieldFactory::TryCreateInstance()
+{
+	TUniquePtr<IPlatformTextField> Result;
+
+	TArray<IPlatformTextFieldFactory*> TextFieldFactories = IModularFeatures::Get().GetModularFeatureImplementations<IPlatformTextFieldFactory>(IPlatformTextFieldFactory::FeatureName);
+	for (IPlatformTextFieldFactory* TextFieldFactory : TextFieldFactories)
+	{
+		Result = TextFieldFactory->CreateInstance();
+		if (Result.IsValid())
+		{
+			break;
+		}
+	}
+
+	return MoveTemp(Result);
+}
diff --git a/Engine/Source/Runtime/Slate/Private/Framework/Application/SlateApplication.cpp b/Engine/Source/Runtime/Slate/Private/Framework/Application/SlateApplication.cpp
--- a/Engine/Source/Runtime/Slate/Private/Framework/Application/SlateApplication.cpp
+++ b/Engine/Source/Runtime/Slate/Private/Framework/Application/SlateApplication.cpp
@@ -4198,6 +4198,11 @@ void FSlateApplication::ShowVirtualKeyboard( bool bShow, int32 UserIndex, TShare
 {
 	SCOPE_CYCLE_COUNTER(STAT_ShowVirtualKeyboard);
 
+	if (!SlateTextField.IsValid())
+	{
+		SlateTextField = IPlatformTextFieldFactory::TryCreateInstance();
+	}
+
 	if (!SlateTextField.IsValid())
 	{
 		SlateTextField = MakeUnique<FPlatformTextField>();
@@ -4208,6 +4213,11 @@ void FSlateApplication::ShowVirtualKeyboard( bool bShow, int32 UserIndex, TShare
 
 bool FSlateApplication::AllowMoveCursor()
 {
+	if (!SlateTextField.IsValid())
+	{
+		SlateTextField = IPlatformTextFieldFactory::TryCreateInstance();
+	}
+
 	if (!SlateTextField.IsValid())
 	{
 		SlateTextField = MakeUnique<FPlatformTextField>();
diff --git a/Engine/Source/Runtime/Slate/Public/Framework/Application/IPlatformTextField.h b/Engine/Source/Runtime/Slate/Public/Framework/Application/IPlatformTextField.h
--- a/Engine/Source/Runtime/Slate/Public/Framework/Application/IPlatformTextField.h
+++ b/Engine/Source/Runtime/Slate/Public/Framework/Application/IPlatformTextField.h
@@ -5,8 +5,12 @@
 #include "CoreMinimal.h"
 #include "HAL/Platform.h"
 #include "Templates/SharedPointer.h"
+#include "Features/IModularFeature.h"
 #include "Widgets/Input/IVirtualKeyboardEntry.h"
 
+#define UE_API SLATE_API
+
+
 class IVirtualKeyboardEntry;
 
 class IPlatformTextField
@@ -22,3 +26,18 @@ public:
 private:
 
 };
+
+
+class IPlatformTextFieldFactory : public IModularFeature
+{
+public:
+	UE_API static FName FeatureName;
+
+	UE_API static TUniquePtr<IPlatformTextField> TryCreateInstance();
+	
+	virtual TUniquePtr<IPlatformTextField> CreateInstance() = 0;
+
+
+};
+
+#undef UE_API
