Slate Dynamic Invalidation

The Dynamic Invalidation relies on the widgets supporting invalidation being marked as such. This information is added to the widget class data `FSlateWidgetClassData` by declaring the widget using the `SLATE_DECLARE_WIDGET_SUPPORTING_INVALIDATION/_API` macro. This  allows the attribute to not be automatically inherited by descendant widgets.

The Dynamic Invalidation itself is handled by the `SInvalidationPanel` which has a new Editor Only field `bUseDynamicInvalidation` which when true, asks each child widget if they support invalidation. This is done via a new `SWidget::SupportsInvalidationRecursive` which checks if the widget class `SupportsInvalidation`, and if all of its children `SupportsInvalidationRecursive` too. This value is cached to minimize the cost.

To make this easier to debug, this CL also adds a stoplight indicator as well as some information in the tooltips of the Widget Reflector to indicate the status of the child widgets of a `SInvalidationPanel` with Dynamic Invalidation enabled.

This Dynamic Invalidation is currently only intended to work in the Editor.

diff --git a/Engine/Source/Developer/SlateReflector/Private/Models/WidgetReflectorNode.cpp b/Engine/Source/Developer/SlateReflector/Private/Models/WidgetReflectorNode.cpp
--- a/Engine/Source/Developer/SlateReflector/Private/Models/WidgetReflectorNode.cpp
+++ b/Engine/Source/Developer/SlateReflector/Private/Models/WidgetReflectorNode.cpp
@@ -140,6 +140,16 @@ bool FLiveWidgetReflectorNode::GetWidgetVisibilityInherited() const
 	return FWidgetReflectorNodeUtils::GetWidgetVisibilityInherited(Widget.Pin());
 }
 
+bool FLiveWidgetReflectorNode::GetWidgetSupportsInvalidation() const
+{
+	return FWidgetReflectorNodeUtils::GetWidgetSupportsInvalidation(Widget.Pin());
+}
+
+bool FLiveWidgetReflectorNode::GetSupportsInvalidationRecursive() const
+{
+	return FWidgetReflectorNodeUtils::GetWidgetSupportsInvalidationRecursive(Widget.Pin());
+}
+
 FText FLiveWidgetReflectorNode::GetWidgetClippingText() const
 {
 	return FWidgetReflectorNodeUtils::GetWidgetClippingText(Widget.Pin());
@@ -270,6 +280,8 @@ FSnapshotWidgetReflectorNode::FSnapshotWidgetReflectorNode(const FArrangedWidget
 	, CachedWidgetVisibilityText(FWidgetReflectorNodeUtils::GetWidgetVisibilityText(InArrangedWidget.Widget))
 	, bCachedWidgetVisible(FWidgetReflectorNodeUtils::GetWidgetVisibility(InArrangedWidget.Widget))
 	, bCachedWidgetVisibleInherited(FWidgetReflectorNodeUtils::GetWidgetVisibilityInherited(InArrangedWidget.Widget))
+	, bCachedSupportsInvalidation(FWidgetReflectorNodeUtils::GetWidgetSupportsInvalidation(InArrangedWidget.Widget))
+	, bCachedSupportsInvalidationRecursive(FWidgetReflectorNodeUtils::GetWidgetSupportsInvalidationRecursive(InArrangedWidget.Widget))
 	, bCachedWidgetFocusable(FWidgetReflectorNodeUtils::GetWidgetFocusable(InArrangedWidget.Widget))
 	, bCachedWidgetNeedsTick(FWidgetReflectorNodeUtils::GetWidgetNeedsTick(InArrangedWidget.Widget))
 	, bCachedWidgetIsVolatile(FWidgetReflectorNodeUtils::GetWidgetIsVolatile(InArrangedWidget.Widget))
@@ -332,6 +344,16 @@ bool FSnapshotWidgetReflectorNode::GetWidgetVisibilityInherited() const
 	return bCachedWidgetVisibleInherited;
 }
 
+bool FSnapshotWidgetReflectorNode::GetWidgetSupportsInvalidation() const
+{
+	return bCachedSupportsInvalidation;
+}
+
+bool FSnapshotWidgetReflectorNode::GetSupportsInvalidationRecursive() const
+{
+	return bCachedSupportsInvalidationRecursive;
+}
+
 FText FSnapshotWidgetReflectorNode::GetWidgetClippingText() const
 {
 	return CachedWidgetClippingText;
@@ -954,6 +976,24 @@ bool FWidgetReflectorNodeUtils::GetWidgetVisibilityInherited(const TSharedPtr<co
 	return InWidget.IsValid() ? InWidget->GetProxyHandle().GetWidgetVisibility(InWidget.Get()).IsVisible() : false;
 }
 
+bool FWidgetReflectorNodeUtils::GetWidgetSupportsInvalidation(const TSharedPtr<const SWidget>& InWidget)
+{
+#if UE_SLATE_WITH_DYNAMIC_INVALIDATION
+	return InWidget.IsValid() ? InWidget->SupportsInvalidation() : false;
+#else
+	return false;
+#endif
+}
+
+bool FWidgetReflectorNodeUtils::GetWidgetSupportsInvalidationRecursive(const TSharedPtr<const SWidget>& InWidget)
+{
+#if UE_SLATE_WITH_DYNAMIC_INVALIDATION
+	return InWidget.IsValid() ? InWidget->SupportsInvalidationRecursive() : false;
+#else
+	return false;
+#endif
+}
+
 bool FWidgetReflectorNodeUtils::GetWidgetFocusable(const TSharedPtr<const SWidget>& InWidget)
 {
 	return InWidget.IsValid() ? InWidget->SupportsKeyboardFocus() : false;
diff --git a/Engine/Source/Developer/SlateReflector/Private/Models/WidgetReflectorNode.h b/Engine/Source/Developer/SlateReflector/Private/Models/WidgetReflectorNode.h
--- a/Engine/Source/Developer/SlateReflector/Private/Models/WidgetReflectorNode.h
+++ b/Engine/Source/Developer/SlateReflector/Private/Models/WidgetReflectorNode.h
@@ -97,6 +97,16 @@ public:
 	 */
 	virtual bool GetWidgetVisibilityInherited() const = 0;
 
+	/**
+	 * @return The bool indicating whether or not the widget we were initialized from reports as supporting invalidation
+	 */
+	virtual bool GetWidgetSupportsInvalidation() const = 0;
+
+	/**
+	 * @return The bool indicating whether or not the widget we were initialized from reports as recursively supporting invalidation
+	 */
+	virtual bool GetSupportsInvalidationRecursive() const = 0;
+
 	/**
 	 * @return The bool indicating whether or not the widget we were initialized from reports needs tick
 	 */
@@ -286,6 +296,8 @@ public:
 	virtual FText GetWidgetVisibilityText() const override;
 	virtual bool GetWidgetVisible() const override;
 	virtual bool GetWidgetVisibilityInherited() const override;
+	virtual bool GetWidgetSupportsInvalidation() const override;
+	virtual bool GetSupportsInvalidationRecursive() const override;
 	virtual FText GetWidgetClippingText() const override;
 	virtual int32 GetWidgetLayerId() const override;
 	virtual int32 GetWidgetLayerIdOut() const override;
@@ -351,6 +363,8 @@ public:
 	virtual int32 GetWidgetLayerIdOut() const override;
 	virtual bool GetWidgetVisible() const override;
 	virtual bool GetWidgetVisibilityInherited() const override;
+	virtual bool GetWidgetSupportsInvalidation() const override;
+	virtual bool GetSupportsInvalidationRecursive() const override;
 	virtual bool GetWidgetFocusable() const override;
 	virtual bool GetWidgetNeedsTick() const override;
 	virtual bool GetWidgetIsVolatile() const override;
@@ -402,6 +416,12 @@ private:
 	
 	/** The visible inherited of the widget at the point it was passed to Initialize */
 	bool bCachedWidgetVisibleInherited;
+	
+	/** Did the widget support invalidation at the point it was passed to Initialize */
+	bool bCachedSupportsInvalidation;
+	
+	/** Did the widget recursively support invalidation at the point it was passed to Initialize */
+	bool bCachedSupportsInvalidationRecursive;
 
 	/** The focusability of the widget at the point it was passed to Initialize */
 	bool bCachedWidgetFocusable;
@@ -567,6 +587,12 @@ public:
 
 	/** Is the widget visible and its parents are also visible? */
 	static bool GetWidgetVisibilityInherited(const TSharedPtr<const SWidget>& InWidget);
+	
+	/** Does the widget supports invalidation? */
+	static bool GetWidgetSupportsInvalidation(const TSharedPtr<const SWidget>& InWidget);
+
+	/** Does this widget and all its children support invalidation? */
+	static bool GetWidgetSupportsInvalidationRecursive(const TSharedPtr<const SWidget>& InWidget);
 
 	/**
 	 * @return The current clipping string for the given widget
diff --git a/Engine/Source/Developer/SlateReflector/Private/Widgets/SWidgetReflectorToolTipWidget.cpp b/Engine/Source/Developer/SlateReflector/Private/Widgets/SWidgetReflectorToolTipWidget.cpp
--- a/Engine/Source/Developer/SlateReflector/Private/Widgets/SWidgetReflectorToolTipWidget.cpp
+++ b/Engine/Source/Developer/SlateReflector/Private/Widgets/SWidgetReflectorToolTipWidget.cpp
@@ -81,6 +81,8 @@ void SReflectorToolTipWidget::Construct( const FArguments& InArgs )
 		{
 			BuildLabelAndValue(LOCTEXT("IsVisible", "Is Visible"), { this, &SReflectorToolTipWidget::GetIsVisible });
 			BuildLabelAndValue(LOCTEXT("IsVisibleInherited", "Is Visible Inherited"), { this, &SReflectorToolTipWidget::GetIsVisibleInherited });
+			BuildLabelAndValue(LOCTEXT("SupportsInvalidation", "Supports Invalidation"), { this, &SReflectorToolTipWidget::GetSupportsInvalidation });
+			BuildLabelAndValue(LOCTEXT("SupportsInvalidationRecursive", "Recursively Supports Invalidation"), { this, &SReflectorToolTipWidget::GetSupportsInvalidationRecursive });
 		}
 		BuildLabelAndValue(LOCTEXT("HasActiveTimers", "Has Active Timer"), { this, &SReflectorToolTipWidget::GetHasActiveTimers });
 		if (WidgetInfo->GetNodeType() == EWidgetReflectorNodeType::Live)
diff --git a/Engine/Source/Developer/SlateReflector/Private/Widgets/SWidgetReflectorToolTipWidget.h b/Engine/Source/Developer/SlateReflector/Private/Widgets/SWidgetReflectorToolTipWidget.h
--- a/Engine/Source/Developer/SlateReflector/Private/Widgets/SWidgetReflectorToolTipWidget.h
+++ b/Engine/Source/Developer/SlateReflector/Private/Widgets/SWidgetReflectorToolTipWidget.h
@@ -96,6 +96,16 @@ private:
 		return (WidgetInfo->GetWidgetVisibilityInherited()) ? TrueText : FalseText;
 	}
 
+	FText GetSupportsInvalidation() const
+	{
+		return (WidgetInfo->GetWidgetSupportsInvalidation()) ? TrueText : FalseText;
+	}
+	
+	FText GetSupportsInvalidationRecursive() const
+	{
+		return (WidgetInfo->GetSupportsInvalidationRecursive()) ? TrueText : FalseText;
+	}
+
 private:
 
 	static const FText TrueText;
diff --git a/Engine/Source/Developer/SlateReflector/Private/Widgets/SWidgetReflectorTreeWidgetItem.cpp b/Engine/Source/Developer/SlateReflector/Private/Widgets/SWidgetReflectorTreeWidgetItem.cpp
--- a/Engine/Source/Developer/SlateReflector/Private/Widgets/SWidgetReflectorTreeWidgetItem.cpp
+++ b/Engine/Source/Developer/SlateReflector/Private/Widgets/SWidgetReflectorTreeWidgetItem.cpp
@@ -5,6 +5,7 @@
 #include "Widgets/Colors/SColorBlock.h"
 #include "Widgets/Input/SHyperlink.h"
 #include "HAL/PlatformApplicationMisc.h"
+#include "Widgets/SInvalidationPanel.h"
 #include "Widgets/Input/SCheckBox.h"
 
 #define LOCTEXT_NAMESPACE "SWidgetReflector"
@@ -66,6 +67,90 @@ TSharedRef<SWidget> SReflectorTreeWidgetItem::GenerateWidgetForColumn(const FNam
 			.ShouldDrawWires(true)
 		];
 
+#if UE_SLATE_WITH_DYNAMIC_INVALIDATION
+		// Here we show an indication for Dynamic Invalidation Panels and their children.
+		// Dynamic Invalidation Panels are Invalidation Panels with UseDynamicInvalidation set to true.
+		static FName InvalidationPanelName = TEXT("SInvalidationPanel");
+
+		const TSharedPtr<SWidget> LiveWidget = WidgetInfo->GetLiveWidget();
+		TSharedPtr<SWidget> InvalidationPanel = LiveWidget;
+		while (InvalidationPanel && InvalidationPanel->GetType() != InvalidationPanelName)
+		{
+			InvalidationPanel = InvalidationPanel->GetParentWidget();
+		}
+
+		if (InvalidationPanel && StaticCastSharedPtr<SInvalidationPanel>(InvalidationPanel)->GetUseDynamicInvalidation()) // if we found a Dynamic Invalidation Panel
+		{
+			const bool CanCache = StaticCastSharedPtr<SInvalidationPanel>(InvalidationPanel)->GetCanCache();
+			// we always show a color block for the Dynamic Invalidation Panel, but for its children we only show the color if it can cache
+			const bool bShowCachingIndicator = CanCache || LiveWidget == InvalidationPanel; 
+			if (bShowCachingIndicator)
+			{
+				TPair<FLinearColor, FText> ColorAndTooltip;
+			
+				if (LiveWidget == InvalidationPanel)
+				{
+					if (!CanCache)
+					{
+						ColorAndTooltip = TPair<FLinearColor, FText> {
+							FLinearColor(1.0f, 0.0f, 0.0f), // Red for disabled Dynamic Invalidation Panel
+							LOCTEXT("DynamicInvalidationPanel_Disabled", "Dynamic Invalidation Panel cannot cache")
+						};
+					}
+					else if (LiveWidget->SupportsInvalidationRecursive())
+					{
+						ColorAndTooltip = TPair<FLinearColor, FText> {
+							FLinearColor(0.0f, 0.5f, 1.0f), // Blue for Dynamic Invalidation Panel that recursively supports invalidation
+							LOCTEXT("DynamicInvalidationPanel_SupportingInvalidationRecursive", "Dynamic Invalidation Panel is caching all its children")
+						} ;
+					}
+					else
+					{
+						ColorAndTooltip = TPair<FLinearColor, FText> {
+							FLinearColor(1.f, 0.f, 1.f), // Magenta for Dynamic Invalidation Panel that is not able to recursively supports invalidation
+							LOCTEXT("DynamicInvalidationPanel_NotSupportingInvalidationRecursive", "Dynamic Invalidation Panel supports invalidation but not all its children do")
+						};
+					}
+				}
+				else
+				{
+					if (!LiveWidget->SupportsInvalidation())
+					{
+						ColorAndTooltip = TPair<FLinearColor, FText> {
+							FLinearColor(1.0f, 0.0f, 0.0f), // Red for widgets not supporting invalidation
+							LOCTEXT("DynamicInvalidationPanel_WidgetNotSupportingInvalidation", "Widget does not support invalidation")
+						};
+					}
+					else if (LiveWidget->SupportsInvalidationRecursive())
+					{
+						ColorAndTooltip = TPair<FLinearColor, FText> {
+							FLinearColor(0.0f, 1.0f, 0.0f), // Green for widgets recursively supporting invalidation
+							LOCTEXT("DynamicInvalidationPanel_WidgetSupportingInvalidationRecursive", "Widget and children all support invalidation")
+						};
+					}
+					else
+					{
+						ColorAndTooltip = TPair<FLinearColor, FText> {
+							FLinearColor(1.0f, .5f, 0.0f), // Orange for widgets supporting invalidation but having children that do not
+							LOCTEXT("DynamicInvalidationPanel_WidgetNotSupportingInvalidationRecursive", "Widget supports invalidation but not all its children do")
+						};
+					}
+				}
+
+				HorizontalBox->AddSlot()
+				.AutoWidth()
+				.Padding(2.0f, 0.0f)
+				.VAlign(VAlign_Center)
+				[
+					SNew(SColorBlock)
+					.Size(FVector2D{10.0f, 10.0f})
+					.Color(ColorAndTooltip.Key)
+					.ToolTipText(ColorAndTooltip.Value)
+				];
+			}
+		}
+#endif
+
 		if (WidgetInfo->GetWidgetIsInvalidationRoot())
 		{
 			HorizontalBox->AddSlot()
diff --git a/Engine/Source/Runtime/Slate/Private/Widgets/SInvalidationPanel.cpp b/Engine/Source/Runtime/Slate/Private/Widgets/SInvalidationPanel.cpp
--- a/Engine/Source/Runtime/Slate/Private/Widgets/SInvalidationPanel.cpp
+++ b/Engine/Source/Runtime/Slate/Private/Widgets/SInvalidationPanel.cpp
@@ -35,11 +35,33 @@ FAutoConsoleVariableRef CVarAlwaysInvalidate(
 
 #endif // WITH_SLATE_DEBUGGING
 
+#if UE_SLATE_WITH_DYNAMIC_INVALIDATION
+namespace DynamicInvalidationCVars
+{
+	static int32 bDynamicInvalidationOptions = 2;
+	static FAutoConsoleVariableRef CVarEnableDynamicInvalidation(
+		TEXT("Slate.DynamicInvalidationOptions"),
+		bDynamicInvalidationOptions,
+		TEXT("Toggle Dynamic Invalidation options for Invalidation Panels.\n")
+		TEXT(" A Dynamic Invalidation Panel refers to an Invalidation Panel with bUseDynamicInvalidation set to true.\n")
+		TEXT(" 0. Disable Dynamic Invalidation and disable caching for all Dynamic Invalidation Panels\n")
+		TEXT(" 1. Disable Dynamic Invalidation and makes the Dynamic Invalidation Panels act like regular Invalidation Panels\n")
+		TEXT(" 2. Enable Dynamic Invalidation for Dynamic Invalidation Panels\n")
+		TEXT(" 3. Forces Dynamic Invalidation on all Invalidation Panels")
+	);
+}
+#endif
+
 
+SLATE_IMPLEMENT_WIDGET(SInvalidationPanel)
+void SInvalidationPanel::PrivateRegisterAttributes(FSlateAttributeInitializer& AttributeInitializer) {}
 
 SInvalidationPanel::SInvalidationPanel()
 	: HittestGrid(MakeShared<FHittestGrid>())
 	, bCanCache(true)
+#if UE_SLATE_WITH_DYNAMIC_INVALIDATION
+	, bUseDynamicInvalidation(false)
+#endif
 	, bPaintedSinceLastPrepass(true)
 	, bWasCachable(false)
 {
@@ -68,6 +90,12 @@ void SInvalidationPanel::Construct( const FArguments& InArgs )
 #endif
 }
 
+FChildren* SInvalidationPanel::GetAllChildren()
+{
+	// The default implementation would call SInvalidationPanel::GetChildren would is always empty if we can cache
+	return SCompoundWidget::GetChildren();
+}
+
 SInvalidationPanel::~SInvalidationPanel()
 {
 	InvalidateRootChildOrder();
@@ -112,6 +140,18 @@ void SInvalidationPanel::EnableInvalidationPanels(bool bEnable)
 
 bool SInvalidationPanel::GetCanCache() const
 {
+#if UE_SLATE_WITH_DYNAMIC_INVALIDATION
+	const bool bForcedDynamicInvalidation = DynamicInvalidationCVars::bDynamicInvalidationOptions == 3;
+	const bool bDynamicInvalidationEnabled = bForcedDynamicInvalidation || (bUseDynamicInvalidation && DynamicInvalidationCVars::bDynamicInvalidationOptions == 2);
+	if (!bDynamicInvalidationEnabled)
+	{
+		const bool bRevertToRegularInvalidation = !bUseDynamicInvalidation || DynamicInvalidationCVars::bDynamicInvalidationOptions == 1;
+		if (!bRevertToRegularInvalidation)
+		{
+			return false;
+		}
+	}
+#endif
 	return bCanCache && !GSlateEnableGlobalInvalidation && bInvalidationPanelsEnabled;
 }
 
@@ -180,6 +220,22 @@ void SInvalidationPanel::SetCanCache(bool InCanCache)
 	}
 }
 
+#if UE_SLATE_WITH_DYNAMIC_INVALIDATION
+bool SInvalidationPanel::GetUseDynamicInvalidation() const
+{
+	return bUseDynamicInvalidation;
+}
+
+void SInvalidationPanel::SetUseDynamicInvalidation(bool InUseDynamicInvalidation)
+{
+	bUseDynamicInvalidation = InUseDynamicInvalidation;
+	
+	// Make sure we update the cached value
+	const bool bUseCachedValue = false;
+	SupportsInvalidationRecursive(bUseCachedValue);
+}
+#endif
+
 FChildren* SInvalidationPanel::GetChildren()
 {
 	if (GetCanCache())
@@ -199,6 +255,16 @@ FChildren* SInvalidationPanel::Debug_GetChildrenForReflector()
 }
 #endif
 
+#if UE_SLATE_WITH_DYNAMIC_INVALIDATION
+bool SInvalidationPanel::SupportsInvalidationRecursive(bool bUseCachedValue) const
+{
+	const bool bForcedDynamicInvalidation = DynamicInvalidationCVars::bDynamicInvalidationOptions == 3;
+	const bool bDynamicInvalidationEnabled = bForcedDynamicInvalidation || (bUseDynamicInvalidation && DynamicInvalidationCVars::bDynamicInvalidationOptions == 2);
+	// For an invalidation panel without bUseDynamicInvalidation, we force the caching of all children, so we always support invalidation
+	return !bDynamicInvalidationEnabled || SWidget::SupportsInvalidationRecursive(bUseCachedValue);
+}
+#endif
+
 int32 SInvalidationPanel::OnPaint( const FPaintArgs& Args, const FGeometry& AllottedGeometry, const FSlateRect& MyCullingRect, FSlateWindowElementList& OutDrawElements, int32 LayerId, const FWidgetStyle& InWidgetStyle, bool bParentEnabled ) const
 {
 #if SLATE_VERBOSE_NAMED_EVENTS
@@ -209,7 +275,11 @@ int32 SInvalidationPanel::OnPaint( const FPaintArgs& Args, const FGeometry& Allo
 	bPaintedSinceLastPrepass = true;
 	SInvalidationPanel* MutableThis = const_cast<SInvalidationPanel*>(this);
 
+#if UE_SLATE_WITH_DYNAMIC_INVALIDATION
+	const bool bCanCacheThisFrame = GetCanCache() && SupportsInvalidationRecursive();
+#else
 	const bool bCanCacheThisFrame = GetCanCache();
+#endif
 	if (bCanCacheThisFrame != bWasCachable)
 	{
 		MutableThis->InvalidateRootChildOrder();
diff --git a/Engine/Source/Runtime/Slate/Public/Widgets/SInvalidationPanel.h b/Engine/Source/Runtime/Slate/Public/Widgets/SInvalidationPanel.h
--- a/Engine/Source/Runtime/Slate/Public/Widgets/SInvalidationPanel.h
+++ b/Engine/Source/Runtime/Slate/Public/Widgets/SInvalidationPanel.h
@@ -20,6 +20,7 @@ class SWindow;
 
 class SInvalidationPanel : public SCompoundWidget, public FSlateInvalidationRoot
 {
+	SLATE_DECLARE_WIDGET_SUPPORTING_INVALIDATION_API(SInvalidationPanel, SCompoundWidget, SLATE_API)
 public:
 	SLATE_BEGIN_ARGS( SInvalidationPanel )
 	{
@@ -40,10 +41,25 @@ public:
 #endif
 	SLATE_API void Construct( const FArguments& InArgs );
 
-	SLATE_API bool GetCanCache() const;
+	SLATE_API virtual FChildren* GetAllChildren() override;
+
+	/** Returns true when the invalidation panel is able to cache the widgets. **/
+	SLATE_API virtual bool GetCanCache() const;
 
 	SLATE_API void SetCanCache(bool InCanCache);
 
+#if UE_SLATE_WITH_DYNAMIC_INVALIDATION
+	/**
+	 * Returns true when the invalidation panel uses Dynamic Invalidation.
+	 * Dynamic Invalidation means that the widget's content will only be cached if all descendant widgets support invalidation.
+	**/
+	SLATE_API virtual bool GetUseDynamicInvalidation() const;
+
+	SLATE_API void SetUseDynamicInvalidation(bool InUseDynamicInvalidation);
+
+	SLATE_API virtual bool SupportsInvalidationRecursive(bool bUseCachedValue = true) const override;
+#endif
+
 	//~ SWidget overrides
 	SLATE_API virtual int32 OnPaint( const FPaintArgs& Args, const FGeometry& AllottedGeometry, const FSlateRect& MyCullingRect, FSlateWindowElementList& OutDrawElements, int32 LayerId, const FWidgetStyle& InWidgetStyle, bool bParentEnabled ) const override;
 	SLATE_API virtual FChildren* GetChildren() override;
@@ -75,6 +91,9 @@ private:
 	mutable FLinearColor LastIncomingColorAndOpacity;
 
 	bool bCanCache;
+#if UE_SLATE_WITH_DYNAMIC_INVALIDATION
+	bool bUseDynamicInvalidation;
+#endif
 
 	mutable bool bPaintedSinceLastPrepass;
 #if SLATE_VERBOSE_NAMED_EVENTS
diff --git a/Engine/Source/Runtime/SlateCore/Private/FastUpdate/SlateInvalidationWidgetList.cpp b/Engine/Source/Runtime/SlateCore/Private/FastUpdate/SlateInvalidationWidgetList.cpp
--- a/Engine/Source/Runtime/SlateCore/Private/FastUpdate/SlateInvalidationWidgetList.cpp
+++ b/Engine/Source/Runtime/SlateCore/Private/FastUpdate/SlateInvalidationWidgetList.cpp
@@ -717,6 +717,13 @@ bool FSlateInvalidationWidgetList::ProcessChildOrderInvalidation(FSlateInvalidat
 		}
 	}
 
+#if UE_SLATE_WITH_DYNAMIC_INVALIDATION
+	// If the children have been modified in any way, the cached value of SupportsInvalidationRecursive might have changed.
+	// We call SupportsInvalidationRecursive to recompute the value and invalidate the parent's value if needed
+	const bool bUseCachedValue = false;
+	bool bNowSupportsInvalidation = WidgetPtr->SupportsInvalidationRecursive(bUseCachedValue);
+#endif
+
 	return bIsInvalidationWidgetStillValid;
 }
 
diff --git a/Engine/Source/Runtime/SlateCore/Private/Widgets/SLeafWidget.cpp b/Engine/Source/Runtime/SlateCore/Private/Widgets/SLeafWidget.cpp
--- a/Engine/Source/Runtime/SlateCore/Private/Widgets/SLeafWidget.cpp
+++ b/Engine/Source/Runtime/SlateCore/Private/Widgets/SLeafWidget.cpp
@@ -5,6 +5,8 @@
 
 /* SLeafWidget interface
  *****************************************************************************/
+SLATE_IMPLEMENT_WIDGET(SLeafWidget)
+void SLeafWidget::PrivateRegisterAttributes(FSlateAttributeInitializer& AttributeInitializer) {}
 
 SLeafWidget::SLeafWidget()
 {
diff --git a/Engine/Source/Runtime/SlateCore/Private/Widgets/SNullWidget.cpp b/Engine/Source/Runtime/SlateCore/Private/Widgets/SNullWidget.cpp
--- a/Engine/Source/Runtime/SlateCore/Private/Widgets/SNullWidget.cpp
+++ b/Engine/Source/Runtime/SlateCore/Private/Widgets/SNullWidget.cpp
@@ -11,6 +11,7 @@
 class SLATECORE_API SNullWidgetContent
 	: public SWidget
 {
+	SLATE_DECLARE_WIDGET_SUPPORTING_INVALIDATION(SNullWidgetContent, SWidget)
 public:
 
 	SLATE_BEGIN_ARGS(SNullWidgetContent)
@@ -68,6 +69,9 @@ protected:
 	// End of SWidget interface
 };
 
+SLATE_IMPLEMENT_WIDGET(SNullWidgetContent)
+void SNullWidgetContent::PrivateRegisterAttributes(FSlateAttributeInitializer& AttributeInitializer) {}
+
 namespace NullWidgetPrivate
 {
 	TSharedRef<SWidget> Construct()
diff --git a/Engine/Source/Runtime/SlateCore/Private/Widgets/SPanel.cpp b/Engine/Source/Runtime/SlateCore/Private/Widgets/SPanel.cpp
--- a/Engine/Source/Runtime/SlateCore/Private/Widgets/SPanel.cpp
+++ b/Engine/Source/Runtime/SlateCore/Private/Widgets/SPanel.cpp
@@ -4,6 +4,9 @@
 #include "Types/PaintArgs.h"
 #include "Layout/ArrangedChildren.h"
 
+SLATE_IMPLEMENT_WIDGET(SPanel);
+void SPanel::PrivateRegisterAttributes(FSlateAttributeInitializer& AttributeInitializer) {}
+
 SPanel::SPanel() = default;
 SPanel::~SPanel() = default;
 
diff --git a/Engine/Source/Runtime/SlateCore/Private/Widgets/SWidget.cpp b/Engine/Source/Runtime/SlateCore/Private/Widgets/SWidget.cpp
--- a/Engine/Source/Runtime/SlateCore/Private/Widgets/SWidget.cpp
+++ b/Engine/Source/Runtime/SlateCore/Private/Widgets/SWidget.cpp
@@ -1962,6 +1962,97 @@ void SWidget::AddMetadataInternal(const TSharedRef<ISlateMetaData>& AddMe)
 #endif
 }
 
+#if UE_SLATE_WITH_DYNAMIC_INVALIDATION
+
+bool SWidget::SupportsInvalidation() const
+{
+	// We make sure that this exact widget's class is defined and supports invalidation,
+	// as we don't want to assume descendant widgets support invalidation
+	return GetType() == GetWidgetClass().GetWidgetType() ? GetWidgetClass().SupportsInvalidation() : false;
+}
+
+UE_TRACE_EVENT_BEGIN(Cpu, SWidget_SupportsInvalidationRecursive, NoSync)
+	UE_TRACE_EVENT_FIELD(UE::Trace::WideString, WidgetClassName)
+	UE_TRACE_EVENT_FIELD(UE::Trace::WideString, CreatedInLocation)
+	UE_TRACE_EVENT_FIELD(bool, UseCachedValue)
+	UE_TRACE_EVENT_FIELD(bool, OldCanBeCached)
+UE_TRACE_EVENT_END()
+
+bool SWidget::SupportsInvalidationRecursive(bool bUseCachedValue) const
+{
+	if (!SupportsInvalidation())
+	{
+		CachedSupportsInvalidationRecursive = false;
+		return false;
+	}
+
+	// If the cached value was true, then all children supported invalidation and the widget would have been invalidated if needed.
+	// If the cached value was false, at least one of the descendant widget did not support invalidation and we therefore need to recheck everytime
+	if (CachedSupportsInvalidationRecursive && bUseCachedValue)
+	{
+		return true;
+	}
+
+#if UE_SLATE_TRACE_ENABLED
+	UE_TRACE_LOG_SCOPED_T(Cpu, SWidget_SupportsInvalidationRecursive, SlateChannel)
+		<< SWidget_SupportsInvalidationRecursive.WidgetClassName(*TypeOfWidget.ToString())
+		<< SWidget_SupportsInvalidationRecursive.CreatedInLocation(*CreatedInLocation.ToString())
+		<< SWidget_SupportsInvalidationRecursive.UseCachedValue(bUseCachedValue)
+		<< SWidget_SupportsInvalidationRecursive.OldCanBeCached(CachedSupportsInvalidationRecursive);
+#endif // UE_SLATE_TRACE_ENABLED
+
+	const bool OldCachedValue = CachedSupportsInvalidationRecursive;
+	CachedSupportsInvalidationRecursive = true;
+
+	const FChildren* Children = const_cast<SWidget*>(this)->GetAllChildren();
+	for (int32 ChildIndex = 0; ChildIndex < Children->Num(); ++ChildIndex)
+	{
+		const TSharedRef<const SWidget> ChildWidget = Children->GetChildAt(ChildIndex);
+		if (ChildWidget->ParentWidgetPtr != nullptr && !ChildWidget->SupportsInvalidationRecursive())
+		{
+			CachedSupportsInvalidationRecursive = false;
+			break;
+		}
+	}
+	
+	// If the value has changed, we invalidate the parent
+	if (CachedSupportsInvalidationRecursive != OldCachedValue)
+	{
+		if (const TSharedPtr<SWidget> ParentWidget = GetParentWidget())
+		{
+			if (!CachedSupportsInvalidationRecursive)
+			{
+				// If we supported invalidation but we do not anymore, tell the parent they do not support invalidation anymore, no need for them to check
+				ParentWidget->SetSelfAndParentAsNotSupportingInvalidationRecursive();
+			}
+			else
+			{
+				// If we now support invalidation, ask the parent to recompute.
+				const bool bParentUseCachedValue = false;
+				bool bNowSupportsInvalidation = ParentWidget->SupportsInvalidationRecursive(bParentUseCachedValue);
+			}
+		}
+	}
+
+	return CachedSupportsInvalidationRecursive;
+}
+
+void SWidget::SetSelfAndParentAsNotSupportingInvalidationRecursive()
+{
+	// if this widget was supporting invalidation and it does not anymore, it needs to tell its parent.
+	// Otherwise, its parent already knew so no need to do anything.
+	if (CachedSupportsInvalidationRecursive)
+	{
+		CachedSupportsInvalidationRecursive = false;
+		if (const TSharedPtr<SWidget> ParentWidget = GetParentWidget())
+		{
+			ParentWidget->SetSelfAndParentAsNotSupportingInvalidationRecursive();
+		}
+	}
+}
+
+#endif
+
 #if WITH_ACCESSIBILITY
 TSharedRef<FSlateAccessibleWidget> SWidget::CreateAccessibleWidget()
 {
diff --git a/Engine/Source/Runtime/SlateCore/Public/SlateGlobals.h b/Engine/Source/Runtime/SlateCore/Public/SlateGlobals.h
--- a/Engine/Source/Runtime/SlateCore/Public/SlateGlobals.h
+++ b/Engine/Source/Runtime/SlateCore/Public/SlateGlobals.h
@@ -38,6 +38,11 @@
 	#define UE_SLATE_WITH_WIDGET_UNIQUE_IDENTIFIER UE_SLATE_TRACE_ENABLED || WITH_SLATE_DEBUGGING
 #endif
 
+/** Enables Dynamic Invalidation for SInvalidationPanels. Currently Editor Only  */
+#ifndef UE_SLATE_WITH_DYNAMIC_INVALIDATION
+	#define UE_SLATE_WITH_DYNAMIC_INVALIDATION WITH_EDITOR
+#endif
+
 // HOW TO GET AN IN-DEPTH PERFORMANCE ANALYSIS OF SLATE
 //
 // Step 1)
diff --git a/Engine/Source/Runtime/SlateCore/Public/Widgets/Images/SImage.h b/Engine/Source/Runtime/SlateCore/Public/Widgets/Images/SImage.h
--- a/Engine/Source/Runtime/SlateCore/Public/Widgets/Images/SImage.h
+++ b/Engine/Source/Runtime/SlateCore/Public/Widgets/Images/SImage.h
@@ -20,7 +20,7 @@ class FSlateWindowElementList;
 class SImage
 	: public SLeafWidget
 {
-	SLATE_DECLARE_WIDGET_API(SImage, SLeafWidget, SLATECORE_API)
+	SLATE_DECLARE_WIDGET_SUPPORTING_INVALIDATION_API(SImage, SLeafWidget, SLATECORE_API)
 
 public:
 	SLATE_BEGIN_ARGS( SImage )
diff --git a/Engine/Source/Runtime/SlateCore/Public/Widgets/SBoxPanel.h b/Engine/Source/Runtime/SlateCore/Public/Widgets/SBoxPanel.h
--- a/Engine/Source/Runtime/SlateCore/Public/Widgets/SBoxPanel.h
+++ b/Engine/Source/Runtime/SlateCore/Public/Widgets/SBoxPanel.h
@@ -25,7 +25,7 @@ class FArrangedChildren;
  */
 class SBoxPanel : public SPanel
 {
-	SLATE_DECLARE_WIDGET_API(SBoxPanel, SPanel, SLATECORE_API)
+	SLATE_DECLARE_WIDGET_SUPPORTING_INVALIDATION_API(SBoxPanel, SPanel, SLATECORE_API)
 
 protected:
 	/**
@@ -156,7 +156,7 @@ protected:
 /** A Horizontal Box Panel. See SBoxPanel for more info. */
 class SHorizontalBox : public SBoxPanel
 {
-	SLATE_DECLARE_WIDGET_API(SHorizontalBox, SBoxPanel, SLATECORE_API)
+	SLATE_DECLARE_WIDGET_SUPPORTING_INVALIDATION_API(SHorizontalBox, SBoxPanel, SLATECORE_API)
 public:
 	class FSlot : public SBoxPanel::TSlot<FSlot>
 	{
@@ -300,7 +300,7 @@ public:
 /** A Vertical Box Panel. See SBoxPanel for more info. */
 class SVerticalBox : public SBoxPanel
 {
-	SLATE_DECLARE_WIDGET_API(SVerticalBox, SBoxPanel, SLATECORE_API)
+	SLATE_DECLARE_WIDGET_SUPPORTING_INVALIDATION_API(SVerticalBox, SBoxPanel, SLATECORE_API)
 public:
 	class FSlot : public SBoxPanel::TSlot<FSlot>
 	{
@@ -446,7 +446,7 @@ public:
 /** A Stack Box Panel that stack vertically or horizontally. See SBoxPanel for more info. */
 class SStackBox : public SBoxPanel
 {
-	SLATE_DECLARE_WIDGET_API(SStackBox, SBoxPanel, SLATECORE_API)
+	SLATE_DECLARE_WIDGET_SUPPORTING_INVALIDATION_API(SStackBox, SBoxPanel, SLATECORE_API)
 public:
 	class FSlot : public SBoxPanel::TSlot<FSlot>
 	{
diff --git a/Engine/Source/Runtime/SlateCore/Public/Widgets/SLeafWidget.h b/Engine/Source/Runtime/SlateCore/Public/Widgets/SLeafWidget.h
--- a/Engine/Source/Runtime/SlateCore/Public/Widgets/SLeafWidget.h
+++ b/Engine/Source/Runtime/SlateCore/Public/Widgets/SLeafWidget.h
@@ -21,6 +21,7 @@ class FSlateWindowElementList;
 class SLeafWidget
 	: public SWidget
 {
+	SLATE_DECLARE_WIDGET_SUPPORTING_INVALIDATION_API(SLeafWidget, SWidget, SLATECORE_API)
 public:
 	SLATECORE_API SLeafWidget();
 
diff --git a/Engine/Source/Runtime/SlateCore/Public/Widgets/SOverlay.h b/Engine/Source/Runtime/SlateCore/Public/Widgets/SOverlay.h
--- a/Engine/Source/Runtime/SlateCore/Public/Widgets/SOverlay.h
+++ b/Engine/Source/Runtime/SlateCore/Public/Widgets/SOverlay.h
@@ -35,7 +35,7 @@ class FSlateWindowElementList;
  */
 class SOverlay : public SPanel
 {
-	SLATE_DECLARE_WIDGET_API(SOverlay, SPanel, SLATECORE_API)
+	SLATE_DECLARE_WIDGET_SUPPORTING_INVALIDATION_API(SOverlay, SPanel, SLATECORE_API)
 public:	
 
 	/** A slot that support alignment of content and padding and z-order */
diff --git a/Engine/Source/Runtime/SlateCore/Public/Widgets/SPanel.h b/Engine/Source/Runtime/SlateCore/Public/Widgets/SPanel.h
--- a/Engine/Source/Runtime/SlateCore/Public/Widgets/SPanel.h
+++ b/Engine/Source/Runtime/SlateCore/Public/Widgets/SPanel.h
@@ -22,6 +22,7 @@ class FSlateWindowElementList;
 class SPanel
 	: public SWidget
 {
+	SLATE_DECLARE_WIDGET_SUPPORTING_INVALIDATION_API(SPanel, SWidget, SLATECORE_API)
 public:	
 
 	/**
diff --git a/Engine/Source/Runtime/SlateCore/Public/Widgets/SWidget.h b/Engine/Source/Runtime/SlateCore/Public/Widgets/SWidget.h
--- a/Engine/Source/Runtime/SlateCore/Public/Widgets/SWidget.h
+++ b/Engine/Source/Runtime/SlateCore/Public/Widgets/SWidget.h
@@ -1515,8 +1515,26 @@ public:
 	virtual bool Advanced_IsInvalidationRoot() const { return false; }
 	virtual const FSlateInvalidationRoot* Advanced_AsInvalidationRoot() const { return nullptr; }
 
-protected:
 
+#if UE_SLATE_WITH_DYNAMIC_INVALIDATION
+
+	/** Returns true if this widget's class supports invalidation. */
+	SLATECORE_API bool SupportsInvalidation() const;
+
+	/**
+	 * Returns true if this widget and all of its children support invalidation.
+	 * This value is cached and reused unless bUseCachedValue is set to true.
+	 */
+	SLATECORE_API virtual bool SupportsInvalidationRecursive(bool bUseCachedValue = true) const;
+private:
+	/** Fast path to invalidate all parents when a descendant widget is not supporting invalidation.. */
+	SLATECORE_API void SetSelfAndParentAsNotSupportingInvalidationRecursive();
+;
+	/** The cached value of the last SupportsInvalidationRecursive call. */
+	mutable bool CachedSupportsInvalidationRecursive = false;
+#endif
+
+protected:
 	/**
 	 * Hidden default constructor.
 	 *
diff --git a/Engine/Source/Runtime/SlateCore/Public/Widgets/SlateControlledConstruction.h b/Engine/Source/Runtime/SlateCore/Public/Widgets/SlateControlledConstruction.h
--- a/Engine/Source/Runtime/SlateCore/Public/Widgets/SlateControlledConstruction.h
+++ b/Engine/Source/Runtime/SlateCore/Public/Widgets/SlateControlledConstruction.h
@@ -17,9 +17,20 @@ namespace SharedPointerInternals
 
 /** */
 #define SLATE_DECLARE_WIDGET(WidgetType, ParentType) \
-	SLATE_DECLARE_WIDGET_API(WidgetType, ParentType, NO_API)
+	SLATE_DECLARE_WIDGET_IMPL(WidgetType, ParentType, false, NO_API)
+
+#define SLATE_DECLARE_WIDGET_SUPPORTING_INVALIDATION(WidgetType, ParentType) \
+	SLATE_DECLARE_WIDGET_IMPL(WidgetType, ParentType, true, NO_API)
+
 
 #define SLATE_DECLARE_WIDGET_API(WidgetType, ParentType, ModuleApiDefine) \
+	SLATE_DECLARE_WIDGET_IMPL(WidgetType, ParentType, false, ModuleApiDefine)
+
+#define SLATE_DECLARE_WIDGET_SUPPORTING_INVALIDATION_API(WidgetType, ParentType, ModuleApiDefine) \
+	SLATE_DECLARE_WIDGET_IMPL(WidgetType, ParentType, true, ModuleApiDefine)
+
+
+#define SLATE_DECLARE_WIDGET_IMPL(WidgetType, ParentType, bInSupportsInvalidation, ModuleApiDefine) \
 	private: \
 		using Super = ParentType; \
 		using ThisClass = WidgetType; \
@@ -31,7 +42,7 @@ namespace SharedPointerInternals
 		friend class SharedPointerInternals::TIntrusiveReferenceController; /* Shouldn't really forward-declare this to allow MakeShared to work, but is for legacy reasons */ \
 		static const FSlateWidgetClassData& GetPrivateWidgetClass() \
 		{ \
-			static FSlateWidgetClassData WidgetClassDataInstance = FSlateWidgetClassData(TIdentity<ParentType>(), #WidgetType, &WidgetType::PrivateRegisterAttributes); \
+			static FSlateWidgetClassData WidgetClassDataInstance = FSlateWidgetClassData(TIdentity<ParentType>(), #WidgetType, bInSupportsInvalidation, &WidgetType::PrivateRegisterAttributes); \
 			return WidgetClassDataInstance; \
 		} \
 		static ModuleApiDefine void PrivateRegisterAttributes(FSlateAttributeInitializer&); \
@@ -52,12 +63,18 @@ private:
 
 	FSlateWidgetClassData(FName InWidgetTypeName)
 		: WidgetType(InWidgetTypeName)
+#if UE_SLATE_WITH_DYNAMIC_INVALIDATION
+		, bSupportsInvalidation(false)
+#endif
 	{}
 
 public:
 	template<typename InWidgetParentType>
-	FSlateWidgetClassData(TIdentity<InWidgetParentType>, FName InWidgetTypeName, void(*AttributeInitializer)(FSlateAttributeInitializer&))
+	FSlateWidgetClassData(TIdentity<InWidgetParentType>, FName InWidgetTypeName, bool bInSupportsInvalidation, void(*AttributeInitializer)(FSlateAttributeInitializer&))
 		: WidgetType(InWidgetTypeName)
+#if UE_SLATE_WITH_DYNAMIC_INVALIDATION
+		, bSupportsInvalidation(bInSupportsInvalidation)
+#endif
 	{
 		// Initialize the parent class if it's not already done
 		const FSlateWidgetClassData& ParentWidgetClassData = InWidgetParentType::StaticWidgetClass();
@@ -68,10 +85,16 @@ public:
 
 	const FSlateAttributeDescriptor& GetAttributeDescriptor() const { return Descriptor; };
 	FName GetWidgetType() const { return WidgetType; }
+#if UE_SLATE_WITH_DYNAMIC_INVALIDATION
+	bool SupportsInvalidation() const { return bSupportsInvalidation; }
+#endif
 
 private:
 	FSlateAttributeDescriptor Descriptor;
 	FName WidgetType;
+#if UE_SLATE_WITH_DYNAMIC_INVALIDATION
+	bool bSupportsInvalidation;
+#endif
 };
 
 /** */
diff --git a/Engine/Source/Runtime/UMG/Private/Components/InvalidationBox.cpp b/Engine/Source/Runtime/UMG/Private/Components/InvalidationBox.cpp
--- a/Engine/Source/Runtime/UMG/Private/Components/InvalidationBox.cpp
+++ b/Engine/Source/Runtime/UMG/Private/Components/InvalidationBox.cpp
@@ -44,6 +44,10 @@ PRAGMA_ENABLE_DEPRECATION_WARNINGS
 		MyInvalidationPanel->SetContent(GetContentSlot()->Content ? GetContentSlot()->Content->TakeWidget() : SNullWidget::NullWidget);
 	}
 	
+#if UE_SLATE_WITH_DYNAMIC_INVALIDATION
+	MyInvalidationPanel->SetUseDynamicInvalidation(IsDesignTime() ? false : bUseDynamicInvalidation);
+#endif
+
 	return MyInvalidationPanel.ToSharedRef();
 }
 
@@ -90,6 +94,31 @@ void UInvalidationBox::SetCanCache(bool CanCache)
 }
 PRAGMA_ENABLE_DEPRECATION_WARNINGS
 
+bool UInvalidationBox::GetUseDynamicInvalidation() const
+{
+#if UE_SLATE_WITH_DYNAMIC_INVALIDATION
+	if ( MyInvalidationPanel.IsValid() )
+	{
+		return MyInvalidationPanel->GetUseDynamicInvalidation();
+	}
+	return bUseDynamicInvalidation;
+#else
+	return false;
+#endif
+}
+
+void UInvalidationBox::SetUseDynamicInvalidation(bool UseDynamicInvalidation)
+{
+#if UE_SLATE_WITH_DYNAMIC_INVALIDATION
+	bUseDynamicInvalidation = UseDynamicInvalidation;
+	if ( MyInvalidationPanel.IsValid() )
+	{
+		MyInvalidationPanel->SetUseDynamicInvalidation(UseDynamicInvalidation);
+	}
+#endif
+}
+
+
 #if WITH_EDITOR
 
 const FText UInvalidationBox::GetPaletteCategory()
diff --git a/Engine/Source/Runtime/UMG/Public/Components/InvalidationBox.h b/Engine/Source/Runtime/UMG/Public/Components/InvalidationBox.h
--- a/Engine/Source/Runtime/UMG/Public/Components/InvalidationBox.h
+++ b/Engine/Source/Runtime/UMG/Public/Components/InvalidationBox.h
@@ -43,6 +43,19 @@ public:
 	UFUNCTION(BlueprintCallable, Category="Invalidation Box")
 	UMG_API void SetCanCache(bool CanCache);
 
+	/**
+	 * @returns true when the invalidation box uses Dynamic Invalidation.
+	 */
+	UFUNCTION(BlueprintCallable, Category="Invalidation Box")
+	UMG_API bool GetUseDynamicInvalidation() const;
+
+	/**
+	 * Tell the InvalidationBox to use Dynamic Invalidation. This option is Editor only.
+	 * Dynamic Invalidation means that the widget's content will only be cached if all descendant widgets support invalidation.
+	 */
+	UFUNCTION(BlueprintCallable, Category="Invalidation Box")
+	UMG_API void SetUseDynamicInvalidation(bool UseDynamicInvalidation);
+
 	UMG_API virtual void ReleaseSlateResources(bool bReleaseChildren) override;
 
 #if WITH_EDITOR
@@ -69,6 +82,13 @@ protected:
 	UPROPERTY(EditAnywhere, BlueprintReadWrite, Getter = "GetCanCache", Setter = "SetCanCache", BlueprintGetter = "GetCanCache", BlueprintSetter = "SetCanCache", Category = "Caching")
 	bool bCanCache;
 
+	/**
+	 * Should Dynamic Invalidation be enabled for this Invalidation Panel?
+	 * Dynamic Invalidation means that the widget's content will only be cached if all descendant widgets support invalidation.
+	 */
+	UPROPERTY(EditAnywhere, BlueprintReadWrite, Getter = "GetUseDynamicInvalidation", Setter = "SetUseDynamicInvalidation", BlueprintGetter = "GetUseDynamicInvalidation", BlueprintSetter = "SetUseDynamicInvalidation", Category = "Caching")
+	bool bUseDynamicInvalidation;
+
 protected:
 	TSharedPtr<class SInvalidationPanel> MyInvalidationPanel;
 };
