SDockTab::ActivateInParent now sets keyboard focus on the tab's content if the content supports focus and the activation cause was a click.

[RN]

diff --git a/Engine/Source/Runtime/Slate/Private/Widgets/Docking/SDockTab.cpp b/Engine/Source/Runtime/Slate/Private/Widgets/Docking/SDockTab.cpp
--- a/Engine/Source/Runtime/Slate/Private/Widgets/Docking/SDockTab.cpp
+++ b/Engine/Source/Runtime/Slate/Private/Widgets/Docking/SDockTab.cpp
@@ -405,6 +405,13 @@ void SDockTab::ActivateInParent(ETabActivationCause InActivationCause)
 	}
 
 	OnTabActivated.ExecuteIfBound(SharedThis(this), InActivationCause);
+
+	if (InActivationCause != UserClickedOnTab || !GetContent()->SupportsKeyboardFocus())
+	{
+		return;
+	}
+
+	FSlateApplication::Get().SetKeyboardFocus(GetContent(), EFocusCause::SetDirectly);
 }
 
 void SDockTab::SetTabManager( const TSharedPtr<FTabManager>& InTabManager)
