Dynamic Invalidation
Expose UseDynamicInvalidation to Slate Argument

diff --git a/Engine/Source/Runtime/Slate/Private/Widgets/SInvalidationPanel.cpp b/Engine/Source/Runtime/Slate/Private/Widgets/SInvalidationPanel.cpp
--- a/Engine/Source/Runtime/Slate/Private/Widgets/SInvalidationPanel.cpp
+++ b/Engine/Source/Runtime/Slate/Private/Widgets/SInvalidationPanel.cpp
@@ -83,6 +83,10 @@ void SInvalidationPanel::Construct( const FArguments& InArgs )
 		InArgs._Content.Widget
 	];
 
+#if UE_SLATE_WITH_DYNAMIC_INVALIDATION
+	bUseDynamicInvalidation = InArgs._UseDynamicInvalidation;
+#endif
+
 #if SLATE_VERBOSE_NAMED_EVENTS
 	DebugName = InArgs._DebugName;
 	DebugTickName = InArgs._DebugName + TEXT("_Tick");
diff --git a/Engine/Source/Runtime/Slate/Public/Widgets/SInvalidationPanel.h b/Engine/Source/Runtime/Slate/Public/Widgets/SInvalidationPanel.h
--- a/Engine/Source/Runtime/Slate/Public/Widgets/SInvalidationPanel.h
+++ b/Engine/Source/Runtime/Slate/Public/Widgets/SInvalidationPanel.h
@@ -23,10 +23,16 @@ class SInvalidationPanel : public SCompoundWidget, public FSlateInvalidationRoot
 	SLATE_DECLARE_WIDGET_SUPPORTING_INVALIDATION_API(SInvalidationPanel, SCompoundWidget, SLATE_API)
 public:
 	SLATE_BEGIN_ARGS( SInvalidationPanel )
+#if UE_SLATE_WITH_DYNAMIC_INVALIDATION
+		: _UseDynamicInvalidation(false)
+#endif
 	{
 		_Visibility = EVisibility::SelfHitTestInvisible;
 	}
 		SLATE_DEFAULT_SLOT(FArguments, Content)
+#if UE_SLATE_WITH_DYNAMIC_INVALIDATION
+		SLATE_ARGUMENT(bool, UseDynamicInvalidation)
+#endif
 #if !UE_BUILD_SHIPPING
 		SLATE_ARGUMENT(FString, DebugName)
 #endif
