Removed NavigateOnScrollIntoView check to allow tile to be scrolled into view when LineOfItem is equal to either min or max displayed line.
[RN] Fixed STileView::ScrollIntoView to correctly scroll an element into view when it's row match the min or max displayed row

diff --git a/Engine/Source/Runtime/Slate/Public/Widgets/Views/STileView.h b/Engine/Source/Runtime/Slate/Public/Widgets/Views/STileView.h
--- a/Engine/Source/Runtime/Slate/Public/Widgets/Views/STileView.h
+++ b/Engine/Source/Runtime/Slate/Public/Widgets/Views/STileView.h
@@ -517,7 +517,7 @@ protected:
 
 					this->SetScrollOffset((float)NewScrollOffset);
 				}
-				else if (this->bNavigateOnScrollIntoView)
+				else
 				{
 					// Make sure the line containing the existing entry for this item is fully in view
 					if (LineOfItem == MinDisplayedLine)
