Allow for rows to skip generation of empty table column entries. This makes it so that users can create an table / outliner with long top level header rows without clipping.

#tests Editor

diff --git a/Engine/Plugins/Experimental/UAF/UAF/Source/UAFEditor/Private/Variables/Outliner/VariablesOutlinerColumns.cpp b/Engine/Plugins/Experimental/UAF/UAF/Source/UAFEditor/Private/Variables/Outliner/VariablesOutlinerColumns.cpp
--- a/Engine/Plugins/Experimental/UAF/UAF/Source/UAFEditor/Private/Variables/Outliner/VariablesOutlinerColumns.cpp
+++ b/Engine/Plugins/Experimental/UAF/UAF/Source/UAFEditor/Private/Variables/Outliner/VariablesOutlinerColumns.cpp
@@ -50,6 +50,7 @@ SHeaderRow::FColumn::FArguments FVariablesOutlinerTypeColumn::ConstructHeaderRow
 		.HAlignCell(HAlign_Left)
 		.VAlignCell(VAlign_Center)
 		.SortMode(EColumnSortMode::None)
+		.ShouldGenerateEmptyWidgetForSpacing(false)
 		.HeaderComboVisibility(EHeaderComboVisibility::Ghosted)
 		[
 			SNew(SBox)
@@ -211,6 +212,7 @@ SHeaderRow::FColumn::FArguments FVariablesOutlinerValueColumn::ConstructHeaderRo
 		.VAlignHeader(VAlign_Center)
 		.HAlignCell(HAlign_Fill)
 		.VAlignCell(VAlign_Center)
+		.ShouldGenerateEmptyWidgetForSpacing(false)
 		[
 			SNew(SBox) 
 			.VAlign(VAlign_Center)
diff --git a/Engine/Source/Runtime/Slate/Private/Widgets/Views/SHeaderRow.cpp b/Engine/Source/Runtime/Slate/Private/Widgets/Views/SHeaderRow.cpp
--- a/Engine/Source/Runtime/Slate/Private/Widgets/Views/SHeaderRow.cpp
+++ b/Engine/Source/Runtime/Slate/Private/Widgets/Views/SHeaderRow.cpp
@@ -732,6 +732,18 @@ bool SHeaderRow::ShouldGeneratedColumn(const FName& InColumnId) const
 	return false;
 }
 
+bool SHeaderRow::ShouldGenerateEmptyColumnForSpacing(const FName& InColumnId) const
+{
+	for (const FColumn& SomeColumn : Columns)
+	{
+		if (SomeColumn.ColumnId == InColumnId)
+		{
+			return SomeColumn.ShouldGenerateWidget.Get(true) && SomeColumn.ShouldGenerateEmptyWidgetForSpacing.Get(true) && SomeColumn.bIsVisible;
+		}
+	}
+	return false;
+}
+
 bool SHeaderRow::IsColumnGenerated(const FName& InColumnId) const
 {
 	for (const TSharedPtr<STableColumnHeader>& ColumnHeader : HeaderWidgets)
diff --git a/Engine/Source/Runtime/Slate/Public/Widgets/Views/SHeaderRow.h b/Engine/Source/Runtime/Slate/Public/Widgets/Views/SHeaderRow.h
--- a/Engine/Source/Runtime/Slate/Public/Widgets/Views/SHeaderRow.h
+++ b/Engine/Source/Runtime/Slate/Public/Widgets/Views/SHeaderRow.h
@@ -167,6 +167,7 @@ public:
 			SLATE_EVENT( FOnSortModeChanged, OnSort )
 
 			SLATE_ATTRIBUTE(bool, ShouldGenerateWidget)
+			SLATE_ATTRIBUTE(bool, ShouldGenerateEmptyWidgetForSpacing)
 			SLATE_ARGUMENT(ETextOverflowPolicy, OverflowPolicy)
 		SLATE_END_ARGS()
 
@@ -198,6 +199,7 @@ public:
 			, CanManuallySort(InArgs._CanManuallySort)
 			, OnSortModeChanged( InArgs._OnSort )
 			, ShouldGenerateWidget(InArgs._ShouldGenerateWidget)
+			, ShouldGenerateEmptyWidgetForSpacing(InArgs._ShouldGenerateEmptyWidgetForSpacing)
 			, bIsVisible(true)
 			{
 			if ( InArgs._FixedWidth.IsSet() )
@@ -292,6 +294,7 @@ public:
 		FOnSortModeChanged OnSortModeChanged;
 
 		TAttribute<bool> ShouldGenerateWidget;
+		TAttribute<bool> ShouldGenerateEmptyWidgetForSpacing;
 		bool bIsVisible;
 	};
 
@@ -372,6 +375,9 @@ public:
 	/** @return is the header widget should be generated */
 	SLATE_API bool ShouldGeneratedColumn( const FName& InColumnId ) const;
 
+	/** @return if we should generate this column even if empty, some users want to skip this to allow for large headers */
+	SLATE_API bool ShouldGenerateEmptyColumnForSpacing( const FName& InColumnId ) const;
+
 	/** @return is the header widget with the column id is generated */
 	SLATE_API bool IsColumnGenerated( const FName& InColumnId ) const;
 
diff --git a/Engine/Source/Runtime/Slate/Public/Widgets/Views/STableRow.h b/Engine/Source/Runtime/Slate/Public/Widgets/Views/STableRow.h
--- a/Engine/Source/Runtime/Slate/Public/Widgets/Views/STableRow.h
+++ b/Engine/Source/Runtime/Slate/Public/Widgets/Views/STableRow.h
@@ -1386,6 +1386,11 @@ protected:
 				{
 					CellContents->SetClipping(EWidgetClipping::OnDemand);
 				}
+				else if ( !InColumnHeaders->ShouldGenerateEmptyColumnForSpacing(Column.ColumnId) )
+				{
+					NewColumnIdToSlotContents.Add(Column.ColumnId, CellContents);
+					continue;
+				}
 
 				switch (Column.SizeRule)
 				{
