Restructure code per review feedback to not use continue.

#tests Editor

diff --git a/Engine/Source/Runtime/Slate/Public/Widgets/Views/STableRow.h b/Engine/Source/Runtime/Slate/Public/Widgets/Views/STableRow.h
--- a/Engine/Source/Runtime/Slate/Public/Widgets/Views/STableRow.h
+++ b/Engine/Source/Runtime/Slate/Public/Widgets/Views/STableRow.h
@@ -1382,80 +1382,83 @@ protected:
 					CellContents = GenerateWidgetForColumn(Column.ColumnId);
 				}
 
+				bool bAddAlignedColumnCell = true;
 				if ( CellContents != SNullWidget::NullWidget )
 				{
 					CellContents->SetClipping(EWidgetClipping::OnDemand);
 				}
 				else if ( !InColumnHeaders->ShouldGenerateEmptyColumnForSpacing(Column.ColumnId) )
 				{
-					NewColumnIdToSlotContents.Add(Column.ColumnId, CellContents);
-					continue;
+					bAddAlignedColumnCell = false;
 				}
 
-				switch (Column.SizeRule)
+				if (bAddAlignedColumnCell)
 				{
-				case EColumnSizeMode::Fill:
-				{
-					TAttribute<float> WidthBinding;
-					WidthBinding.BindRaw(&Column, &SHeaderRow::FColumn::GetWidth);
-
-					Box->AddSlot()
-					.HAlign(Column.CellHAlignment)
-					.VAlign(Column.CellVAlignment)
-					.FillWidth(WidthBinding)
-					[
-						CellContents
-					];
-				}
-				break;
+					switch (Column.SizeRule)
+					{
+					case EColumnSizeMode::Fill:
+					{
+						TAttribute<float> WidthBinding;
+						WidthBinding.BindRaw(&Column, &SHeaderRow::FColumn::GetWidth);
 
-				case EColumnSizeMode::Fixed:
-				{
-					Box->AddSlot()
-					.AutoWidth()
-					[
-						SNew(SBox)
-						.WidthOverride(Column.Width.Get())
+						Box->AddSlot()
 						.HAlign(Column.CellHAlignment)
 						.VAlign(Column.CellVAlignment)
-						.Clipping(EWidgetClipping::OnDemand)
+						.FillWidth(WidthBinding)
 						[
 							CellContents
-						]
-					];
-				}
-				break;
+						];
+					}
+					break;
 
-				case EColumnSizeMode::Manual:
-				case EColumnSizeMode::FillSized:
-				{
-					auto GetColumnWidthAsOptionalSize = [&Column]() -> FOptionalSize
+					case EColumnSizeMode::Fixed:
 					{
-						const float DesiredWidth = Column.GetWidth();
-						return FOptionalSize(DesiredWidth);
-					};
-
-					TAttribute<FOptionalSize> WidthBinding;
-					WidthBinding.Bind(TAttribute<FOptionalSize>::FGetter::CreateLambda(GetColumnWidthAsOptionalSize));
-
-					Box->AddSlot()
-					.AutoWidth()
-					[
-						SNew(SBox)
-						.WidthOverride(WidthBinding)
-						.HAlign(Column.CellHAlignment)
-						.VAlign(Column.CellVAlignment)
-						.Clipping(EWidgetClipping::OnDemand)
+						Box->AddSlot()
+						.AutoWidth()
 						[
-							CellContents
-						]
-					];
-				}
-				break;
+							SNew(SBox)
+							.WidthOverride(Column.Width.Get())
+							.HAlign(Column.CellHAlignment)
+							.VAlign(Column.CellVAlignment)
+							.Clipping(EWidgetClipping::OnDemand)
+							[
+								CellContents
+							]
+						];
+					}
+					break;
+
+					case EColumnSizeMode::Manual:
+					case EColumnSizeMode::FillSized:
+					{
+						auto GetColumnWidthAsOptionalSize = [&Column]() -> FOptionalSize
+						{
+							const float DesiredWidth = Column.GetWidth();
+							return FOptionalSize(DesiredWidth);
+						};
 
-				default:
-					ensure(false);
+						TAttribute<FOptionalSize> WidthBinding;
+						WidthBinding.Bind(TAttribute<FOptionalSize>::FGetter::CreateLambda(GetColumnWidthAsOptionalSize));
+
+						Box->AddSlot()
+						.AutoWidth()
+						[
+							SNew(SBox)
+							.WidthOverride(WidthBinding)
+							.HAlign(Column.CellHAlignment)
+							.VAlign(Column.CellVAlignment)
+							.Clipping(EWidgetClipping::OnDemand)
+							[
+								CellContents
+							]
+						];
+					}
 					break;
+
+					default:
+						ensure(false);
+						break;
+					}
 				}
 
 				NewColumnIdToSlotContents.Add(Column.ColumnId, CellContents);
