[Tabs Docking] Crash when docking the Undo History window

Could not repro the crash from the jira issue, but the stack points to a line with an unchecked dereferenced pointer.
Adding a check to prevent potential crashing.

diff --git a/Engine/Source/Runtime/Slate/Private/Framework/Docking/SDockingTabStack.cpp b/Engine/Source/Runtime/Slate/Private/Framework/Docking/SDockingTabStack.cpp
--- a/Engine/Source/Runtime/Slate/Private/Framework/Docking/SDockingTabStack.cpp
+++ b/Engine/Source/Runtime/Slate/Private/Framework/Docking/SDockingTabStack.cpp
@@ -282,18 +282,22 @@ void SDockingTabStack::OpenTab(const TSharedRef<SDockTab>& InTab, int32 InsertLo
 	const FTabId TabId = InTab->GetLayoutIdentifier();
 
 	// the insert index is not the same as the tab index in the array for new tabs so find the tab again to check the tab state.
-	const FTabManager::FTab& TabInfo = *Tabs.FindByPredicate([TabId](const FTabManager::FTab& TestTab) {return TestTab.TabId == TabId; });
+	const FTabManager::FTab* TabInfo = Tabs.FindByPredicate([TabId](const FTabManager::FTab& TestTab) {return TestTab.TabId == TabId; });
+	if (!TabInfo)
+	{
+		return;
+	}
 
-	if (TabInfo.TabState == ETabState::SidebarTab)
+	if (TabInfo->TabState == ETabState::SidebarTab)
 	{
 		FSidebarTabLists SidebarLists;
-		if (TabInfo.SidebarLocation == ESidebarLocation::Left)
+		if (TabInfo->SidebarLocation == ESidebarLocation::Left)
 		{
 			SidebarLists.LeftSidebarTabs.Add(InTab);
 		}
 		else
 		{
-			ensure(TabInfo.SidebarLocation == ESidebarLocation::Right);
+			ensure(TabInfo->SidebarLocation == ESidebarLocation::Right);
 			SidebarLists.RightSidebarTabs.Add(InTab);
 		}
 
