NotificationManager - Prevent moving & showing notifications outside of the screen
Tentative fix for UE-314545

#tests Tested on Mac, Windows & Linux (thx James)

diff --git a/Engine/Source/Runtime/Slate/Private/Framework/Notifications/NotificationManager.cpp b/Engine/Source/Runtime/Slate/Private/Framework/Notifications/NotificationManager.cpp
--- a/Engine/Source/Runtime/Slate/Private/Framework/Notifications/NotificationManager.cpp
+++ b/Engine/Source/Runtime/Slate/Private/Framework/Notifications/NotificationManager.cpp
@@ -66,24 +66,29 @@ void FSlateNotificationManager::FRegionalNotificationList::Arrange()
 			const FVector2D DesiredSize = PinnedWindow->GetDesiredSize();
 			
 			const FVector2D NewPosition(AnchorPoint.X - DesiredSize.X, AnchorPoint.Y - DesiredSize.Y - StackOffset);
-
-			if( NewPosition != PinnedWindow->GetPositionInScreen() && DesiredSize != PinnedWindow->GetSizeInScreen())
-			{
-				PinnedWindow->ReshapeWindow( NewPosition, DesiredSize );
-			}
-			else if (NewPosition != PinnedWindow->GetPositionInScreen())
-			{
-				PinnedWindow->MoveWindowTo(NewPosition);
-			}
 			
-			// Only display the notification when it's ready and arranged
-			if constexpr (!bShowNotificationImmediately)
+			// Check if the point is inside the screen to prevent triggering change screen event
+			if (Region.ContainsPoint(NewPosition))
 			{
+				if( NewPosition != PinnedWindow->GetPositionInScreen() && DesiredSize != PinnedWindow->GetSizeInScreen())
+				{
+					PinnedWindow->ReshapeWindow( NewPosition, DesiredSize );
+				}
+				else if (NewPosition != PinnedWindow->GetPositionInScreen())
+				{
+					PinnedWindow->MoveWindowTo(NewPosition);
+				}
+				
+				// Only display the notification when it's ready and arranged
 				if (!PinnedWindow->IsVisible())
 				{
 					PinnedWindow->ShowWindow();
 				}
 			}
+			else if (PinnedWindow->IsVisible())
+			{
+				PinnedWindow->HideWindow();
+			}
 			
 			AnchorPoint.Y -= DesiredSize.Y;
 		}
