Fixing some memory leak found by our automation in the file.

diff --git a/Engine/Source/Runtime/Slate/Private/Framework/MultiBox/Mac/MacMenu.cpp b/Engine/Source/Runtime/Slate/Private/Framework/MultiBox/Mac/MacMenu.cpp
--- a/Engine/Source/Runtime/Slate/Private/Framework/MultiBox/Mac/MacMenu.cpp
+++ b/Engine/Source/Runtime/Slate/Private/Framework/MultiBox/Mac/MacMenu.cpp
@@ -607,10 +607,10 @@ void FSlateMacMenu::UpdateWithMultiBox(const TSharedPtr< FMultiBox > MultiBox)
 					[Menu addItem:[NSMenuItem separatorItem]];
 				}
 			}
-
-			delete SafeMultiBoxPtr;
 		}
 
+		delete SafeMultiBoxPtr;
+
 		FPlatformApplicationMisc::bChachedMacMenuStateNeedsUpdate = true;
 	}, false);
 	// Now we destroy all the menu blocks in the game thread to ensure ~SWidget() is only called in Game Thread
@@ -916,6 +916,10 @@ void FSlateMacMenu::ExecuteMenuItemAction(const TSharedRef< const class FMenuEnt
 			delete MenuBlock;
 		}, false);
 	}
+	else
+	{
+		delete MenuBlock;
+	}
 }
 
 static const TSharedRef<SWidget> FindTextBlockWidget(TSharedRef<SWidget> Content)
