Dynamic Invalidation - TypeTraits

Invalidation support is not a property a widget should inherit from its parent as we don't want to assume descendant widgets have been made invalidation compatible.
The current version relies on a property added to the widget class, and on widgets marking themselves as invalidation incompatible or compatible using respectively the `SLATE_DECLARE_WIDGET` and the newly introduced `SLATE_DECLARE_WIDGET_SUPPORTING_INVALIDATION` macros.
The issue is that many widgets do not use these macros and end up inheriting the widget class of their parent, therefore also inheriting their invalidation compatibility. To avoid this issue, it was proposed to verify that the widget type name (set in `SNew`) would match the name of the widget class (set in `SLATE_DECLARE_WIDGET`).
Even though this is working for most widgets, template widgets have been an issue as their widget class name (ex: `STableRow`, note that it excludes the template type) does not exactly match the widget type name (ex: `STableRow< TSharedPtr< FAssetViewItem > >`)

To avoid this, this CL proposes the use of a TypeTrait instead of the `SLATE_DECLARE_WIDGET_SUPPORTING_INVALIDATION` macro which allows certainty on the invalidation compatibility. As we do not have access to the type of the widget in `SWidget::SupportsInvalidation` to use the TypeTrait, a variable has been added to each widget to indicate if they are invalidation compatible. It is set when the widget is constructed via `SNew` and equivalent, as this is the only time where we have access to their type.

diff --git a/Engine/Source/Runtime/Slate/Private/Widgets/SInvalidationPanel.cpp b/Engine/Source/Runtime/Slate/Private/Widgets/SInvalidationPanel.cpp
--- a/Engine/Source/Runtime/Slate/Private/Widgets/SInvalidationPanel.cpp
+++ b/Engine/Source/Runtime/Slate/Private/Widgets/SInvalidationPanel.cpp
@@ -38,17 +38,33 @@ FAutoConsoleVariableRef CVarAlwaysInvalidate(
 #if UE_SLATE_WITH_DYNAMIC_INVALIDATION
 namespace DynamicInvalidationCVars
 {
-	static int32 bDynamicInvalidationOptions = 2;
-	static FAutoConsoleVariableRef CVarEnableDynamicInvalidation(
-		TEXT("Slate.DynamicInvalidationOptions"),
-		bDynamicInvalidationOptions,
+	void ConsoleVariableEnableDynamicInvalidationChanged(IConsoleVariable*);
+
+	static int32 DynamicInvalidationOptions = 2;
+	static FAutoConsoleVariableRef CVarDynamicInvalidationOptions(
+		TEXT("Slate.DynamicInvalidation.Options"),
+		DynamicInvalidationOptions,
 		TEXT("Toggle Dynamic Invalidation options for Invalidation Panels.\n")
 		TEXT(" A Dynamic Invalidation Panel refers to an Invalidation Panel with bUseDynamicInvalidation set to true.\n")
 		TEXT(" 0. Disable Dynamic Invalidation and disable caching for all Dynamic Invalidation Panels\n")
 		TEXT(" 1. Disable Dynamic Invalidation and makes the Dynamic Invalidation Panels act like regular Invalidation Panels\n")
 		TEXT(" 2. Enable Dynamic Invalidation for Dynamic Invalidation Panels\n")
-		TEXT(" 3. Forces Dynamic Invalidation on all Invalidation Panels")
+		TEXT(" 3. Forces Dynamic Invalidation on all Invalidation Panels"),
+	FConsoleVariableDelegate::CreateStatic(ConsoleVariableEnableDynamicInvalidationChanged)
 	);
+
+	void ConsoleVariableEnableDynamicInvalidationChanged(IConsoleVariable*)
+	{
+		DynamicInvalidationOptions = FMath::Clamp(DynamicInvalidationOptions, 0, 4);
+#if UE_SLATE_TRACE_ENABLED
+		const FString Description = DynamicInvalidationOptions == 0 ? TEXT("Disabled") :
+			DynamicInvalidationOptions == 1 ? TEXT("Disabled: forced Static") :
+			DynamicInvalidationOptions == 2 ? TEXT("Enabled") :
+			DynamicInvalidationOptions == 3 ? TEXT("Forced") : TEXT("Unknown");
+
+		UE_TRACE_SLATE_BOOKMARK(TEXT("Dynamic Invalidation changed [%d]: %s"), DynamicInvalidationOptions, *Description)
+#endif // UE_SLATE_TRACE_ENABLED
+	}
 }
 #endif
 
@@ -145,11 +161,11 @@ void SInvalidationPanel::EnableInvalidationPanels(bool bEnable)
 bool SInvalidationPanel::GetCanCache() const
 {
 #if UE_SLATE_WITH_DYNAMIC_INVALIDATION
-	const bool bForcedDynamicInvalidation = DynamicInvalidationCVars::bDynamicInvalidationOptions == 3;
-	const bool bDynamicInvalidationEnabled = bForcedDynamicInvalidation || (bUseDynamicInvalidation && DynamicInvalidationCVars::bDynamicInvalidationOptions == 2);
+	const bool bForcedDynamicInvalidation = DynamicInvalidationCVars::DynamicInvalidationOptions == 3;
+	const bool bDynamicInvalidationEnabled = bForcedDynamicInvalidation || (bUseDynamicInvalidation && DynamicInvalidationCVars::DynamicInvalidationOptions == 2);
 	if (!bDynamicInvalidationEnabled)
 	{
-		const bool bRevertToRegularInvalidation = !bUseDynamicInvalidation || DynamicInvalidationCVars::bDynamicInvalidationOptions == 1;
+		const bool bRevertToRegularInvalidation = !bUseDynamicInvalidation || DynamicInvalidationCVars::DynamicInvalidationOptions == 1;
 		if (!bRevertToRegularInvalidation)
 		{
 			return false;
@@ -262,8 +278,8 @@ FChildren* SInvalidationPanel::Debug_GetChildrenForReflector()
 #if UE_SLATE_WITH_DYNAMIC_INVALIDATION
 bool SInvalidationPanel::SupportsInvalidationRecursive(bool bUseCachedValue) const
 {
-	const bool bForcedDynamicInvalidation = DynamicInvalidationCVars::bDynamicInvalidationOptions == 3;
-	const bool bDynamicInvalidationEnabled = bForcedDynamicInvalidation || (bUseDynamicInvalidation && DynamicInvalidationCVars::bDynamicInvalidationOptions == 2);
+	const bool bForcedDynamicInvalidation = DynamicInvalidationCVars::DynamicInvalidationOptions == 3;
+	const bool bDynamicInvalidationEnabled = bForcedDynamicInvalidation || (bUseDynamicInvalidation && DynamicInvalidationCVars::DynamicInvalidationOptions == 2);
 	// For an invalidation panel without bUseDynamicInvalidation, we force the caching of all children, so we always support invalidation
 	return !bDynamicInvalidationEnabled || SWidget::SupportsInvalidationRecursive(bUseCachedValue);
 }
diff --git a/Engine/Source/Runtime/Slate/Public/Widgets/SInvalidationPanel.h b/Engine/Source/Runtime/Slate/Public/Widgets/SInvalidationPanel.h
--- a/Engine/Source/Runtime/Slate/Public/Widgets/SInvalidationPanel.h
+++ b/Engine/Source/Runtime/Slate/Public/Widgets/SInvalidationPanel.h
@@ -18,9 +18,16 @@ class FSlateRenderDataHandle;
 class FSlateWindowElementList;
 class SWindow;
 
+
+template <>
+struct TWidgetTypeTraits<class SInvalidationPanel>
+{
+	static constexpr bool SupportsInvalidation() { return true; }
+};
+
 class SInvalidationPanel : public SCompoundWidget, public FSlateInvalidationRoot
 {
-	SLATE_DECLARE_WIDGET_SUPPORTING_INVALIDATION_API(SInvalidationPanel, SCompoundWidget, SLATE_API)
+	SLATE_DECLARE_WIDGET_API(SInvalidationPanel, SCompoundWidget, SLATE_API)
 public:
 	SLATE_BEGIN_ARGS( SInvalidationPanel )
 #if UE_SLATE_WITH_DYNAMIC_INVALIDATION
diff --git a/Engine/Source/Runtime/Slate/Public/Widgets/Views/STableRow.h b/Engine/Source/Runtime/Slate/Public/Widgets/Views/STableRow.h
--- a/Engine/Source/Runtime/Slate/Public/Widgets/Views/STableRow.h
+++ b/Engine/Source/Runtime/Slate/Public/Widgets/Views/STableRow.h
@@ -68,6 +68,14 @@ DECLARE_DELEGATE_OneParam(FOnTableRowDragEnter, FDragDropEvent const&);
 DECLARE_DELEGATE_OneParam(FOnTableRowDragLeave, FDragDropEvent const&);
 DECLARE_DELEGATE_RetVal_OneParam(FReply, FOnTableRowDrop, FDragDropEvent const&);
 
+template<typename ItemType>
+class STableRow;
+
+template <typename ItemType>
+struct TWidgetTypeTraits<STableRow<ItemType>>
+{
+	static constexpr bool SupportsInvalidation() { return true; }
+};
 
 /**
  * The ListView is populated by Selectable widgets.
@@ -1294,6 +1302,15 @@ private:
 };
 
 
+template<typename ItemType>
+class SMultiColumnTableRow;
+
+template <typename ItemType>
+struct TWidgetTypeTraits<SMultiColumnTableRow<ItemType>>
+{
+	static constexpr bool SupportsInvalidation() { return true; }
+};
+
 template<typename ItemType>
 class SMultiColumnTableRow : public STableRow<ItemType>
 {
diff --git a/Engine/Source/Runtime/Slate/Slate.Build.cs b/Engine/Source/Runtime/Slate/Slate.Build.cs
--- a/Engine/Source/Runtime/Slate/Slate.Build.cs
+++ b/Engine/Source/Runtime/Slate/Slate.Build.cs
@@ -18,6 +18,10 @@ public class Slate : ModuleRules
 				"SlateCore",
 				"ImageWrapper"
 			});
+		PrivateDependencyModuleNames.AddRange(
+			new string[] {
+				"TraceLog",
+			});
 
 		if (Target.bCompileAgainstApplicationCore)
 		{
diff --git a/Engine/Source/Runtime/SlateCore/Private/Widgets/SLeafWidget.cpp b/Engine/Source/Runtime/SlateCore/Private/Widgets/SLeafWidget.cpp
--- a/Engine/Source/Runtime/SlateCore/Private/Widgets/SLeafWidget.cpp
+++ b/Engine/Source/Runtime/SlateCore/Private/Widgets/SLeafWidget.cpp
@@ -5,8 +5,6 @@
 
 /* SLeafWidget interface
  *****************************************************************************/
-SLATE_IMPLEMENT_WIDGET(SLeafWidget)
-void SLeafWidget::PrivateRegisterAttributes(FSlateAttributeInitializer& AttributeInitializer) {}
 
 SLeafWidget::SLeafWidget()
 {
diff --git a/Engine/Source/Runtime/SlateCore/Private/Widgets/SNullWidget.cpp b/Engine/Source/Runtime/SlateCore/Private/Widgets/SNullWidget.cpp
--- a/Engine/Source/Runtime/SlateCore/Private/Widgets/SNullWidget.cpp
+++ b/Engine/Source/Runtime/SlateCore/Private/Widgets/SNullWidget.cpp
@@ -7,11 +7,15 @@
 #include "Layout/Children.h"
 #include "Widgets/DeclarativeSyntaxSupport.h"
 
+template <>
+struct TWidgetTypeTraits<class SNullWidgetContent>
+{
+	static constexpr bool SupportsInvalidation() { return true; }
+};
 
 class SLATECORE_API SNullWidgetContent
 	: public SWidget
 {
-	SLATE_DECLARE_WIDGET_SUPPORTING_INVALIDATION(SNullWidgetContent, SWidget)
 public:
 
 	SLATE_BEGIN_ARGS(SNullWidgetContent)
@@ -69,9 +73,6 @@ protected:
 	// End of SWidget interface
 };
 
-SLATE_IMPLEMENT_WIDGET(SNullWidgetContent)
-void SNullWidgetContent::PrivateRegisterAttributes(FSlateAttributeInitializer& AttributeInitializer) {}
-
 namespace NullWidgetPrivate
 {
 	TSharedRef<SWidget> Construct()
diff --git a/Engine/Source/Runtime/SlateCore/Private/Widgets/SPanel.cpp b/Engine/Source/Runtime/SlateCore/Private/Widgets/SPanel.cpp
--- a/Engine/Source/Runtime/SlateCore/Private/Widgets/SPanel.cpp
+++ b/Engine/Source/Runtime/SlateCore/Private/Widgets/SPanel.cpp
@@ -4,9 +4,6 @@
 #include "Types/PaintArgs.h"
 #include "Layout/ArrangedChildren.h"
 
-SLATE_IMPLEMENT_WIDGET(SPanel);
-void SPanel::PrivateRegisterAttributes(FSlateAttributeInitializer& AttributeInitializer) {}
-
 SPanel::SPanel() = default;
 SPanel::~SPanel() = default;
 
diff --git a/Engine/Source/Runtime/SlateCore/Private/Widgets/SWidget.cpp b/Engine/Source/Runtime/SlateCore/Private/Widgets/SWidget.cpp
--- a/Engine/Source/Runtime/SlateCore/Private/Widgets/SWidget.cpp
+++ b/Engine/Source/Runtime/SlateCore/Private/Widgets/SWidget.cpp
@@ -1966,9 +1966,7 @@ void SWidget::AddMetadataInternal(const TSharedRef<ISlateMetaData>& AddMe)
 
 bool SWidget::SupportsInvalidation() const
 {
-	// We make sure that this exact widget's class is defined and supports invalidation,
-	// as we don't want to assume descendant widgets support invalidation
-	return GetType() == GetWidgetClass().GetWidgetType() ? GetWidgetClass().SupportsInvalidation() : false;
+	return bSupportsInvalidation;
 }
 
 UE_TRACE_EVENT_BEGIN(Cpu, SWidget_SupportsInvalidationRecursive, NoSync)
diff --git a/Engine/Source/Runtime/SlateCore/Public/Widgets/DeclarativeSyntaxSupport.h b/Engine/Source/Runtime/SlateCore/Public/Widgets/DeclarativeSyntaxSupport.h
--- a/Engine/Source/Runtime/SlateCore/Public/Widgets/DeclarativeSyntaxSupport.h
+++ b/Engine/Source/Runtime/SlateCore/Public/Widgets/DeclarativeSyntaxSupport.h
@@ -14,6 +14,7 @@
 #include "Trace/SlateMemoryTags.h"
 #include "Widgets/SNullWidget.h"
 #include "Widgets/Accessibility/SlateWidgetAccessibleTypes.h"
+#include "Widgets/SWidget.h"
 #include "Templates/Identity.h"
 
 class IToolTip;
@@ -924,6 +925,10 @@ struct TSlateDecl
 		}
 
 		_Widget->SetDebugInfo(InType, InFile, OnLine, sizeof(WidgetType));
+
+#if UE_SLATE_WITH_DYNAMIC_INVALIDATION
+		_Widget->bSupportsInvalidation = TWidgetTypeTraits<WidgetType>::SupportsInvalidation();
+#endif
 	}
 
 	/**
diff --git a/Engine/Source/Runtime/SlateCore/Public/Widgets/Images/SImage.h b/Engine/Source/Runtime/SlateCore/Public/Widgets/Images/SImage.h
--- a/Engine/Source/Runtime/SlateCore/Public/Widgets/Images/SImage.h
+++ b/Engine/Source/Runtime/SlateCore/Public/Widgets/Images/SImage.h
@@ -14,13 +14,20 @@
 class FPaintArgs;
 class FSlateWindowElementList;
 
+template <>
+struct TWidgetTypeTraits<class SImage>
+{
+	static constexpr bool SupportsInvalidation() { return true; }
+};
+
+
 /**
  * Implements a widget that displays an image with a desired width and height.
  */
 class SImage
 	: public SLeafWidget
 {
-	SLATE_DECLARE_WIDGET_SUPPORTING_INVALIDATION_API(SImage, SLeafWidget, SLATECORE_API)
+	SLATE_DECLARE_WIDGET_API(SImage, SLeafWidget, SLATECORE_API)
 
 public:
 	SLATE_BEGIN_ARGS( SImage )
diff --git a/Engine/Source/Runtime/SlateCore/Public/Widgets/SBoxPanel.h b/Engine/Source/Runtime/SlateCore/Public/Widgets/SBoxPanel.h
--- a/Engine/Source/Runtime/SlateCore/Public/Widgets/SBoxPanel.h
+++ b/Engine/Source/Runtime/SlateCore/Public/Widgets/SBoxPanel.h
@@ -20,12 +20,19 @@
 
 class FArrangedChildren;
 
+template <>
+struct TWidgetTypeTraits<class SBoxPanel>
+{
+	static constexpr bool SupportsInvalidation() { return true; }
+};
+
+
 /**
  * A BoxPanel contains one child and describes how that child should be arranged on the screen.
  */
 class SBoxPanel : public SPanel
 {
-	SLATE_DECLARE_WIDGET_SUPPORTING_INVALIDATION_API(SBoxPanel, SPanel, SLATECORE_API)
+	SLATE_DECLARE_WIDGET_API(SBoxPanel, SPanel, SLATECORE_API)
 
 protected:
 	/**
@@ -153,10 +160,16 @@ protected:
 };
 
 
+template <>
+struct TWidgetTypeTraits<class SHorizontalBox>
+{
+	static constexpr bool SupportsInvalidation() { return true; }
+};
+
 /** A Horizontal Box Panel. See SBoxPanel for more info. */
 class SHorizontalBox : public SBoxPanel
 {
-	SLATE_DECLARE_WIDGET_SUPPORTING_INVALIDATION_API(SHorizontalBox, SBoxPanel, SLATECORE_API)
+	SLATE_DECLARE_WIDGET_API(SHorizontalBox, SBoxPanel, SLATECORE_API)
 public:
 	class FSlot : public SBoxPanel::TSlot<FSlot>
 	{
@@ -297,10 +310,17 @@ public:
 	SLATECORE_API void Construct( const FArguments& InArgs );
 };
 
+
+template <>
+struct TWidgetTypeTraits<class SVerticalBox>
+{
+	static constexpr bool SupportsInvalidation() { return true; }
+};
+
 /** A Vertical Box Panel. See SBoxPanel for more info. */
 class SVerticalBox : public SBoxPanel
 {
-	SLATE_DECLARE_WIDGET_SUPPORTING_INVALIDATION_API(SVerticalBox, SBoxPanel, SLATECORE_API)
+	SLATE_DECLARE_WIDGET_API(SVerticalBox, SBoxPanel, SLATECORE_API)
 public:
 	class FSlot : public SBoxPanel::TSlot<FSlot>
 	{
@@ -443,10 +463,17 @@ public:
 	SLATECORE_API void Construct( const FArguments& InArgs );
 };
 
+
+template <>
+struct TWidgetTypeTraits<class SStackBox>
+{
+	static constexpr bool SupportsInvalidation() { return true; }
+};
+
 /** A Stack Box Panel that stack vertically or horizontally. See SBoxPanel for more info. */
 class SStackBox : public SBoxPanel
 {
-	SLATE_DECLARE_WIDGET_SUPPORTING_INVALIDATION_API(SStackBox, SBoxPanel, SLATECORE_API)
+	SLATE_DECLARE_WIDGET_API(SStackBox, SBoxPanel, SLATECORE_API)
 public:
 	class FSlot : public SBoxPanel::TSlot<FSlot>
 	{
diff --git a/Engine/Source/Runtime/SlateCore/Public/Widgets/SLeafWidget.h b/Engine/Source/Runtime/SlateCore/Public/Widgets/SLeafWidget.h
--- a/Engine/Source/Runtime/SlateCore/Public/Widgets/SLeafWidget.h
+++ b/Engine/Source/Runtime/SlateCore/Public/Widgets/SLeafWidget.h
@@ -12,6 +12,12 @@ class FArrangedChildren;
 class FPaintArgs;
 class FSlateWindowElementList;
 
+template <>
+struct TWidgetTypeTraits<class SLeafWidget>
+{
+	static constexpr bool SupportsInvalidation() { return true; }
+};
+
 /**
  * Implements a leaf widget.
  *
@@ -21,7 +27,6 @@ class FSlateWindowElementList;
 class SLeafWidget
 	: public SWidget
 {
-	SLATE_DECLARE_WIDGET_SUPPORTING_INVALIDATION_API(SLeafWidget, SWidget, SLATECORE_API)
 public:
 	SLATECORE_API SLeafWidget();
 
diff --git a/Engine/Source/Runtime/SlateCore/Public/Widgets/SOverlay.h b/Engine/Source/Runtime/SlateCore/Public/Widgets/SOverlay.h
--- a/Engine/Source/Runtime/SlateCore/Public/Widgets/SOverlay.h
+++ b/Engine/Source/Runtime/SlateCore/Public/Widgets/SOverlay.h
@@ -18,6 +18,13 @@ class FArrangedChildren;
 class FPaintArgs;
 class FSlateWindowElementList;
 
+
+template <>
+struct TWidgetTypeTraits<class SOverlay>
+{
+	static constexpr bool SupportsInvalidation() { return true; }
+};
+
 /**
  * Implements an overlay widget.
  *
@@ -35,7 +42,7 @@ class FSlateWindowElementList;
  */
 class SOverlay : public SPanel
 {
-	SLATE_DECLARE_WIDGET_SUPPORTING_INVALIDATION_API(SOverlay, SPanel, SLATECORE_API)
+	SLATE_DECLARE_WIDGET_API(SOverlay, SPanel, SLATECORE_API)
 public:	
 
 	/** A slot that support alignment of content and padding and z-order */
diff --git a/Engine/Source/Runtime/SlateCore/Public/Widgets/SPanel.h b/Engine/Source/Runtime/SlateCore/Public/Widgets/SPanel.h
--- a/Engine/Source/Runtime/SlateCore/Public/Widgets/SPanel.h
+++ b/Engine/Source/Runtime/SlateCore/Public/Widgets/SPanel.h
@@ -12,6 +12,12 @@ class FChildren;
 class FPaintArgs;
 class FSlateWindowElementList;
 
+template <>
+struct TWidgetTypeTraits<class SPanel>
+{
+	static constexpr bool SupportsInvalidation() { return true; }
+};
+
 /**
  * A Panel arranges its child widgets on the screen.
  *
@@ -22,7 +28,6 @@ class FSlateWindowElementList;
 class SPanel
 	: public SWidget
 {
-	SLATE_DECLARE_WIDGET_SUPPORTING_INVALIDATION_API(SPanel, SWidget, SLATECORE_API)
 public:	
 
 	/**
diff --git a/Engine/Source/Runtime/SlateCore/Public/Widgets/SWidget.h b/Engine/Source/Runtime/SlateCore/Public/Widgets/SWidget.h
--- a/Engine/Source/Runtime/SlateCore/Public/Widgets/SWidget.h
+++ b/Engine/Source/Runtime/SlateCore/Public/Widgets/SWidget.h
@@ -63,6 +63,15 @@ DECLARE_DWORD_ACCUMULATOR_STAT_EXTERN(TEXT("Total Widgets"), STAT_SlateTotalWidg
 DECLARE_MEMORY_STAT_EXTERN(TEXT("SWidget Total Allocated Size"), STAT_SlateSWidgetAllocSize, STATGROUP_SlateMemory, SLATECORE_API);
 
 
+/** Type Traits to cover the generic information of a Widget type **/
+template <typename WidgetType>
+struct TWidgetTypeTraits
+{
+	/** Returns true if a widget supports invalidation. **/
+	static constexpr bool SupportsInvalidation() { return false; }
+};
+
+
 enum class EAccessibleType : uint8
 {
 	Main,
@@ -1789,6 +1798,14 @@ private:
 	/** Is the PrepassLayoutScaleMultiplier set. */
 	uint8 bPrepassLayoutScaleMultiplierSet : 1;
 
+#if UE_SLATE_WITH_DYNAMIC_INVALIDATION
+	template<class WidgetType, typename RequiredArgsPayloadType>
+	friend struct TSlateDecl;
+
+	// We need to keep this information for each widget as not all widgets have a class to store this information
+	uint8 bSupportsInvalidation : 1 = false;
+#endif
+
 protected:
 	uint8 bHasCustomPrepass : 1;
 
diff --git a/Engine/Source/Runtime/SlateCore/Public/Widgets/SlateControlledConstruction.h b/Engine/Source/Runtime/SlateCore/Public/Widgets/SlateControlledConstruction.h
--- a/Engine/Source/Runtime/SlateCore/Public/Widgets/SlateControlledConstruction.h
+++ b/Engine/Source/Runtime/SlateCore/Public/Widgets/SlateControlledConstruction.h
@@ -17,20 +17,9 @@ namespace SharedPointerInternals
 
 /** */
 #define SLATE_DECLARE_WIDGET(WidgetType, ParentType) \
-	SLATE_DECLARE_WIDGET_IMPL(WidgetType, ParentType, false, NO_API)
-
-#define SLATE_DECLARE_WIDGET_SUPPORTING_INVALIDATION(WidgetType, ParentType) \
-	SLATE_DECLARE_WIDGET_IMPL(WidgetType, ParentType, true, NO_API)
-
+	SLATE_DECLARE_WIDGET_API(WidgetType, ParentType, NO_API)
 
 #define SLATE_DECLARE_WIDGET_API(WidgetType, ParentType, ModuleApiDefine) \
-	SLATE_DECLARE_WIDGET_IMPL(WidgetType, ParentType, false, ModuleApiDefine)
-
-#define SLATE_DECLARE_WIDGET_SUPPORTING_INVALIDATION_API(WidgetType, ParentType, ModuleApiDefine) \
-	SLATE_DECLARE_WIDGET_IMPL(WidgetType, ParentType, true, ModuleApiDefine)
-
-
-#define SLATE_DECLARE_WIDGET_IMPL(WidgetType, ParentType, bInSupportsInvalidation, ModuleApiDefine) \
 	private: \
 		using Super = ParentType; \
 		using ThisClass = WidgetType; \
@@ -42,7 +31,7 @@ namespace SharedPointerInternals
 		friend class SharedPointerInternals::TIntrusiveReferenceController; /* Shouldn't really forward-declare this to allow MakeShared to work, but is for legacy reasons */ \
 		static const FSlateWidgetClassData& GetPrivateWidgetClass() \
 		{ \
-			static FSlateWidgetClassData WidgetClassDataInstance = FSlateWidgetClassData(TIdentity<ParentType>(), #WidgetType, bInSupportsInvalidation, &WidgetType::PrivateRegisterAttributes); \
+			static FSlateWidgetClassData WidgetClassDataInstance = FSlateWidgetClassData(TIdentity<ParentType>(), #WidgetType, &WidgetType::PrivateRegisterAttributes); \
 			return WidgetClassDataInstance; \
 		} \
 		static ModuleApiDefine void PrivateRegisterAttributes(FSlateAttributeInitializer&); \
@@ -63,18 +52,12 @@ private:
 
 	FSlateWidgetClassData(FName InWidgetTypeName)
 		: WidgetType(InWidgetTypeName)
-#if UE_SLATE_WITH_DYNAMIC_INVALIDATION
-		, bSupportsInvalidation(false)
-#endif
 	{}
 
 public:
 	template<typename InWidgetParentType>
-	FSlateWidgetClassData(TIdentity<InWidgetParentType>, FName InWidgetTypeName, bool bInSupportsInvalidation, void(*AttributeInitializer)(FSlateAttributeInitializer&))
+	FSlateWidgetClassData(TIdentity<InWidgetParentType>, FName InWidgetTypeName, void(*AttributeInitializer)(FSlateAttributeInitializer&))
 		: WidgetType(InWidgetTypeName)
-#if UE_SLATE_WITH_DYNAMIC_INVALIDATION
-		, bSupportsInvalidation(bInSupportsInvalidation)
-#endif
 	{
 		// Initialize the parent class if it's not already done
 		const FSlateWidgetClassData& ParentWidgetClassData = InWidgetParentType::StaticWidgetClass();
@@ -85,16 +68,10 @@ public:
 
 	const FSlateAttributeDescriptor& GetAttributeDescriptor() const { return Descriptor; };
 	FName GetWidgetType() const { return WidgetType; }
-#if UE_SLATE_WITH_DYNAMIC_INVALIDATION
-	bool SupportsInvalidation() const { return bSupportsInvalidation; }
-#endif
 
 private:
 	FSlateAttributeDescriptor Descriptor;
 	FName WidgetType;
-#if UE_SLATE_WITH_DYNAMIC_INVALIDATION
-	bool bSupportsInvalidation;
-#endif
 };
 
 /** */
