[TabManager] Do not take up space for the DisplayName if specifically requested not to

diff --git a/Engine/Source/Runtime/Slate/Private/Framework/Docking/SDockingTabWell.cpp b/Engine/Source/Runtime/Slate/Private/Framework/Docking/SDockingTabWell.cpp
--- a/Engine/Source/Runtime/Slate/Private/Framework/Docking/SDockingTabWell.cpp
+++ b/Engine/Source/Runtime/Slate/Private/Framework/Docking/SDockingTabWell.cpp
@@ -8,6 +8,8 @@
 
 const FVector2D FDockingConstants::MaxMinorTabSize(160.f, 25.0f);
 const FVector2D FDockingConstants::MaxMajorTabSize(210.f, 50.f);
+const float FDockingConstants::MaxTabSizeNoNameWidth(53.f);
+const float FDockingConstants::MaxTabSizeNoNameCantCloseWidth(32.f);
 
 const FVector2D FDockingConstants::GetMaxTabSizeFor( ETabRole TabRole )
 {
@@ -93,16 +95,39 @@ void SDockingTabWell::OnArrangeChildren( const FGeometry& AllottedGeometry, FArr
 
 	// Tabs have a uniform size.
 	const FVector2D ChildSize = ComputeChildSize(AllottedGeometry);
+	const FVector2D ChildSizeNoName = FVector2D(FDockingConstants::MaxTabSizeNoNameWidth, ChildSize.Y);
+	const FVector2D ChildSizeNoNameCantClose = FVector2D(FDockingConstants::MaxTabSizeNoNameCantCloseWidth, ChildSize.Y);
 
+	// Get the correct child size to use for the current tab
+	const auto GetChildSizeToUse = [ChildSize, ChildSizeNoName, ChildSizeNoNameCantClose] (const TSharedPtr<SDockTab>& InCurTab)
+		{
+			if (!InCurTab.IsValid())
+			{
+				return ChildSize;
+			}
 
-	const float DraggedChildCenter = ChildBeingDraggedOffset + ChildSize.X / 2;
+			if (InCurTab->IsTabNameHidden())
+			{
+				if (InCurTab->CanCloseTab())
+				{
+					return ChildSizeNoName;
+				}
+				return ChildSizeNoNameCantClose;
+			}
+			return ChildSize;
+		};
 
 	// Arrange all the tabs left to right.
 	float XOffset = 0;
+
 	for( int32 TabIndex=0; TabIndex < NumChildren; ++TabIndex )
 	{
 		const TSharedRef<SDockTab> CurTab = Tabs[TabIndex];
-		const float ChildWidthWithOverlap = ChildSize.X - CurTab->GetOverlapWidth();
+		FVector2D ChildSizeToUse = GetChildSizeToUse(CurTab);
+
+		const float DraggedChildCenter = ChildBeingDraggedOffset + ChildSizeToUse.X / 2;
+
+		const float ChildWidthWithOverlap = ChildSizeToUse.X - CurTab->GetOverlapWidth();
 
 		// This tab being dragged is arranged later.  It should not be arranged twice
 		if (CurTab == TabBeingDragged)
@@ -117,7 +142,7 @@ void SDockingTabWell::OnArrangeChildren( const FGeometry& AllottedGeometry, FArr
 			XOffset += ChildWidthWithOverlap;
 		}
 
-		ArrangedChildren.AddWidget( AllottedGeometry.MakeChild(CurTab, FVector2D(XOffset, 0), ChildSize) );
+		ArrangedChildren.AddWidget( AllottedGeometry.MakeChild(CurTab, FVector2D(XOffset, 0), ChildSizeToUse) );
 
 		XOffset += ChildWidthWithOverlap;
 	}
@@ -125,7 +150,8 @@ void SDockingTabWell::OnArrangeChildren( const FGeometry& AllottedGeometry, FArr
 	// Arrange the tab currently being dragged by the user, if any
 	if ( TabBeingDragged.IsValid() )
 	{
-		ArrangedChildren.AddWidget( AllottedGeometry.MakeChild( TabBeingDragged.ToSharedRef(), FVector2D(ChildBeingDraggedOffset,0), ChildSize) );
+		const FVector2D ChildSizeToUse = GetChildSizeToUse(TabBeingDragged);
+		ArrangedChildren.AddWidget( AllottedGeometry.MakeChild( TabBeingDragged.ToSharedRef(), FVector2D(ChildBeingDraggedOffset,0), ChildSizeToUse) );
 	}
 }
 	
@@ -287,7 +313,16 @@ FVector2D SDockingTabWell::ComputeChildSize( const FGeometry& AllottedGeometry )
 float SDockingTabWell::ComputeDraggedTabOffset( const FGeometry& MyGeometry, const FPointerEvent& MouseEvent, const FVector2D& InTabGrabOffsetFraction ) const
 {
 	const FVector2D ComputedChildSize = ComputeChildSize(MyGeometry);
-	return MyGeometry.AbsoluteToLocal( MouseEvent.GetScreenSpacePosition() ).X - InTabGrabOffsetFraction .X * ComputedChildSize.X;
+	const FVector2D ChildSizeNoName = FVector2D(FDockingConstants::MaxTabSizeNoNameWidth, ComputedChildSize.Y);
+	const FVector2D ChildSizeNoNameCantClose = FVector2D(FDockingConstants::MaxTabSizeNoNameCantCloseWidth, ComputedChildSize.Y);
+
+	FVector2D ChildSizeToUse = ComputedChildSize;
+	if (TabBeingDraggedPtr.IsValid() && TabBeingDraggedPtr->IsTabNameHidden())
+	{
+		ChildSizeToUse = TabBeingDraggedPtr->CanCloseTab() ? ChildSizeNoName : ChildSizeNoNameCantClose;
+	}
+
+	return MyGeometry.AbsoluteToLocal( MouseEvent.GetScreenSpacePosition() ).X - InTabGrabOffsetFraction .X * ChildSizeToUse.X;
 }
 
 
@@ -295,7 +330,31 @@ int32 SDockingTabWell::ComputeChildDropIndex(const FGeometry& MyGeometry, const
 {
 	const float ChildWidth = ComputeChildSize(MyGeometry).X;
 	const float ChildWidthWithOverlap = ChildWidth - TabBeingDragged->GetOverlapWidth();
-	const float DraggedChildCenter = ChildBeingDraggedOffset + ChildWidth / 2;
+	float DraggedChildCenter = ChildBeingDraggedOffset + ChildWidth / 2;
+
+	// Consider difference in TabSize to adjust the dropped index to the correct one
+	float XOffset = 0.f;
+	for ( int32 TabIndex=0; TabIndex < Tabs.Num(); ++TabIndex )
+	{
+		if (XOffset >= ChildBeingDraggedOffset)
+		{
+			break;
+		}
+
+		XOffset += ChildWidth;
+		const TSharedRef<SDockTab>& SomeTab = Tabs[TabIndex];
+		if (SomeTab->IsTabNameHidden())
+		{
+			if (!SomeTab->CanCloseTab())
+			{
+				DraggedChildCenter += ChildWidth - FDockingConstants::MaxTabSizeNoNameCantCloseWidth;
+			}
+			else
+			{
+				DraggedChildCenter += ChildWidth - FDockingConstants::MaxTabSizeNoNameWidth;
+			}
+		}
+	}
 	const int32 DropLocationIndex = FMath::Clamp(static_cast<int32>(DraggedChildCenter / ChildWidthWithOverlap), 0, Tabs.Num());
 	return DropLocationIndex;
 }
diff --git a/Engine/Source/Runtime/Slate/Private/Framework/Docking/SDockingTabWell.h b/Engine/Source/Runtime/Slate/Private/Framework/Docking/SDockingTabWell.h
--- a/Engine/Source/Runtime/Slate/Private/Framework/Docking/SDockingTabWell.h
+++ b/Engine/Source/Runtime/Slate/Private/Framework/Docking/SDockingTabWell.h
@@ -19,6 +19,8 @@ struct FDockingConstants
 {
 	static const FVector2D MaxMinorTabSize;
 	static const FVector2D MaxMajorTabSize;
+	static const float MaxTabSizeNoNameWidth;
+	static const float MaxTabSizeNoNameCantCloseWidth;
 	static const FVector2D SLATE_API GetMaxTabSizeFor( ETabRole TabRole );
 };
 
diff --git a/Engine/Source/Runtime/Slate/Private/Framework/Docking/TabManager.cpp b/Engine/Source/Runtime/Slate/Private/Framework/Docking/TabManager.cpp
--- a/Engine/Source/Runtime/Slate/Private/Framework/Docking/TabManager.cpp
+++ b/Engine/Source/Runtime/Slate/Private/Framework/Docking/TabManager.cpp
@@ -1177,8 +1177,10 @@ void FTabManager::MakeSpawnerMenuEntry( FMenuBuilder &PopulateMe, const TSharedP
 	// We don't want to add a menu entry for this tab if it is hidden, or if we are in read only mode and it is asking to be hidden
 	if (InSpawnerNode->MenuType.Get() != ETabSpawnerMenuType::Hidden && !(bReadOnly && InSpawnerNode->ReadOnlyBehavior == ETabReadOnlyBehavior::Hidden) )
 	{
+		const FText Label = GetTabLabelBasedOnSpawner(InSpawnerNode);
+
 		PopulateMe.AddMenuEntry(
-			InSpawnerNode->GetDisplayName().IsEmpty() ? FText::FromName(InSpawnerNode->TabType ) : InSpawnerNode->GetDisplayName(),
+			Label,
 			InSpawnerNode->GetTooltipText(),
 			InSpawnerNode->GetIcon(),
 			GetUIActionForTabSpawnerMenuEntry(InSpawnerNode),
@@ -2092,8 +2094,12 @@ TSharedPtr<SDockTab> FTabManager::SpawnTab(const FTabId& TabId, const TSharedPtr
 			}
 
 			NewTabWidget->SetLayoutIdentifier(TabId);
-			NewTabWidget->ProvideDefaultLabel(Spawner->GetDisplayName().IsEmpty() ? FText::FromName(Spawner->TabType) : Spawner->GetDisplayName());
+
+			const FText Label = GetTabLabelBasedOnSpawner(Spawner);
+
+			NewTabWidget->ProvideDefaultLabel(Label);
 			NewTabWidget->ProvideDefaultIcon(Spawner->GetIcon().GetIcon());
+			NewTabWidget->SetIsTabNameHidden(Spawner->IsTabNameHidden());
 
 			// The spawner tracks that last tab it spawned
 			Spawner->SpawnedTabPtr = NewTabWidget;
@@ -2112,8 +2118,9 @@ TSharedPtr<SDockTab> FTabManager::SpawnTab(const FTabId& TabId, const TSharedPtr
 		// 1) Make a dummy tab so that things aren't entirely broken (previous versions of UE did this in all cases).
 		// 2) Do not open the widget and return nullptr, but keep the unknown widget saved in the layout. E.g., applied when calling RestoreFrom() from MainFrameModule.
 
-		FString StringToDisplay = (Spawner.IsValid() && !Spawner->GetDisplayName().IsEmpty() ? Spawner->GetDisplayName().ToString() : TabId.TabType.ToString());
-		if (StringToDisplay.IsEmpty())
+		FString StringToDisplay = GetTabLabelBasedOnSpawner(Spawner).ToString();
+
+		if (StringToDisplay.IsEmpty() && !(Spawner.IsValid() && Spawner->IsTabNameHidden()))
 		{
 			StringToDisplay = FString("Unknown");
 		}
@@ -2294,6 +2301,23 @@ TSharedRef<FTabManager::FArea> FTabManager::GetAreaForTabId(const FTabId& TabId)
 	return NewArea( GetDefaultTabWindowSize(TabId) );
 }
 
+FText FTabManager::GetTabLabelBasedOnSpawner(const TSharedPtr<FTabSpawnerEntry>& InSpawnerEntry) const
+{
+	if (!InSpawnerEntry.IsValid())
+	{
+		return FText::GetEmpty();
+	}
+
+	FText Label = FText::GetEmpty();
+
+	if (!InSpawnerEntry->IsTabNameHidden())
+	{
+		Label = InSpawnerEntry->GetDisplayName().IsEmpty() ? FText::FromName(InSpawnerEntry->TabType) : InSpawnerEntry->GetDisplayName();
+	}
+
+	return Label;
+}
+
 void FGlobalTabmanager::SetInitialLayoutSP(TSharedPtr<FTabManager::FLayout> InLayout)
 {
 	InitialLayoutSP = InLayout;
diff --git a/Engine/Source/Runtime/Slate/Private/Widgets/Docking/SDockTab.cpp b/Engine/Source/Runtime/Slate/Private/Widgets/Docking/SDockTab.cpp
--- a/Engine/Source/Runtime/Slate/Private/Widgets/Docking/SDockTab.cpp
+++ b/Engine/Source/Runtime/Slate/Private/Widgets/Docking/SDockTab.cpp
@@ -651,7 +651,7 @@ void SDockTab::Construct( const FArguments& InArgs )
 				+ SHorizontalBox::Slot()
 				.AutoWidth()
 				.VAlign(VAlign_Center)
-				.Padding(0, 0, 5, 0)
+				.Padding(IsTabNameHidden() ? FMargin(0) : FMargin(0, 0, 5, 0))
 				[
 					SNew(SBorder)
 					// Don't allow active tab overlay to absorb mouse clicks
@@ -706,7 +706,7 @@ void SDockTab::Construct( const FArguments& InArgs )
 
 				+ SHorizontalBox::Slot()
 				.AutoWidth()
-				.HAlign(HAlign_Right)
+				.HAlign(IsTabNameHidden() ? HAlign_Left : HAlign_Right)
 				.VAlign(VAlign_Center)
 				[
 					SNew(SButton)
@@ -1083,4 +1083,14 @@ void SDockTab::SetParentDockTabStackTabWellHidden(bool bIsTabWellHidden)
 			ParentStack->SetTabWellHidden(false);
 		}
 	}
-}
\ No newline at end of file
+}
+
+bool SDockTab::IsTabNameHidden() const
+{
+	return bIsTabNameHidden;
+}
+
+void SDockTab::SetIsTabNameHidden(bool bInIsTabNameHidden)
+{
+	bIsTabNameHidden = bInIsTabNameHidden;
+}
diff --git a/Engine/Source/Runtime/Slate/Public/Framework/Docking/TabManager.h b/Engine/Source/Runtime/Slate/Public/Framework/Docking/TabManager.h
--- a/Engine/Source/Runtime/Slate/Public/Framework/Docking/TabManager.h
+++ b/Engine/Source/Runtime/Slate/Public/Framework/Docking/TabManager.h
@@ -225,10 +225,15 @@ struct FTabSpawnerEntry : public FWorkspaceItem
 		, MenuType(ETabSpawnerMenuType::Enabled)
 		, bAutoGenerateMenuEntry(true)
 		, bCanSidebarTab(true)
+		, bHiddenTabName(false)
 		, SpawnedTabPtr()
 		, ReadOnlyBehavior(ETabReadOnlyBehavior::Disabled)
-	
 	{
+		// Special treatment for the HomeScreen tab to not have its name displayed, this is currently intended to not be a public feature
+		if (InTabType == FName(TEXT("HomeScreen")))
+		{
+			bHiddenTabName = true;
+		}
 	}
 
 	FTabSpawnerEntry& SetReadOnlyBehavior( const ETabReadOnlyBehavior& InReadOnlyBehavior)
@@ -302,6 +307,11 @@ struct FTabSpawnerEntry : public FWorkspaceItem
 		return bCanSidebarTab;
 	}
 
+	bool IsTabNameHidden() const
+	{
+		return bHiddenTabName;
+	}
+
 	virtual TSharedPtr<FTabSpawnerEntry> AsSpawnerEntry() override
 	{
 		return SharedThis(this);
@@ -329,6 +339,7 @@ private:
 	bool bAutoGenerateMenuEntry;
 	/** Whether or not this tab can ever be in a sidebar */
 	bool bCanSidebarTab;
+	bool bHiddenTabName;
 
 	TWeakPtr<SDockTab> SpawnedTabPtr;
 
@@ -1198,6 +1209,9 @@ protected:
 	 */
 	TSharedRef<FArea> GetAreaForTabId(const FTabId& TabId);
 
+	/** Based on if the tab can have a hidden display name it will return the correct Label to use for the tab */
+	FText GetTabLabelBasedOnSpawner(const TSharedPtr<FTabSpawnerEntry>& InSpawnerEntry) const;
+
 	protected:
 		FTabSpawner TabSpawner;
 		TSharedRef<FTabSpawner> NomadTabSpawner;
diff --git a/Engine/Source/Runtime/Slate/Public/Widgets/Docking/SDockTab.h b/Engine/Source/Runtime/Slate/Public/Widgets/Docking/SDockTab.h
--- a/Engine/Source/Runtime/Slate/Public/Widgets/Docking/SDockTab.h
+++ b/Engine/Source/Runtime/Slate/Public/Widgets/Docking/SDockTab.h
@@ -342,6 +342,9 @@ public:
 
 	SLATE_API void SetParentDockTabStackTabWellHidden(bool bIsTabWellHidden);
 
+	/** Get if the TabName should be hidden and not take up space for this Tab */
+	SLATE_API bool IsTabNameHidden() const;
+
 protected:
 	/** Provide a default tab icon. */
 	SLATE_API void ProvideDefaultIcon(const FSlateBrush* InDefaultIcon);
@@ -392,6 +395,9 @@ protected:
 	SLATE_API FMargin GetTabIconBorderPadding() const;
 
 private:
+	/** Set whether the TabName should be hidden and not take up space for this Tab */
+	SLATE_API void SetIsTabNameHidden(bool bInIsTabNameHidden);
+
 	/** Activates the tab in its tab well */
 	SLATE_API EActiveTimerReturnType TriggerActivateTab( double InCurrentTime, float InDeltaTime );
 
@@ -412,6 +418,9 @@ private:
 	/** The handle to the active tab activation tick */
 	TWeakPtr<FActiveTimerHandle> DragDropTimerHandle;
 	TWeakPtr<FActiveTimerHandle> UpdateStyleTimerHandle;
+
+	/** If true the TabName will be hidden and not take up space for this tab */
+	bool bIsTabNameHidden = false;
 protected:
 
 	/** The tab manager that created this tab. */
