Adding the panel drawer.
The panel drawer is a section of layout that can host a tab without disrupting the current layout of the opened editor.
If the hosted tab is nomad, changing between the major tabs will keep the pannel open. Keeping user focus in their workflow without needing to move the tab along with them.
The feature is still experimental and should be available in the vast majority of the major tabs in the editor (Every tab that has a status bar).
Minor change how a tab is selected when closing the current tab, we select the tab to the right by default,  otherwise the one to left if there none to the right.
Minor change to how the broadcast for the foregrounding and brackgrounded tab work when closing or opening a new tab. This will now behave more as expected.

Note: Some features are still upcoming for this in the near term (persistence and the ability to tear off the tab outside of the drawer)

diff --git a/Engine/Source/Editor/LevelEditor/Private/SLevelEditor.cpp b/Engine/Source/Editor/LevelEditor/Private/SLevelEditor.cpp
--- a/Engine/Source/Editor/LevelEditor/Private/SLevelEditor.cpp
+++ b/Engine/Source/Editor/LevelEditor/Private/SLevelEditor.cpp
@@ -365,31 +365,12 @@ void SLevelEditor::Initialize( const TSharedRef<SDockTab>& OwnerTab, const TShar
 
 	FLevelEditorMenu::MakeLevelEditorMenu(LevelEditorCommands, SharedThis(this));
 
-	SecondaryModeToolbarWidget =
-		SNew(SBorder)
-		.Padding(0.f)
-		.BorderImage(FAppStyle::Get().GetBrush( "NoBorder" ));
-	
-	SecondaryModeToolbarWidget->SetContent(FLevelEditorToolBar::MakeLevelEditorSecondaryModeToolbar(LevelEditorCommands.ToSharedRef(), ModeUILayers));
-		
+
 	ChildSlot
 	[
 		SNew(SVerticalBox)
 		+SVerticalBox::Slot()
-		.Padding(FMargin(0.0f,0.0f,0.0f,2.0f))
-		.AutoHeight()
-		[
-			FLevelEditorToolBar::MakeLevelEditorToolBar(LevelEditorCommands.ToSharedRef(), SharedThis(this))
-		]
-		+SVerticalBox::Slot()
-		.Padding(FMargin(0.0f,0.0f,0.0f,2.0f))
-		.AutoHeight()
-		[
-			SecondaryModeToolbarWidget.ToSharedRef()
-		]
-		+SVerticalBox::Slot()
-		.Padding(4.0f, 2.f, 4.f, 2.f)
-		.FillHeight( 1.0f )
+		.FillHeight(1.0f)
 		[
 			ContentArea
 		]
@@ -1766,7 +1747,40 @@ TSharedRef<SWidget> SLevelEditor::RestoreContentArea( const TSharedRef<SDockTab>
 		ensureMsgf(false, TEXT("Some additional testing of that layout file should be done."));
 	}
 	check(ContentAreaWidget.IsValid());
-	return ContentAreaWidget.ToSharedRef();
+
+	SecondaryModeToolbarWidget =
+		SNew(SBorder)
+		.Padding(0.f)
+		.BorderImage(FAppStyle::Get().GetBrush("NoBorder"));
+
+	SecondaryModeToolbarWidget->SetContent(FLevelEditorToolBar::MakeLevelEditorSecondaryModeToolbar(LevelEditorCommands.ToSharedRef(), ModeUILayers));
+
+	TSharedRef<SWidget> LevelEditorContent = SNew(SVerticalBox)
+		+ SVerticalBox::Slot()
+		.Padding(FMargin(0.0f, 0.0f, 0.0f, 2.0f))
+		.AutoHeight()
+		[
+			FLevelEditorToolBar::MakeLevelEditorToolBar(LevelEditorCommands.ToSharedRef(), SharedThis(this))
+		]
+		+ SVerticalBox::Slot()
+		.Padding(FMargin(0.0f, 0.0f, 0.0f, 2.0f))
+		.AutoHeight()
+		[
+			SecondaryModeToolbarWidget.ToSharedRef()
+		]
+		+ SVerticalBox::Slot()
+		.Padding(4.0f, 2.f, 4.f, 2.f)
+		.FillHeight(1.0f)
+		[
+			ContentAreaWidget.ToSharedRef()
+		];
+
+	if (TSharedPtr<SWidget> PanelDrawer = LevelEditorTabManager->RestorePanelDrawerFrom(DefaultLayout, LevelEditorContent, OwnerWindow))
+	{
+		LevelEditorContent = PanelDrawer.ToSharedRef();
+	}
+
+	return LevelEditorContent;
 }
 
 void SLevelEditor::HandleExperimentalSettingChanged(FName PropertyName)
diff --git a/Engine/Source/Editor/UnrealEd/Private/Toolkits/SStandaloneAssetEditorToolkitHost.cpp b/Engine/Source/Editor/UnrealEd/Private/Toolkits/SStandaloneAssetEditorToolkitHost.cpp
--- a/Engine/Source/Editor/UnrealEd/Private/Toolkits/SStandaloneAssetEditorToolkitHost.cpp
+++ b/Engine/Source/Editor/UnrealEd/Private/Toolkits/SStandaloneAssetEditorToolkitHost.cpp
@@ -196,25 +196,37 @@ void SStandaloneAssetEditorToolkitHost::RestoreFromLayout( const TSharedRef<FTab
 		StatusBarCreationOptions = HostedAssetEditorToolkit->GetStatusBarCreationOptions();
 	}
 
+	TSharedRef<SWidget> AssetEditorContent = SNew(SVerticalBox)
+		+ SVerticalBox::Slot()
+		.Padding(FMargin(0.0f, 0.0f, 0.0f, 2.0f))
+		.AutoHeight()
+		.Expose(ToolbarSlot)
+		+ SVerticalBox::Slot()
+		.Padding(4.f, 2.f, 4.f, 2.f)
+		.FillHeight(1.0f)
+		[
+			RestoredUI.ToSharedRef()
+		];
+
 	if (UStatusBarSubsystem* StatusBarSubsystem = GEditor ? GEditor->GetEditorSubsystem<UStatusBarSubsystem>() : nullptr)
 	{
 		StatusBarWidget = StatusBarSubsystem->MakeStatusBarWidget(StatusBarName, HostTab, StatusBarCreationOptions);
 	}
 
+	if (TSharedPtr<SWidget> PanelDrawer = MyTabManager->RestorePanelDrawerFrom(NewLayout, AssetEditorContent, ParentWindow))
+	{
+		AssetEditorContent = PanelDrawer.ToSharedRef();
+	}
+
 	this->ChildSlot
 	[
 		SNew(SVerticalBox)
 		+ SVerticalBox::Slot()
-		.Padding(FMargin(0.0f, 0.0f, 0.0f, 2.0f))
-		.AutoHeight()
-		.Expose(ToolbarSlot)
-		+ SVerticalBox::Slot()
-		.Padding(4.f, 2.f, 4.f, 2.f)
-		.FillHeight(1.0f)
+		.FillHeight(1.f)
 		[
-			RestoredUI.ToSharedRef()
+			AssetEditorContent
 		]
-	+ SVerticalBox::Slot()
+		+ SVerticalBox::Slot()
 		.Padding(0.0f, 2.0f, 0.0f, 0.0f)
 		.AutoHeight()
 		[
diff --git a/Engine/Source/Runtime/Slate/Private/Framework/Docking/SDockingTabWell.cpp b/Engine/Source/Runtime/Slate/Private/Framework/Docking/SDockingTabWell.cpp
--- a/Engine/Source/Runtime/Slate/Private/Framework/Docking/SDockingTabWell.cpp
+++ b/Engine/Source/Runtime/Slate/Private/Framework/Docking/SDockingTabWell.cpp
@@ -63,6 +63,13 @@ void SDockingTabWell::AddTab( const TSharedRef<SDockTab>& InTab, int32 AtIndex,
 	else
 	{
 		AtIndex = FMath::Clamp( AtIndex, 0, Tabs.Num() );
+
+		if (AtIndex <= ForegroundTabIndex)
+		{
+			// Update the currently active index (Otherwise we don't broadcast the right info)
+			++ForegroundTabIndex;
+		}
+
 		this->Tabs.Insert( InTab, AtIndex );
 
 		if (!bKeepInactive)
@@ -690,22 +697,37 @@ void SDockingTabWell::RemoveAndDestroyTab(const TSharedRef<SDockTab>& TabToRemov
 
 		// Remove the old tab from the list of tabs and activate the new tab.
 		{
-			int32 OldTabIndex = ForegroundTabIndex;
-
+			int32 OldTabIndex = FMath::Max(ForegroundTabIndex, 0);
+	
+			// Not sure why but we always bring the tab that is about be remove to the foreground before removing it
 			BringTabToFront(TabIndex);
-			Tabs.RemoveAt(TabIndex);
-			// We no longer have a tab in the foreground.
-			// This is important because BringTabToFront triggers notifications based on the difference in active tab indexes.
-			ForegroundTabIndex = INDEX_NONE;
-
-			// Now bring the last tab that we were on to the foreground
-			if (OldTabIndex == INDEX_NONE || TabIndex <= OldTabIndex)
+			
+			// The tab that will be removed is the same that was selected before
+			if (TabIndex == OldTabIndex)
 			{
-				BringTabToFront(FMath::Max(OldTabIndex - 1, 0));
+				if (OldTabIndex == Tabs.Num() - 1)
+				{
+					// Select the previous tab
+					OldTabIndex = FMath::Max(OldTabIndex - 1, 0);
+
+				}
+				else
+				{
+					// Select the next tab
+					++OldTabIndex;
+				}
 			}
-			else
+
+			// Allow the transfer from old to new foreground tab
+			BringTabToFront(OldTabIndex);
+			
+			// Actually remove the tab
+			Tabs.RemoveAt(TabIndex);
+
+			// Update the selected tab index if needed
+			if (TabIndex <= ForegroundTabIndex)
 			{
-				BringTabToFront(OldTabIndex);
+				--ForegroundTabIndex;
 			}
 
 			if (RemovalMethod == SDockingNode::ELayoutModification::TabRemoval_Sidebar && ForegroundTabIndex == INDEX_NONE)
diff --git a/Engine/Source/Runtime/Slate/Private/Framework/Docking/SPanelDrawerArea.cpp b/Engine/Source/Runtime/Slate/Private/Framework/Docking/SPanelDrawerArea.cpp
new file mode 100644
--- /dev/null
+++ b/Engine/Source/Runtime/Slate/Private/Framework/Docking/SPanelDrawerArea.cpp
@@ -0,0 +1,340 @@
+// Copyright Epic Games, Inc. All Rights Reserved.
+
+#include "Framework/Docking/SPanelDrawerArea.h"
+
+#include "Containers/Array.h"
+#include "Layout/Geometry.h"
+#include "Math/UnrealMathUtility.h"
+#include "Templates/SharedPointer.h"
+#include "Widgets/Docking/SDockTab.h"
+#include "Widgets/Layout/SBox.h"
+#include "Widgets/Layout/SConstraintCanvas.h"
+#include "Widgets/Layout/SSplitter.h"
+#include "Widgets/SBoxPanel.h"
+#include "Widgets/SNullWidget.h"
+
+namespace UE::Slate::Private
+{
+constexpr float SSplitterHandleWith = 5.f;
+constexpr float AnimationLength = 0.15f;
+
+void SPanelDrawerArea::Construct(const FArguments& InArgs, const TSharedRef<SWidget>& InContent)
+{
+	AreaContent = InContent;
+
+	OpenCloseAnimation = FCurveSequence(0.f, AnimationLength, ECurveEaseFunction::QuadOut);
+
+	// Add Restoring logic here
+	SetupClosedLayout();
+}
+
+SPanelDrawerArea::~SPanelDrawerArea()
+{
+	FSlateThrottleManager::Get().LeaveResponsiveMode(AnimationThrottle);
+}
+
+void SPanelDrawerArea::Tick(const FGeometry& AllottedGeometry, const double InCurrentTime, const float InDeltaTime)
+{
+	// Compute the desired final size of the panel (Always cached to avoid issue jumping issue)
+	DesiredOpenWidth = (AllottedGeometry.GetLocalSize().X - SSplitterHandleWith) * Size.OpenPanelDrawerSize / (Size.ContentSlotSize + Size.OpenPanelDrawerSize);
+
+	if (OpenCloseAnimation.IsPlaying())
+	{
+		UpdateAnimatedSlideWidth();
+	}
+}
+
+void SPanelDrawerArea::TogglePanel(TSharedPtr<SDockTab> InDockTab, bool bWithAnimation)
+{
+	HostedTab = InDockTab;
+	if (bIsOpen)
+	{
+		ClosePanel(bWithAnimation);
+	}
+	else
+	{
+		if (HostedTab)
+		{
+			OpenPanel(bWithAnimation);
+		}
+	}
+}
+
+void SPanelDrawerArea::OpenPanel(bool bWithAnimation)
+{
+	if (bWithAnimation)
+	{
+		constexpr bool bIsOpening = true;
+		PlayAnimation(bIsOpening);
+	}
+	else
+	{
+		SetupOpenedLayout();
+	}
+
+	bIsOpen = true;
+}
+
+void SPanelDrawerArea::ClosePanel(bool bWithAnimation)
+{
+	if (bWithAnimation)
+	{
+		constexpr bool bIsOpening = false;
+		PlayAnimation(bIsOpening);
+	}
+	else
+	{
+		SetupClosedLayout();
+	}
+
+	bIsOpen = false;
+}
+
+void SPanelDrawerArea::SetupAnimationLayout()
+{
+	ChildSlot
+	[
+		SNew(SSplitter)
+		.Orientation(Orient_Horizontal)
+		+ SSplitter::Slot()
+		.Value(bIsOpen ? 1.f : Size.ContentSlotSize)
+		[
+			SNew(SHorizontalBox)
+			+SHorizontalBox::Slot()
+			.FillWidth(1.f)
+			[
+				// Add Clipping for content that isn't responsive to dynamic sizes
+				SNew(SBox)
+				.HAlign(HAlign_Fill)
+				.VAlign(VAlign_Fill)
+				.Clipping(EWidgetClipping::OnDemand)
+				[
+					AreaContent.ToSharedRef()
+				]
+			]
+			+SHorizontalBox::Slot()
+			.AutoWidth()
+			[
+				SNew(SBox)
+				.VAlign(VAlign_Fill)
+				.HAlign(HAlign_Right)
+				.WidthOverride(this, &SPanelDrawerArea::GetAnimatedWidthOverrideForSpacer)
+				.Visibility(this, &SPanelDrawerArea::GetAnimatedSpacerVisibility)
+			]
+		]
+		+ SSplitter::Slot()
+		.Value(bIsOpen ? 0.0f : Size.OpenPanelDrawerSize)
+		.SizeRule(SSplitter::ESizeRule::SizeToContent)
+		[
+			SNew(SBox)
+			.HAlign(HAlign_Left)
+			.VAlign(VAlign_Fill)
+			.WidthOverride(this, &SPanelDrawerArea::GetAnimatedWidthOverrideForPanelDrawer)
+			.Visibility(this, &SPanelDrawerArea::GetAnimatedDrawerPanelVisibility)
+			[
+				SNew(SConstraintCanvas)
+				+ SConstraintCanvas::Slot()
+				.Anchors(FAnchors(0.f, 0.f, 0.f, 1.f))
+				.Offset(FMargin(0.f, 0.f, DesiredOpenWidth, 0.f))
+				.Alignment(FVector2D(0.f, 0.f))
+				[
+					HostedTab ? HostedTab->GetContent() : SNullWidget::NullWidget
+				]
+			]
+		]
+	];
+}
+
+void SPanelDrawerArea::SetupClosedLayout()
+{
+	ChildSlot
+	[
+		AreaContent.ToSharedRef()
+	];
+}
+
+void SPanelDrawerArea::SetupOpenedLayout()
+{
+	ChildSlot
+	[
+		SNew(SSplitter)
+		.Orientation(Orient_Horizontal)
+		+ SSplitter::Slot()
+		.OnSlotResized(this, &SPanelDrawerArea::OnContentSlotResized)
+		.Value(this, &SPanelDrawerArea::GetContentSlotSize)
+		[
+			// Add Clipping for content that isn't responsive to dynamic sizes
+			SNew(SBox)
+			.HAlign(HAlign_Fill)
+			.VAlign(VAlign_Fill)
+			.Clipping(EWidgetClipping::OnDemand)
+			[
+				AreaContent.ToSharedRef()
+			]
+		]
+		+ SSplitter::Slot()
+		.OnSlotResized(this, &SPanelDrawerArea::OnPanelDrawerResized)
+		.Value(this, &SPanelDrawerArea::GetPanelDrawerSize)
+		[
+			HostedTab ? HostedTab->GetContent() : SNullWidget::NullWidget
+		]
+	];
+}
+
+void SPanelDrawerArea::OnContentSlotResized(float InContentSlotSize)
+{
+	Size.ContentSlotSize = InContentSlotSize;
+}
+
+void SPanelDrawerArea::OnPanelDrawerResized(float InOpenPanelDrawerSize)
+{
+	Size.OpenPanelDrawerSize = InOpenPanelDrawerSize;
+}
+
+float SPanelDrawerArea::GetContentSlotSize() const
+{
+	return Size.ContentSlotSize;
+}
+
+float SPanelDrawerArea::GetPanelDrawerSize() const
+{
+	return Size.OpenPanelDrawerSize;
+}
+
+void SPanelDrawerArea::PlayAnimation(bool bIsOpening)
+{
+	if (OpenCloseAnimation.IsPlaying())
+	{
+		OpenCloseAnimation.PlayRelative(AsShared(), bIsOpening);
+	}
+	else
+	{
+		constexpr bool bIsLooped = false;
+		constexpr float StartTime = 0.f;
+		constexpr bool bRequireActiveTimer = false;
+		if (bIsOpening)
+		{
+			OpenCloseAnimation.Play(AsShared(), bIsLooped, StartTime, bRequireActiveTimer);
+		}
+		else
+		{
+			OpenCloseAnimation.PlayReverse(AsShared(), bIsLooped, StartTime, bRequireActiveTimer);
+		}
+
+		UpdateAnimatedSlideWidth();
+		SetupAnimationLayout();
+		AnimationThrottle = FSlateThrottleManager::Get().EnterResponsiveMode();
+		OpenCloseTimer = RegisterActiveTimer(0.f, FWidgetActiveTimerDelegate::CreateSP(this, &SPanelDrawerArea::UpdateAnimation));
+	}
+}
+
+EActiveTimerReturnType SPanelDrawerArea::UpdateAnimation(const double InCurrentTime, const float InDeltaTime)
+{
+	if (!OpenCloseAnimation.IsPlaying())
+	{
+		if (bIsOpen)
+		{
+			SetupOpenedLayout();
+		}
+		else
+		{
+			SetupClosedLayout();
+		}
+
+		FSlateThrottleManager::Get().LeaveResponsiveMode(AnimationThrottle);
+		OpenCloseTimer.Reset();
+
+		return EActiveTimerReturnType::Stop;
+	}
+
+	return EActiveTimerReturnType::Continue;
+}
+
+void SPanelDrawerArea::UpdateAnimatedSlideWidth()
+{
+	AnimatedSlideWidth = FMath::Lerp(0.f, DesiredOpenWidth + SSplitterHandleWith, OpenCloseAnimation.GetLerp());
+}
+
+FOptionalSize SPanelDrawerArea::GetAnimatedWidthOverrideForSpacer() const
+{
+	return FMath::Min(5.f, AnimatedSlideWidth);
+}
+
+FOptionalSize SPanelDrawerArea::GetAnimatedWidthOverrideForPanelDrawer() const
+{
+	return FMath::Max(0.f, AnimatedSlideWidth - 5.f);
+}
+
+EVisibility SPanelDrawerArea::GetAnimatedSpacerVisibility() const
+{
+	return AnimatedSlideWidth >= 5.f ? EVisibility::Collapsed : EVisibility::Visible;
+}
+
+EVisibility SPanelDrawerArea::GetAnimatedDrawerPanelVisibility() const
+{
+	return AnimatedSlideWidth < 5.f ? EVisibility::Collapsed : EVisibility::Visible;
+}
+
+bool SPanelDrawerArea::IsOpen() const
+{
+	return bIsOpen;
+}
+
+bool SPanelDrawerArea::IsHostingTab(TSharedPtr<SDockTab> InDockTab) const
+{
+	return InDockTab == HostedTab;
+}
+
+void SPanelDrawerArea::SetHostedTab(TSharedPtr<SDockTab> InDockTab)
+{
+	HostedTab = InDockTab;
+}
+
+TSharedPtr<SDockTab> SPanelDrawerArea::GetHostedTab() const
+{
+	return HostedTab;
+}
+
+const FPanelDrawerSize SPanelDrawerArea::GetSize() const
+{
+	return Size;
+}
+
+void SPanelDrawerArea::SetSize(const FPanelDrawerSize& InSize)
+{
+	Size = InSize;
+}
+
+FPanelDrawerData SPanelDrawerArea::TakeOwnershipOfHostedTab(bool bWithAnimation)
+{
+	FPanelDrawerData HostedData{ Size, bIsOpen, MoveTemp(HostedTab) };
+	if (bIsOpen)
+	{
+		ClosePanel(bWithAnimation);
+	}
+	return HostedData;
+}
+
+void SPanelDrawerArea::TransferOwnershipOfHostedTab(FPanelDrawerData&& InDrawerData, bool bWithAnimation)
+{
+	Size = InDrawerData.Size;
+	HostedTab = MoveTemp(InDrawerData.HostedTab);
+
+	if (InDrawerData.bIsOpen)
+	{
+		OpenPanel(bWithAnimation);
+
+	}
+}
+
+bool SPanelDrawerArea::IsHostingNomadTab() const
+{
+	if (HostedTab)
+	{
+		return HostedTab->GetTabRole() == ETabRole::NomadTab;
+	}
+
+	return false;
+}
+
+}
diff --git a/Engine/Source/Runtime/Slate/Private/Framework/Docking/SPanelDrawerArea.h b/Engine/Source/Runtime/Slate/Private/Framework/Docking/SPanelDrawerArea.h
new file mode 100644
--- /dev/null
+++ b/Engine/Source/Runtime/Slate/Private/Framework/Docking/SPanelDrawerArea.h
@@ -0,0 +1,99 @@
+// Copyright Epic Games, Inc. All Rights Reserved.
+
+#pragma once
+
+#include "Animation/CurveSequence.h"
+#include "Widgets/SCompoundWidget.h"
+#include "Application/ThrottleManager.h"
+
+class SDockTab;
+
+namespace UE::Slate::Private
+{
+
+struct FPanelDrawerSize
+{
+	float ContentSlotSize = 0.8f;
+	float OpenPanelDrawerSize = 0.2f;
+};
+
+struct FPanelDrawerData
+{
+	FPanelDrawerSize Size;
+	bool bIsOpen = false;
+	TSharedPtr<SDockTab> HostedTab;
+};
+
+/**
+ * A widget that hold the content where a panel drawer can be invoked.
+ */
+class SPanelDrawerArea : public SCompoundWidget
+{
+public:
+	SLATE_BEGIN_ARGS(SPanelDrawerArea) {}
+	SLATE_END_ARGS()
+
+	void Construct(const FArguments& InArgs, const TSharedRef<SWidget>& InContent);
+
+	virtual ~SPanelDrawerArea();
+
+	virtual void Tick(const FGeometry& AllottedGeometry, const double InCurrentTime, const float InDeltaTime) override;
+
+	void TogglePanel(TSharedPtr<SDockTab> InDockTab, bool bWithAnimation);
+
+	void OpenPanel(bool bWithAnimation);
+	void ClosePanel(bool bWithAnimation);
+
+	bool IsOpen() const;
+	bool IsHostingNomadTab() const;
+	bool IsHostingTab(TSharedPtr<SDockTab> InDockTab) const;
+	void SetHostedTab(TSharedPtr<SDockTab> InDockTab);
+	TSharedPtr<SDockTab> GetHostedTab() const;
+
+	const FPanelDrawerSize GetSize() const;
+	void SetSize(const FPanelDrawerSize& InSize);
+
+	FPanelDrawerData TakeOwnershipOfHostedTab(bool bWithAnimation);
+	void TransferOwnershipOfHostedTab(FPanelDrawerData&& InDrawerData, bool bWithAnimation);
+
+private:
+	void SetupAnimationLayout();
+
+	void SetupOpenedLayout();
+	void SetupClosedLayout();
+
+	void OnContentSlotResized(float InContentSlotSize);
+	void OnPanelDrawerResized(float InOpenPanelDrawerSize);
+
+	float GetContentSlotSize() const;
+	float GetPanelDrawerSize() const;
+
+	void PlayAnimation(bool bIsOpening);
+	EActiveTimerReturnType UpdateAnimation(const double InCurrentTime, const float InDeltaTime);
+	void UpdateAnimatedSlideWidth();
+
+	FOptionalSize GetAnimatedWidthOverrideForSpacer() const;
+	FOptionalSize GetAnimatedWidthOverrideForPanelDrawer() const;
+	EVisibility GetAnimatedSpacerVisibility() const;
+	EVisibility GetAnimatedDrawerPanelVisibility() const;
+
+
+	bool bIsOpen = false;
+
+	FPanelDrawerSize Size;
+
+	TSharedPtr<SWidget> AreaContent;
+	TSharedPtr<SDockTab> HostedTab;
+
+
+	// Animation Data
+	bool bRequestedAnimation = false;
+	FCurveSequence OpenCloseAnimation;
+	FThrottleRequest AnimationThrottle;
+	TSharedPtr<FActiveTimerHandle> OpenCloseTimer;
+
+	float DesiredOpenWidth = 0.f;
+	// The animated include also sliding in of the splitter handle
+	float AnimatedSlideWidth = 0.f;
+};
+}
diff --git a/Engine/Source/Runtime/Slate/Private/Framework/Docking/TabManager.cpp b/Engine/Source/Runtime/Slate/Private/Framework/Docking/TabManager.cpp
--- a/Engine/Source/Runtime/Slate/Private/Framework/Docking/TabManager.cpp
+++ b/Engine/Source/Runtime/Slate/Private/Framework/Docking/TabManager.cpp
@@ -17,10 +17,13 @@
 #include "Widgets/Docking/SDockTab.h"
 #include "Framework/Docking/SDockingTabStack.h"
 #include "Framework/Docking/SDockingTabWell.h"
+#include "Framework/Docking/SPanelDrawerArea.h"
 #include "Framework/Docking/LayoutExtender.h"
 #include "Misc/NamePermissionList.h"
 #include "Trace/SlateMemoryTags.h"
+#include "HAL/IConsoleManager.h"
 #include "HAL/PlatformApplicationMisc.h"
+
 #if PLATFORM_MAC
 #include "Framework/MultiBox/Mac/MacMenu.h"
 #endif
@@ -31,6 +34,12 @@ TMap<FTabId, FVector2D> FTabManager::DefaultTabWindowSizeMap;
 
 DEFINE_LOG_CATEGORY_STATIC(LogTabManager, Display, All);
 
+TAutoConsoleVariable<bool> CVarPanelDrawerToggle(
+	TEXT("EnablePanelDrawer"),
+	true,
+	TEXT("Enables or disables the panel drawer. When the drawer is off invocations will just create a normal tab.")
+);
+
 #define LOCTEXT_NAMESPACE "TabManager"
 
 static const FString UE_TABMANAGER_OPENED_TAB_STRING = TEXT("OpenedTab");
@@ -698,6 +707,8 @@ void FTabManager::FPrivateApi::OnDockAreaClosing( const TSharedRef<SDockingArea>
 	{
 		TabManager.CollapsedDockAreas.Add( PersistentDockAreaLayout.ToSharedRef() );
 	}
+
+	TabManager.QueuedPanelDrawerTransfers.Remove(DockAreaThatIsClosing->GetParentWindow());
 }
 
 void FTabManager::FPrivateApi::OnTabManagerClosing()
@@ -911,6 +922,8 @@ void FTabManager::CloseAllAreas()
 
 	CollapsedDockAreas.Empty();
 	InvalidDockAreas.Empty();
+
+	QueuedPanelDrawerTransfers.Empty();
 }
 
 
@@ -1084,6 +1097,21 @@ TSharedPtr<SWidget> FTabManager::RestoreFrom(const TSharedRef<FLayout>& Layout,
 	return PrimaryDockArea;
 }
 
+TSharedPtr<SWidget> FTabManager::RestorePanelDrawerFrom(const TSharedRef<FLayout>& Layout, const TSharedRef<SWidget>& InContent, const TSharedPtr<SWindow>& ParentWindow)
+{
+	TSharedRef<UE::Slate::Private::SPanelDrawerArea> NewPanelDrawerArea = SNew(UE::Slate::Private::SPanelDrawerArea, InContent);
+	PanelDrawerArea = NewPanelDrawerArea;
+
+	// Use a queued transfer if we have one
+	TUniquePtr<UE::Slate::Private::FPanelDrawerData> PtrToPreviousPanelData = GetQueuedNomadDrawerPanelForWindow(ParentWindow);
+	if (PtrToPreviousPanelData)
+	{
+		TransferNomadPanelDrawer(MoveTemp(*PtrToPreviousPanelData.Get()));
+	}
+
+	return MoveTemp(NewPanelDrawerArea);
+}
+
 struct FPopulateTabSpawnerMenu_Args
 {
 	FPopulateTabSpawnerMenu_Args( const TSharedRef< TArray< TWeakPtr<FTabSpawnerEntry> > >& InAllSpawners, const TSharedRef<FWorkspaceItem>& InMenuNode, int32 InLevel )
@@ -1308,6 +1336,13 @@ void FTabManager::DrawAttention( const TSharedRef<SDockTab>& TabToHighlight )
 
 		FGlobalTabmanager::Get()->UpdateMainMenu(TabToHighlight, true);
 	}
+	else if (const TSharedPtr<FTabManager> ManagerOfTabToHighlight = TabToHighlight->GetTabManagerPtr())
+	{
+
+
+		FGlobalTabmanager::Get()->DrawAttentionToTabManager(ManagerOfTabToHighlight.ToSharedRef());
+		TabToHighlight->FlashTab();
+	}
 }
 
 void FTabManager::InsertNewDocumentTab(FName PlaceholderId, FName NewTabId, const FSearchPreference& SearchPreference, const TSharedRef<SDockTab>& UnmanagedTab)
@@ -1387,41 +1422,24 @@ TSharedPtr<SDockTab> FTabManager::TryInvokeTab(const FTabId& TabId, bool bInvoke
 	return NewTab;
 }
 
-TSharedPtr<SDockTab> FTabManager::InvokeTab_Internal(const FTabId& TabId, bool bInvokeAsInactive, bool bForceOpenWindowIfNeeded)
+bool FTabManager::InvokeTab_SharedInternal(const FTabId& TabId, bool bInvokeAsInactive, TSharedPtr<SDockTab>& OutExistingTab)
 {
-	// Tab Spawning Rules:
-	// 
-	//     * Find live instance --yes--> use it.
-	//         |no
-	//         v
-	//     * [non-Document only]
-	//       Find closed instance with matching TabId --yes--> restore it.	  
-	//         |no
-	//         v
-	//     * Find any tab of matching TabType (closed or open) --yes--> spawn next to it.
-	//         | no
-	//         v
-	//     * Is a nomad tab and we are NOT the global tab manager --yes--> try to invoke in the global tab manager
-	//         | no
-	//         v
-	//     * Spawn in a new window.
-
 	if (!IsAllowedTab(TabId))
 	{
 		UE_LOG(LogTabManager, Warning, TEXT("Cannot spawn tab for '%s'"), *(TabId.ToString()));
-		return nullptr;
+		return false;
 	}
 
 	TSharedPtr<FTabSpawnerEntry> Spawner = FindTabSpawnerFor(TabId.TabType);
 
-	if ( !Spawner.IsValid() )
+	if (!Spawner.IsValid())
 	{
 		UE_LOG(LogTabManager, Warning, TEXT("Cannot spawn tab because no spawner is registered for '%s'"), *(TabId.ToString()));
 	}
 	else
-	{	
+	{
 		TSharedPtr<SDockTab> ExistingTab = Spawner->OnFindTabToReuse.IsBound()
-			? Spawner->OnFindTabToReuse.Execute( TabId )
+			? Spawner->OnFindTabToReuse.Execute(TabId)
 			: Spawner->SpawnedTabPtr.Pin();
 
 		if (ExistingTab.IsValid())
@@ -1436,13 +1454,76 @@ TSharedPtr<SDockTab> FTabManager::InvokeTab_Internal(const FTabId& TabId, bool b
 			// 1. Tab is not active
 			// 2. Tab's owning major tab is not in the foreground (making the tab we want to draw attention to is not visible)
 			// 3. Tab is nomad and is not in the foreground
+			// 4. Tab is not in closed DrawerPanel
 			// If the tab is not active or the tabs major tab is not in the foreground, activate it
 			if (!bInvokeAsInactive && (!ExistingTab->IsActive() || (MajorTab && !MajorTab->IsForeground()) || !ExistingTab->IsForeground()))
 			{
-				// Draw attention to this tab if it didn't already have focus
-				DrawAttention(ExistingTab.ToSharedRef());
+				if (TSharedPtr<UE::Slate::Private::SPanelDrawerArea> PanelDrawer = PanelDrawerArea.Pin())
+				{
+					if (!PanelDrawer->IsHostingTab(ExistingTab) || PanelDrawer->IsOpen())
+					{
+						// Draw attention to this tab if it didn't already have focus
+						DrawAttention(ExistingTab.ToSharedRef());
+					}
+				}
+				else
+				{
+					// Draw attention to this tab if it didn't already have focus
+					DrawAttention(ExistingTab.ToSharedRef());
+				}
+			}
+			OutExistingTab = ExistingTab;
+		}
+	}
+
+	return true;
+
+}
+
+TSharedPtr<SDockTab> FTabManager::InvokeTab_Internal(const FTabId& TabId, bool bInvokeAsInactive, bool bForceOpenWindowIfNeeded)
+{
+	// Tab Spawning Rules:
+	// 
+	//     * Find live instance --yes--> use it.
+	//         |no
+	//         v
+	//     * [non-Document only]
+	//       Find closed instance with matching TabId --yes--> restore it.	  
+	//         |no
+	//         v
+	//     * Find any tab of matching TabType (closed or open) --yes--> spawn next to it.
+	//         | no
+	//         v
+	//     * Is a nomad tab and we are NOT the global tab manager --yes--> try to invoke in the global tab manager
+	//         | no
+	//         v
+	//     * Spawn in a new window.
+
+	{
+		TSharedPtr<SDockTab> ExistingTab;
+		if (!InvokeTab_SharedInternal(TabId, bInvokeAsInactive, ExistingTab))
+		{
+			return {};
+		}
+
+		if (ExistingTab.IsValid())
+		{
+			bool bShouldReturnExistingTab = true;
+
+			if (TSharedPtr<FTabManager> ExistingTabManager = ExistingTab->GetTabManagerPtr())
+			{
+				// Will destroy the hidden drawer tab 
+				if (ExistingTabManager->TakeOwnerShipOfHiddenTabInDrawerPanel(ExistingTab))
+				{
+					// Let the tab die on purpose it will be recreated in a normal tab
+					bShouldReturnExistingTab = false;
+				}
+			}
+
+			if (bShouldReturnExistingTab)
+			{
+				return ExistingTab;
 			}
-			return ExistingTab.ToSharedRef();
 		}
 	}
 
@@ -2183,9 +2264,97 @@ TSharedPtr<SDockTab> FTabManager::FindExistingLiveTab( const FTabId& TabId ) con
 		}
 	}
 
+	if (TSharedPtr<UE::Slate::Private::SPanelDrawerArea> PanelDrawer = PanelDrawerArea.Pin())
+	{
+		if (PanelDrawer->IsOpen())
+		{
+			if (TSharedPtr<SDockTab> HostedTab = PanelDrawer->GetHostedTab())
+			{
+				if (TabId == HostedTab->GetLayoutIdentifier())
+				{
+					return HostedTab;
+				}
+			}
+		}
+	}
+
 	return TSharedPtr<SDockTab>();
 }
 
+TSharedPtr<SDockTab> FTabManager::TryToggleTabInPanelDrawer(const FTabId& TabId)
+{
+	TSharedPtr<UE::Slate::Private::SPanelDrawerArea> PanelDrawerAreaPtr = PanelDrawerArea.Pin();
+	if (CVarPanelDrawerToggle.GetValueOnGameThread() && PanelDrawerAreaPtr)
+	{
+		TSharedPtr<SDockTab> ExistingTab;
+		// Grab the existing without drawing attention
+		if (!InvokeTab_SharedInternal(TabId, true, ExistingTab))
+		{
+			return {};
+		}
+
+		if (ExistingTab.IsValid())
+		{
+			if (PanelDrawerAreaPtr->IsOpen() && PanelDrawerAreaPtr->IsHostingTab(ExistingTab))
+			{
+				// Close the open panel
+				PanelDrawerAreaPtr->ClosePanel(true);
+				return ExistingTab;
+			}
+
+			if (TSharedPtr<FTabManager> OtherTabManager = ExistingTab->GetTabManagerPtr())
+			{
+				bool bShouldOpenDrawer = false;
+				if (TSharedPtr<UE::Slate::Private::SPanelDrawerArea> OtherPanelDrawerAreaPtr = OtherTabManager->PanelDrawerArea.Pin())
+				{
+					if (OtherPanelDrawerAreaPtr->IsHostingTab(ExistingTab))
+					{
+						if (OtherPanelDrawerAreaPtr->IsOpen())
+						{
+							OtherPanelDrawerAreaPtr->SetHostedTab(TSharedPtr<SDockTab>());
+							OtherPanelDrawerAreaPtr->ClosePanel(true);
+							bShouldOpenDrawer = true;
+						}
+					}
+				}
+				if (OtherTabManager->TakeOwnerShipOfHiddenTabInDrawerPanel(ExistingTab))
+				{
+					bShouldOpenDrawer = true;
+				}
+
+				if (bShouldOpenDrawer)
+				{
+					// Open our panel
+					PanelDrawerAreaPtr->TogglePanel(ExistingTab, true);
+
+					ExistingTab->SetTabManager(AsShared());
+					return ExistingTab;
+				}
+			}
+
+			// Draw attention to the tab
+			InvokeTab_SharedInternal(TabId, false, ExistingTab);
+			return ExistingTab;
+		}
+
+		
+		const TSharedPtr<SDockTab> NewTab = SpawnTab(TabId, GetPrivateApi().GetParentWindow());
+
+		if (NewTab.IsValid())
+		{
+			PanelDrawerAreaPtr->TogglePanel(NewTab, true);
+		}
+
+		return NewTab;
+		
+	}
+
+	// If no drawer is available fall be back to a normal tab
+	TryInvokeTab(TabId);
+
+	return {};
+}
+
 FTabManager::~FTabManager()
 {
 	ClearPendingLayoutSave();
@@ -2594,6 +2763,86 @@ void FTabManager::UpdateStats()
 	StaticCastSharedRef<FTabManager>(FGlobalTabmanager::Get())->UpdateStats();
 }
 
+void FTabManager::TransferNomadPanelDrawerOnRestore(UE::Slate::Private::FPanelDrawerData&& HostedTabData, const TSharedPtr<SWindow>& OwningWindow)
+{
+	if (OwningWindow && HostedTabData.HostedTab && HostedTabData.HostedTab->GetTabRole() == ETabRole::NomadTab)
+	{
+		HostedTabData.HostedTab->SetTabManager(AsShared());
+		QueuedPanelDrawerTransfers.Add(OwningWindow, MakeUnique<UE::Slate::Private::FPanelDrawerData>(MoveTemp(HostedTabData)));
+	}
+}
+
+void FTabManager::TransferNomadPanelDrawer(UE::Slate::Private::FPanelDrawerData&& HostedTabData)
+{
+	if (!HostedTabData.HostedTab || HostedTabData.HostedTab->GetTabRole() != ETabRole::NomadTab)
+	{
+		return;
+	}
+
+	if (TSharedPtr<UE::Slate::Private::SPanelDrawerArea> PanelDrawer = PanelDrawerArea.Pin())
+	{
+		HostedTabData.HostedTab->SetTabManager(AsShared());
+		constexpr bool bWithAnimation = false;
+		PanelDrawer->TransferOwnershipOfHostedTab(MoveTemp(HostedTabData), bWithAnimation);
+	}
+}
+
+TUniquePtr<UE::Slate::Private::FPanelDrawerData> FTabManager::GetQueuedNomadDrawerPanelForWindow(const TSharedPtr<SWindow>& InWindow)
+{
+	TUniquePtr<UE::Slate::Private::FPanelDrawerData> QueuedPanelDrawerData;
+	QueuedPanelDrawerTransfers.RemoveAndCopyValue(InWindow, QueuedPanelDrawerData);
+	if (!QueuedPanelDrawerData)
+	{
+		FTabManager& GlobalTabManager = FGlobalTabmanager::Get().Get();
+		if (this != &GlobalTabManager)
+		{
+			return GlobalTabManager.GetQueuedNomadDrawerPanelForWindow(InWindow);
+		}
+	}
+
+	return QueuedPanelDrawerData;
+}
+
+TUniquePtr<UE::Slate::Private::FPanelDrawerData> FTabManager::TakeOwnerShipOfHiddenTabInDrawerPanel(const TSharedPtr<SDockTab>& InTab)
+{
+	
+	if (TSharedPtr<UE::Slate::Private::SPanelDrawerArea> PanelDrawer = PanelDrawerArea.Pin())
+	{
+		if (!PanelDrawer->IsOpen() && PanelDrawer->IsHostingTab(InTab))
+		{
+			constexpr bool bWithAnimation = false;
+			return MakeUnique<UE::Slate::Private::FPanelDrawerData>(PanelDrawer->TakeOwnershipOfHostedTab(bWithAnimation));
+		}
+	}
+
+	TWeakPtr<SWindow> FoundInWindow;
+	for (const TPair<TWeakPtr<SWindow>, TUniquePtr<UE::Slate::Private::FPanelDrawerData>>& Pair : QueuedPanelDrawerTransfers)
+	{
+		if (Pair.Value->HostedTab == InTab)
+		{
+			FoundInWindow = Pair.Key;
+		}
+	}
+
+	TUniquePtr<UE::Slate::Private::FPanelDrawerData> QueuedPanelDrawerData;
+	if (FoundInWindow.IsValid())
+	{
+		QueuedPanelDrawerTransfers.RemoveAndCopyValue(FoundInWindow, QueuedPanelDrawerData);
+
+	}
+
+	if (!QueuedPanelDrawerData)
+	{
+		FTabManager& GlobalTabManager = FGlobalTabmanager::Get().Get();
+		if (this != &GlobalTabManager)
+		{
+			return GlobalTabManager.TakeOwnerShipOfHiddenTabInDrawerPanel(InTab);
+		}
+	}
+
+	return QueuedPanelDrawerData;
+}
+
 void FTabManager::GetRecordableStats( int32& OutTabCount, TArray<TSharedPtr<SWindow>>& OutUniqueParentWindows ) const
 {
 	OutTabCount = 0;
@@ -2867,6 +3116,8 @@ FName FGlobalTabmanager::GetTabTypeForPotentiallyLegacyTab(FName InTabType) cons
 
 void FGlobalTabmanager::OnTabForegrounded( const TSharedPtr<SDockTab>& NewForegroundTab, const TSharedPtr<SDockTab>& BackgroundedTab )
 {
+	FTabAndManagerForDrawer ForegroundTabAndManager;
+	ForegroundTabAndManager.Key = NewForegroundTab;
 	if (NewForegroundTab.IsValid())
 	{
 		// Show any child windows associated with the Major Tab that got foregrounded.
@@ -2875,11 +3126,14 @@ void FGlobalTabmanager::OnTabForegrounded( const TSharedPtr<SDockTab>& NewForegr
 		{
 			TSharedPtr<FTabManager> ForegroundTabManager = SubTabManagers[ForegroundedTabIndex].TabManager.Pin();
 			ForegroundTabManager->GetPrivateApi().ShowWindows();
+			ForegroundTabAndManager.Value = ForegroundTabManager;
 		}
 
 		NewForegroundTab->UpdateActivationTime();
 	}
 	
+	FTabAndManagerForDrawer BackgroundedTabAndSubManager;
+	BackgroundedTabAndSubManager.Key = BackgroundedTab;
 	if (BackgroundedTab.IsValid())
 	{
 		// Hide any child windows associated with the Major Tab that got backgrounded.
@@ -2888,14 +3142,99 @@ void FGlobalTabmanager::OnTabForegrounded( const TSharedPtr<SDockTab>& NewForegr
 		{
 			TSharedPtr<FTabManager> BackgroundedTabManager = SubTabManagers[BackgroundedTabIndex].TabManager.Pin();
 			BackgroundedTabManager->GetPrivateApi().HideWindows();
+			BackgroundedTabAndSubManager.Value = BackgroundedTabManager;
 		}
 	}
 
+	RelocatePanelDrawerNomadTab(MoveTemp(ForegroundTabAndManager), MoveTemp(BackgroundedTabAndSubManager));
+
 	TabForegrounded.Broadcast(NewForegroundTab, BackgroundedTab);
 }
 
+void FGlobalTabmanager::RelocatePanelDrawerNomadTab(FTabAndManagerForDrawer ForegroundTabAndSubManager, FTabAndManagerForDrawer BackgroundedTabAndSubManager)
+{
+	/** 
+	 * 1. Check if the background tab has a drawer with nomad tab.
+	 * 2. If not used the queued operation.
+	 * 3. Check if the new tab can receive the drawer nomad tab.
+	 * 4. If not queue the transfer.
+	 * 5. If yes then do the transfer 
+	 */
+
+	FTabManager* BackgroundedTabDrawerManager = BackgroundedTabAndSubManager.Value.Get();
+	if (!BackgroundedTabDrawerManager)
+	{
+		BackgroundedTabDrawerManager = this;
+	}
+
+	constexpr bool bWithAnimation = false;
+	UE::Slate::Private::FPanelDrawerData PreviousPanelData;
+	if (BackgroundedTabDrawerManager)
+	{
+		if (TSharedPtr<UE::Slate::Private::SPanelDrawerArea> PreviousPanelDrawer = BackgroundedTabDrawerManager->PanelDrawerArea.Pin())
+		{
+			if (PreviousPanelDrawer->IsHostingNomadTab())
+			{
+				PreviousPanelData = PreviousPanelDrawer->TakeOwnershipOfHostedTab(bWithAnimation);
+			}
+		}
+	}
+
+	TSharedPtr<SWindow> CurrentWindow;
+	if (ForegroundTabAndSubManager.Key)
+	{
+		CurrentWindow = ForegroundTabAndSubManager.Key->GetParentWindow();
+	}
+	else if (BackgroundedTabAndSubManager.Key)
+	{
+		CurrentWindow = BackgroundedTabAndSubManager.Key->GetParentWindow();
+	}
+	else
+	{
+		return;
+	}
+
+	if (!PreviousPanelData.HostedTab)
+	{
+		// Use a queued transfer if we have one
+		TUniquePtr<UE::Slate::Private::FPanelDrawerData> PtrToPreviousPanelData = GetQueuedNomadDrawerPanelForWindow(CurrentWindow);
+		if (PtrToPreviousPanelData)
+		{
+			PreviousPanelData = MoveTemp(*PtrToPreviousPanelData.Get());
+		}
+	}
+
+	if (!PreviousPanelData.HostedTab)
+	{
+		return;
+	}
+
+	// Check if the tab is nomad (we only transfer nomad tabs)
+	if (!ensure(PreviousPanelData.HostedTab->GetTabRole() == ETabRole::NomadTab))
+	{
+		return;
+	}
+
+	FTabManager* ForegroundTabDrawerManager = ForegroundTabAndSubManager.Value.Get();
+	if (!ForegroundTabDrawerManager)
+	{
+		ForegroundTabDrawerManager = this;
+	}
+
+	TSharedPtr<UE::Slate::Private::SPanelDrawerArea> NewPanelDrawer = ForegroundTabDrawerManager->PanelDrawerArea.Pin();
+	if (!NewPanelDrawer)
+	{
+		ForegroundTabDrawerManager->TransferNomadPanelDrawerOnRestore(MoveTemp(PreviousPanelData), CurrentWindow);
+	}
+	else
+	{
+		ForegroundTabDrawerManager->TransferNomadPanelDrawer(MoveTemp(PreviousPanelData));
+	}
+}
+
 void FGlobalTabmanager::OnTabRelocated( const TSharedRef<SDockTab>& RelocatedTab, const TSharedPtr<SWindow>& NewOwnerWindow )
 {
+	// Handle transferring the drawer
 	if ( RelocatedTab->GetTabRole() == ETabRole::MajorTab || RelocatedTab->GetTabRole() == ETabRole::NomadTab )
 	{
 		LastMajorDockWindow = NewOwnerWindow;
diff --git a/Engine/Source/Runtime/Slate/Public/Framework/Docking/TabManager.h b/Engine/Source/Runtime/Slate/Public/Framework/Docking/TabManager.h
--- a/Engine/Source/Runtime/Slate/Public/Framework/Docking/TabManager.h
+++ b/Engine/Source/Runtime/Slate/Public/Framework/Docking/TabManager.h
@@ -24,6 +24,12 @@ class FLayoutExtender;
 struct FTabMatcher;
 struct FSidebarTabLists;
 
+namespace UE::Slate::Private
+{
+	class SPanelDrawerArea;
+	struct FPanelDrawerData;
+}
+
 DECLARE_MULTICAST_DELEGATE_TwoParams(
 	FOnActiveTabChanged,
 	/** Newly active tab */
@@ -932,6 +938,9 @@ class FTabManager : public TSharedFromThis<FTabManager>
 		SLATE_API TSharedPtr<SWidget> RestoreFrom(const TSharedRef<FLayout>& Layout, const TSharedPtr<SWindow>& ParentWindow, const bool bEmbedTitleAreaContent = false,
 			const EOutputCanBeNullptr RestoreAreaOutputCanBeNullptr = EOutputCanBeNullptr::Never);
 
+		UE_EXPERIMENTAL(5.7, "The panel drawer is still in development and might be removed. The api, implementation details and its form may change in near releases.")
+		SLATE_API TSharedPtr<SWidget> RestorePanelDrawerFrom(const TSharedRef<FLayout>& Layout, const TSharedRef<SWidget>& InContent, const TSharedPtr<SWindow>& ParentWindow);
+
 		SLATE_API void PopulateLocalTabSpawnerMenu( FMenuBuilder& PopulateMe );
 
 		SLATE_API void PopulateTabSpawnerMenu(FMenuBuilder& PopulateMe, TSharedRef<FWorkspaceItem> MenuStructure);
@@ -983,6 +992,16 @@ class FTabManager : public TSharedFromThis<FTabManager>
 		 */
 		SLATE_API TSharedPtr<SDockTab> FindExistingLiveTab(const FTabId& TabId) const;
 
+		/**
+		 * Try to open tab in the drawer. If it already exists elsewhere, it will draw attention to the tab.
+		 * If the drawer is already open with the tab it will close the drawer
+		 * If no drawer is available in the manager we will just open a new standard tab
+		 * 
+		 * @param TabId The tab identifier.
+		 */
+		UE_EXPERIMENTAL(5.7, "The panel drawer is still in development and might be removed. The api, implementation details and its form may change in near releases.")
+		SLATE_API TSharedPtr<SDockTab> TryToggleTabInPanelDrawer(const FTabId& TabId);
+
 		SLATE_API virtual ~FTabManager();
 
 		/** Sets whether or not this tab manager supports a custom menu bar for the active major tab that will be shown on top of the major tab area in the window this tab manager resides in. */
@@ -1188,6 +1207,18 @@ protected:
 		/** Called at the end of RestoreFrom for tab managers to complete any work after all tabs have been restored */
 		virtual void FinishRestore() {};
 
+		UE_EXPERIMENTAL(5.7, "The panel drawer is still in development and might be removed. The api, implementation details and its form may change in near releases.")
+		void TransferNomadPanelDrawerOnRestore(UE::Slate::Private::FPanelDrawerData&& HostedTabData, const TSharedPtr<SWindow>& OwningWindow);
+		
+		UE_EXPERIMENTAL(5.7, "The panel drawer is still in development and might be removed. The api, implementation details and its form may change in near releases.")
+		void TransferNomadPanelDrawer(UE::Slate::Private::FPanelDrawerData&& HostedTabData);
+
+		UE_EXPERIMENTAL(5.7, "The panel drawer is still in development and might be removed. The api, implementation details and its form may change in near releases.")
+		TUniquePtr<UE::Slate::Private::FPanelDrawerData> GetQueuedNomadDrawerPanelForWindow(const TSharedPtr<SWindow>& InWindow);
+		
+		UE_EXPERIMENTAL(5.7, "The panel drawer is still in development and might be removed. The api, implementation details and its form may change in near releases.")
+		TUniquePtr<UE::Slate::Private::FPanelDrawerData> TakeOwnerShipOfHiddenTabInDrawerPanel(const TSharedPtr<SDockTab>& InTab);
+
 	private:
 		/** Checks all dock areas and adds up the number of open tabs and unique parent windows in the manager */
 		SLATE_API void GetRecordableStats( int32& OutTabCount, TArray<TSharedPtr<SWindow>>& OutUniqueParentWindows ) const;
@@ -1212,6 +1243,11 @@ protected:
 	/** Based on if the tab can have a hidden display name it will return the correct Label to use for the tab */
 	FText GetTabLabelBasedOnSpawner(const TSharedPtr<FTabSpawnerEntry>& InSpawnerEntry) const;
 
+	/*
+	 * Does the shared logic between InvokeTab_Internal and TryInvokeTabInDrawer
+	 */
+	bool InvokeTab_SharedInternal(const FTabId& TabID, bool bInvokeAsInactive, TSharedPtr<SDockTab>& OutExistingTab);
+
 	protected:
 		FTabSpawner TabSpawner;
 		TSharedRef<FTabSpawner> NomadTabSpawner;
@@ -1236,6 +1272,12 @@ protected:
 		/** A Major tab that contains this TabManager's widgets. */
 		TWeakPtr<SDockTab> OwnerTabPtr;
 
+		/** A optional Panel Drawer Area for that tab manager */
+		TWeakPtr<UE::Slate::Private::SPanelDrawerArea> PanelDrawerArea;
+
+		/** A queued transfer for a nomad tab in the panel drawer when restored */
+		TMap<TWeakPtr<SWindow>, TUniquePtr<UE::Slate::Private::FPanelDrawerData>> QueuedPanelDrawerTransfers;
+
 		/** The current menu multi-box for the tab, used to construct platform native main menus */
 		TSharedPtr<FMultiBox> MenuMultiBox;
 		TSharedPtr<SWidget> MenuWidget;
@@ -1448,6 +1490,10 @@ public:
 	SLATE_API TSharedPtr<FTabManager::FLayout> GetInitialLayoutSP();
 	
 private:
+
+	using FTabAndManagerForDrawer = TPair<TSharedPtr<SDockTab>, TSharedPtr<FTabManager>>;
+	void RelocatePanelDrawerNomadTab(FTabAndManagerForDrawer ForegroundTabAndSubManager, FTabAndManagerForDrawer BackgroundedTabAndSubManager);
+
 	/** Whether tabs label should use MiddleEllipsis OverflowPolicy*/
 	bool bShouldUseMiddleEllipsisForDockTabLabel = false;
 
