Slate Dynamic Invalidation - ExpanderArrow

diff --git a/Engine/Plugins/Animation/DeformerGraph/Source/OptimusEditor/Private/Widgets/SOptimusEditorGraphNode.cpp b/Engine/Plugins/Animation/DeformerGraph/Source/OptimusEditor/Private/Widgets/SOptimusEditorGraphNode.cpp
--- a/Engine/Plugins/Animation/DeformerGraph/Source/OptimusEditor/Private/Widgets/SOptimusEditorGraphNode.cpp
+++ b/Engine/Plugins/Animation/DeformerGraph/Source/OptimusEditor/Private/Widgets/SOptimusEditorGraphNode.cpp
@@ -81,8 +81,8 @@ class SOptimusEditorExpanderArrow : public SExpanderArrow
 
 	FMargin GetExpanderPadding_Extended() const
 	{
-		const int32 NestingDepth = FMath::Max(0, OwnerRowPtr.Pin()->GetIndentLevel() - BaseIndentLevel.Get());
-		const float Indent = IndentAmount.Get(8.0f);
+		const int32 NestingDepth = FMath::Max(0, OwnerRowPtr.Pin()->GetIndentLevel() - GetBaseIndentLevel());
+		const float Indent = GetIndentAmount();
 		return bLeftAligned ? FMargin(NestingDepth * Indent, 0, 0, 0) : FMargin(0, 0, NestingDepth * Indent, 0);
 	}
 	
diff --git a/Engine/Plugins/Enterprise/DataprepEditor/Source/DataprepEditor/Private/Widgets/SDataprepPalette.cpp b/Engine/Plugins/Enterprise/DataprepEditor/Source/DataprepEditor/Private/Widgets/SDataprepPalette.cpp
--- a/Engine/Plugins/Enterprise/DataprepEditor/Source/DataprepEditor/Private/Widgets/SDataprepPalette.cpp
+++ b/Engine/Plugins/Enterprise/DataprepEditor/Source/DataprepEditor/Private/Widgets/SDataprepPalette.cpp
@@ -59,8 +59,8 @@ public:
 		{
 			OwnerRowPtr = ActionMenuData.TableRow;
 			StyleSet = InArgs._StyleSet;
-			IndentAmount = InArgs._IndentAmount;
-			BaseIndentLevel = InArgs._BaseIndentLevel;
+			SetIndentAmount(InArgs._IndentAmount);
+			SetBaseIndentLevel(InArgs._BaseIndentLevel);
 			ShouldDrawWires.Assign(*this, InArgs._ShouldDrawWires);
 
 			ChildSlot
diff --git a/Engine/Plugins/Experimental/SequenceNavigator/Source/SequenceNavigator/Private/Widgets/SNavigationToolExpanderArrow.cpp b/Engine/Plugins/Experimental/SequenceNavigator/Source/SequenceNavigator/Private/Widgets/SNavigationToolExpanderArrow.cpp
--- a/Engine/Plugins/Experimental/SequenceNavigator/Source/SequenceNavigator/Private/Widgets/SNavigationToolExpanderArrow.cpp
+++ b/Engine/Plugins/Experimental/SequenceNavigator/Source/SequenceNavigator/Private/Widgets/SNavigationToolExpanderArrow.cpp
@@ -42,7 +42,7 @@ int32 SNavigationToolExpanderArrow::OnPaint(const FPaintArgs& InArgs
 	if (VerticalBarBrush && ShouldDrawWires.Get())
 	{
 		const FLinearColor WireColor = WireTint.Get(FLinearColor(0.1f, 0.1f, 0.1f, 0.25f));
-		const float Indent = IndentAmount.Get(10.f);
+		const float Indent = GetIndentAmount();
 
 		const TSharedPtr<ITableRow> OwnerRow = OwnerRowPtr.Pin();
 		check(OwnerRow.IsValid());
diff --git a/Engine/Plugins/FX/Niagara/Source/NiagaraEditor/Public/Widgets/SNiagaraActionMenuExpander.h b/Engine/Plugins/FX/Niagara/Source/NiagaraEditor/Public/Widgets/SNiagaraActionMenuExpander.h
--- a/Engine/Plugins/FX/Niagara/Source/NiagaraEditor/Public/Widgets/SNiagaraActionMenuExpander.h
+++ b/Engine/Plugins/FX/Niagara/Source/NiagaraEditor/Public/Widgets/SNiagaraActionMenuExpander.h
@@ -16,7 +16,7 @@ public:
 	void Construct(const FArguments& InArgs, const FCustomExpanderData& ActionMenuData)
 	{
 		OwnerRowPtr = ActionMenuData.TableRow;
-		IndentAmount = InArgs._IndentAmount;
+		SetIndentAmount(InArgs._IndentAmount);
 		if (!ActionMenuData.RowAction.IsValid())
 		{
 			SExpanderArrow::FArguments SuperArgs;
diff --git a/Engine/Plugins/FX/Niagara/Source/NiagaraEditorWidgets/Private/Stack/SNiagaraStackFunctionInputValue.h b/Engine/Plugins/FX/Niagara/Source/NiagaraEditorWidgets/Private/Stack/SNiagaraStackFunctionInputValue.h
--- a/Engine/Plugins/FX/Niagara/Source/NiagaraEditorWidgets/Private/Stack/SNiagaraStackFunctionInputValue.h
+++ b/Engine/Plugins/FX/Niagara/Source/NiagaraEditorWidgets/Private/Stack/SNiagaraStackFunctionInputValue.h
@@ -99,7 +99,7 @@ private:
 		void Construct(const FArguments& InArgs, const FCustomExpanderData& ActionMenuData)
 		{
 			OwnerRowPtr = ActionMenuData.TableRow;
-			IndentAmount = InArgs._IndentAmount;
+			SetIndentAmount(InArgs._IndentAmount);
 			if (!ActionMenuData.RowAction.IsValid())
 			{
 				SExpanderArrow::FArguments SuperArgs;
diff --git a/Engine/Plugins/Mutable/Source/CustomizableObjectEditor/Private/MuCOE/Widgets/MutableExpanderArrow.cpp b/Engine/Plugins/Mutable/Source/CustomizableObjectEditor/Private/MuCOE/Widgets/MutableExpanderArrow.cpp
--- a/Engine/Plugins/Mutable/Source/CustomizableObjectEditor/Private/MuCOE/Widgets/MutableExpanderArrow.cpp
+++ b/Engine/Plugins/Mutable/Source/CustomizableObjectEditor/Private/MuCOE/Widgets/MutableExpanderArrow.cpp
@@ -18,7 +18,7 @@ int32 SMutableExpanderArrow::OnPaint(const FPaintArgs& Args, const FGeometry& Al
 	static constexpr float HalfWireThickness = WireThickness / 2.0f;
 	
 	static const FName NAME_VerticalBarBrush = TEXT("WhiteBrush");
-	const float Indent = IndentAmount.Get(10.f);
+	const float Indent = GetIndentAmount();
 	const FSlateBrush* VerticalBarBrush = (StyleSet == nullptr) ? nullptr : StyleSet->GetBrush(NAME_VerticalBarBrush);
 
 	if (ShouldDrawWires.Get() == true && VerticalBarBrush != nullptr)
diff --git a/Engine/Plugins/Runtime/Metasound/Source/MetasoundEditor/Private/SMetasoundActionMenu.cpp b/Engine/Plugins/Runtime/Metasound/Source/MetasoundEditor/Private/SMetasoundActionMenu.cpp
--- a/Engine/Plugins/Runtime/Metasound/Source/MetasoundEditor/Private/SMetasoundActionMenu.cpp
+++ b/Engine/Plugins/Runtime/Metasound/Source/MetasoundEditor/Private/SMetasoundActionMenu.cpp
@@ -36,7 +36,7 @@ namespace Metasound
 		void SMetasoundActionMenuExpanderArrow::Construct(const FArguments& InArgs, const FCustomExpanderData& ActionMenuData)
 		{
 			OwnerRowPtr = ActionMenuData.TableRow;
-			IndentAmount = InArgs._IndentAmount;
+			SetIndentAmount(InArgs._IndentAmount);
 			ActionPtr = ActionMenuData.RowAction;
 
 			if (ActionPtr.IsValid())
diff --git a/Engine/Plugins/VirtualProduction/Avalanche/Source/AvalancheOutliner/Private/Slate/SAvaOutlinerExpanderArrow.cpp b/Engine/Plugins/VirtualProduction/Avalanche/Source/AvalancheOutliner/Private/Slate/SAvaOutlinerExpanderArrow.cpp
--- a/Engine/Plugins/VirtualProduction/Avalanche/Source/AvalancheOutliner/Private/Slate/SAvaOutlinerExpanderArrow.cpp
+++ b/Engine/Plugins/VirtualProduction/Avalanche/Source/AvalancheOutliner/Private/Slate/SAvaOutlinerExpanderArrow.cpp
@@ -35,7 +35,7 @@ int32 SAvaOutlinerExpanderArrow::OnPaint(const FPaintArgs& Args
 	//
 	static const FName NAME_VerticalBarBrush = TEXT("WhiteBrush");
 	
-	const float Indent = IndentAmount.Get(10.f);
+	const float Indent = GetIndentAmount();
 	
 	const FSlateBrush* const VerticalBarBrush = StyleSet == nullptr
 		? nullptr
diff --git a/Engine/Source/Editor/Kismet/Private/SBlueprintActionMenu.cpp b/Engine/Source/Editor/Kismet/Private/SBlueprintActionMenu.cpp
--- a/Engine/Source/Editor/Kismet/Private/SBlueprintActionMenu.cpp
+++ b/Engine/Source/Editor/Kismet/Private/SBlueprintActionMenu.cpp
@@ -750,7 +750,7 @@ void SBlueprintActionMenu::Tick(const FGeometry& AllottedGeometry, const double
 void SBlueprintActionMenuExpander::Construct(const FArguments& InArgs, const FCustomExpanderData& ActionMenuData)
 {
 	OwnerRowPtr  = ActionMenuData.TableRow;
-	IndentAmount = InArgs._IndentAmount;
+	SetIndentAmount(InArgs._IndentAmount);
 	ActionPtr    = ActionMenuData.RowAction;
 
 	if (!ActionPtr.IsValid())
diff --git a/Engine/Source/Runtime/Slate/Private/Widgets/Views/SExpanderArrow.cpp b/Engine/Source/Runtime/Slate/Private/Widgets/Views/SExpanderArrow.cpp
--- a/Engine/Source/Runtime/Slate/Private/Widgets/Views/SExpanderArrow.cpp
+++ b/Engine/Source/Runtime/Slate/Private/Widgets/Views/SExpanderArrow.cpp
@@ -10,10 +10,15 @@ SLATE_IMPLEMENT_WIDGET(SExpanderArrow)
 void SExpanderArrow::PrivateRegisterAttributes(FSlateAttributeInitializer& AttributeInitializer)
 {
 	SLATE_ADD_MEMBER_ATTRIBUTE_DEFINITION(AttributeInitializer, ShouldDrawWires, EInvalidateWidgetReason::Paint);
+	SLATE_ADD_MEMBER_ATTRIBUTE_DEFINITION(AttributeInitializer, BaseIndentLevelAttribute, EInvalidateWidgetReason::Layout);
+	SLATE_ADD_MEMBER_ATTRIBUTE_DEFINITION(AttributeInitializer, IndentAmountAttribute, EInvalidateWidgetReason::Layout);
 }
 
 SExpanderArrow::SExpanderArrow()
-	: ShouldDrawWires(*this, false)
+	: StyleSet(nullptr)
+	, ShouldDrawWires(*this, false)
+	, IndentAmountAttribute(*this, 10.0f)
+	, BaseIndentLevelAttribute(*this, 0)
 {
 
 }
@@ -24,8 +29,8 @@ void SExpanderArrow::Construct( const FArguments& InArgs, const TSharedPtr<class
 {
 	OwnerRowPtr = TableRow;
 	StyleSet = InArgs._StyleSet;
-	IndentAmount = InArgs._IndentAmount;
-	BaseIndentLevel = InArgs._BaseIndentLevel;
+	SetIndentAmount(InArgs._IndentAmount);
+	SetBaseIndentLevel(InArgs._BaseIndentLevel);
 	ShouldDrawWires.Assign(*this, InArgs._ShouldDrawWires);
 
 	this->ChildSlot
@@ -70,7 +75,7 @@ int32 SExpanderArrow::OnPaint(const FPaintArgs& Args, const FGeometry& AllottedG
 	//
 
 	static const FName NAME_VerticalBarBrush = TEXT("WhiteBrush");
-	const float Indent = IndentAmount.Get(10.f);
+	const float Indent = IndentAmountAttribute.Get();
 	const FSlateBrush* VerticalBarBrush = (StyleSet == nullptr) ? nullptr : StyleSet->GetBrush(NAME_VerticalBarBrush);
 
 	if (ShouldDrawWires.Get() == true && VerticalBarBrush != nullptr)
@@ -169,6 +174,16 @@ FReply SExpanderArrow::OnArrowClicked()
 	return FReply::Handled();
 }
 
+void SExpanderArrow::SetIndentAmount(TAttribute<float> InIndentAmount)
+{
+	IndentAmountAttribute.Assign(*this, MoveTemp(InIndentAmount));
+}
+
+void SExpanderArrow::SetBaseIndentLevel(TAttribute<int32> InBaseIndentLevel)
+{
+	BaseIndentLevelAttribute.Assign(*this,  MoveTemp(InBaseIndentLevel));
+}
+
 /** @return Visible when has children; invisible otherwise */
 EVisibility SExpanderArrow::GetExpanderVisibility() const
 {
@@ -178,8 +193,8 @@ EVisibility SExpanderArrow::GetExpanderVisibility() const
 /** @return the margin corresponding to how far this item is indented */
 FMargin SExpanderArrow::GetExpanderPadding() const
 {
-	const int32 NestingDepth = FMath::Max(0, OwnerRowPtr.Pin()->GetIndentLevel() - BaseIndentLevel.Get());
-	const float Indent = IndentAmount.Get(10.f);
+	const int32 NestingDepth = FMath::Max(0, OwnerRowPtr.Pin()->GetIndentLevel() - BaseIndentLevelAttribute.Get());
+	const float Indent = IndentAmountAttribute.Get();
 	return FMargin( NestingDepth * Indent, 0,0,0 );
 }
 
diff --git a/Engine/Source/Runtime/Slate/Public/Widgets/Views/SExpanderArrow.h b/Engine/Source/Runtime/Slate/Public/Widgets/Views/SExpanderArrow.h
--- a/Engine/Source/Runtime/Slate/Public/Widgets/Views/SExpanderArrow.h
+++ b/Engine/Source/Runtime/Slate/Public/Widgets/Views/SExpanderArrow.h
@@ -14,6 +14,12 @@
 class ITableRow;
 class SButton;
 
+template<>
+struct TWidgetTypeTraits<class SExpanderArrow>
+{
+	static constexpr bool SupportsInvalidation() { return true; }
+};
+
 /**
  * Expander arrow and indentation component that can be placed in a TableRow
  * of a TreeView. Intended for use by TMultiColumnRow in TreeViews.
@@ -49,6 +55,21 @@ protected:
 	/** Invoked when the expanded button is clicked (toggle item expansion) */
 	SLATE_API FReply OnArrowClicked();
 
+	SLATE_API void SetIndentAmount(TAttribute<float> InIndentAmount);
+	SLATE_API void SetBaseIndentLevel(TAttribute<int32> InBaseIndentLevel);
+	float GetIndentAmount() const { return IndentAmountAttribute.Get(); }
+	int32 GetBaseIndentLevel() const { return BaseIndentLevelAttribute.Get(); }
+
+	TSlateAttributeRef<float> GetIndentAmountAttribute() const
+	{
+		return TSlateAttributeRef<float>{SharedThis(this), IndentAmountAttribute};
+	}
+
+	TSlateAttributeRef<int32> GetBaseIndentLevelAttribute() const
+	{
+		return TSlateAttributeRef<int32>{SharedThis(this), BaseIndentLevelAttribute};
+	}
+
 	/** @return Visible when has children; invisible otherwise */
 	SLATE_API EVisibility GetExpanderVisibility() const;
 
@@ -66,12 +87,22 @@ protected:
 	/** The slate style to use */
 	const ISlateStyle* StyleSet;
 
+#if WITH_EDITORONLY_DATA
 	/** The amount of space to indent at each level */
-	TAttribute<float> IndentAmount;
+	UE_DEPRECATED(5.7, "Use SetIndentAmount / GetIndentAmount / GetIndentAmountAttribute")
+	TSlateDeprecatedTAttribute<float> IndentAmount;
 
 	/** The level in the tree that begins the indention amount */
-	TAttribute<int32> BaseIndentLevel;
+	UE_DEPRECATED(5.7, "Use SetBaseIndentLevel / GetBaseIndentLevel / GetBaseIndentLevelAttribute")
+	TSlateDeprecatedTAttribute<int32> BaseIndentLevel;
+#endif
 
 	/** Whether to draw the wires that visually reinforce the tree hierarchy. */
 	TSlateAttribute<bool> ShouldDrawWires;
+
+private:
+	/** The amount of space to indent at each level */
+	TSlateAttribute<float> IndentAmountAttribute;
+	/** The level in the tree that begins the indention amount */
+	TSlateAttribute<int32> BaseIndentLevelAttribute;
 };
