Adding callback when the text starts to be edited for SEditableTextBox

#FailedValidations Preflight

diff --git a/Engine/Source/Runtime/Slate/Private/Widgets/Input/SEditableText.cpp b/Engine/Source/Runtime/Slate/Private/Widgets/Input/SEditableText.cpp
--- a/Engine/Source/Runtime/Slate/Private/Widgets/Input/SEditableText.cpp
+++ b/Engine/Source/Runtime/Slate/Private/Widgets/Input/SEditableText.cpp
@@ -36,6 +36,7 @@ void SEditableText::Construct( const FArguments& InArgs )
 	bAllowContextMenu = InArgs._AllowContextMenu;
 	OnContextMenuOpening = InArgs._OnContextMenuOpening;
 	OnIsTypedCharValid = InArgs._OnIsTypedCharValid;
+	OnBeginTextEditCallback = InArgs._OnBeginTextEdit;
 	OnTextChangedCallback = InArgs._OnTextChanged;
 	OnTextCommittedCallback = InArgs._OnTextCommitted;
 	MinDesiredWidth = InArgs._MinDesiredWidth;
@@ -664,6 +665,11 @@ TSharedPtr<SWidget> SEditableText::BuildContextMenuContent() const
 	return EditableTextLayout->BuildDefaultContextMenu(MenuExtender);
 }
 
+void SEditableText::OnBeginTextEdit(const FText& InText)
+{
+	OnBeginTextEditCallback.ExecuteIfBound(InText);
+}
+
 void SEditableText::OnTextChanged(const FText& InText)
 {
 	OnTextChangedCallback.ExecuteIfBound(InText);
diff --git a/Engine/Source/Runtime/Slate/Private/Widgets/Input/SEditableTextBox.cpp b/Engine/Source/Runtime/Slate/Private/Widgets/Input/SEditableTextBox.cpp
--- a/Engine/Source/Runtime/Slate/Private/Widgets/Input/SEditableTextBox.cpp
+++ b/Engine/Source/Runtime/Slate/Private/Widgets/Input/SEditableTextBox.cpp
@@ -73,6 +73,7 @@ void SEditableTextBox::Construct( const FArguments& InArgs )
 					.AllowContextMenu( InArgs._AllowContextMenu )
 					.OnContextMenuOpening( InArgs._OnContextMenuOpening )
 					.ContextMenuExtender( InArgs._ContextMenuExtender )
+					.OnBeginTextEdit(InArgs._OnBeginTextEdit)
 					.OnTextChanged(this, &SEditableTextBox::OnEditableTextChanged)
 					.OnTextCommitted(this, &SEditableTextBox::OnEditableTextCommitted)
 					.MinDesiredWidth( InArgs._MinDesiredWidth )
diff --git a/Engine/Source/Runtime/Slate/Private/Widgets/Input/SMultiLineEditableTextBox.cpp b/Engine/Source/Runtime/Slate/Private/Widgets/Input/SMultiLineEditableTextBox.cpp
--- a/Engine/Source/Runtime/Slate/Private/Widgets/Input/SMultiLineEditableTextBox.cpp
+++ b/Engine/Source/Runtime/Slate/Private/Widgets/Input/SMultiLineEditableTextBox.cpp
@@ -105,6 +105,7 @@ void SMultiLineEditableTextBox::Construct( const FArguments& InArgs )
 					.AllowMultiLine( InArgs._AllowMultiLine )
 					.OnContextMenuOpening( InArgs._OnContextMenuOpening )
 					.OnIsTypedCharValid( InArgs._OnIsTypedCharValid )
+					.OnBeginTextEdit( InArgs._OnBeginTextEdit )
 					.OnTextChanged(this, &SMultiLineEditableTextBox::OnEditableTextChanged)
 					.OnTextCommitted(this, &SMultiLineEditableTextBox::OnEditableTextCommitted)
 					.OnCursorMoved( InArgs._OnCursorMoved )
diff --git a/Engine/Source/Runtime/Slate/Private/Widgets/Text/SInlineEditableTextBlock.cpp b/Engine/Source/Runtime/Slate/Private/Widgets/Text/SInlineEditableTextBlock.cpp
--- a/Engine/Source/Runtime/Slate/Private/Widgets/Text/SInlineEditableTextBlock.cpp
+++ b/Engine/Source/Runtime/Slate/Private/Widgets/Text/SInlineEditableTextBlock.cpp
@@ -13,7 +13,6 @@ void SInlineEditableTextBlock::Construct( const FArguments& InArgs )
 {
 	check(InArgs._Style);
 
-	OnBeginTextEditDelegate = InArgs._OnBeginTextEdit;
 	OnTextCommittedDelegate = InArgs._OnTextCommitted;
 	IsSelected = InArgs._IsSelected;
 	OnVerifyTextChanged = InArgs._OnVerifyTextChanged;
@@ -62,6 +61,7 @@ void SInlineEditableTextBlock::Construct( const FArguments& InArgs )
 			.Style(&InArgs._Style->EditableTextBoxStyle)
 			.Font(InArgs._Font)
 			.ToolTipText(InArgs._ToolTipText)
+			.OnBeginTextEdit(InArgs._OnBeginTextEdit)
 			.OnTextChanged(this, &SInlineEditableTextBlock::OnTextChanged)
 			.OnTextCommitted(this, &SInlineEditableTextBlock::OnTextBoxCommitted)
 			.WrapTextAt(InArgs._WrapTextAt)
@@ -82,6 +82,7 @@ void SInlineEditableTextBlock::Construct( const FArguments& InArgs )
 			.Style(&InArgs._Style->EditableTextBoxStyle)
 			.Font(InArgs._Font)
 			.ToolTipText(InArgs._ToolTipText)
+			.OnBeginTextEdit(InArgs._OnBeginTextEdit)
 			.OnTextChanged(this, &SInlineEditableTextBlock::OnTextChanged)
 			.OnTextCommitted(this, &SInlineEditableTextBlock::OnTextBoxCommitted)
 			.SelectAllTextWhenFocused(true)
@@ -134,8 +135,6 @@ void SInlineEditableTextBlock::EnterEditingMode()
 			FSlateApplication::Get().SetKeyboardFocus(ActiveTextBox, EFocusCause::SetDirectly);
 
 			TextBlock->SetVisibility(EVisibility::Collapsed);
-
-			OnBeginTextEditDelegate.ExecuteIfBound( CurrentText );
 		}
 	}
 }
@@ -223,11 +222,6 @@ void SInlineEditableTextBlock::SetOverflowPolicy(TOptional<ETextOverflowPolicy>
 	TextBlock->SetOverflowPolicy(InOverflowPolicy);
 }
 
-void SInlineEditableTextBlock::SetOnBeginTextEdit(FOnBeginTextEdit InOnBeginTextEdit)
-{
-	OnBeginTextEditDelegate = MoveTemp(InOnBeginTextEdit);
-}
-
 void SInlineEditableTextBlock::SetOnTextCommitted(FOnTextCommitted InOnTextCommitted)
 {
 	OnTextCommittedDelegate = MoveTemp(InOnTextCommitted);
diff --git a/Engine/Source/Runtime/Slate/Private/Widgets/Text/SMultiLineEditableText.cpp b/Engine/Source/Runtime/Slate/Private/Widgets/Text/SMultiLineEditableText.cpp
--- a/Engine/Source/Runtime/Slate/Private/Widgets/Text/SMultiLineEditableText.cpp
+++ b/Engine/Source/Runtime/Slate/Private/Widgets/Text/SMultiLineEditableText.cpp
@@ -426,6 +426,11 @@ TSharedPtr<SWidget> SMultiLineEditableText::BuildContextMenuContent() const
 	return EditableTextLayout->BuildDefaultContextMenu(MenuExtender);
 }
 
+void SMultiLineEditableText::OnBeginTextEdit(const FText& InText)
+{
+	OnBeginTextEditCallback.ExecuteIfBound(InText);
+}
+
 void SMultiLineEditableText::OnTextChanged(const FText& InText)
 {
 	OnTextChangedCallback.ExecuteIfBound(InText);
diff --git a/Engine/Source/Runtime/Slate/Private/Widgets/Text/SlateEditableTextLayout.cpp b/Engine/Source/Runtime/Slate/Private/Widgets/Text/SlateEditableTextLayout.cpp
--- a/Engine/Source/Runtime/Slate/Private/Widgets/Text/SlateEditableTextLayout.cpp
+++ b/Engine/Source/Runtime/Slate/Private/Widgets/Text/SlateEditableTextLayout.cpp
@@ -882,6 +882,7 @@ bool FSlateEditableTextLayout::HandleFocusReceived(const FFocusEvent& InFocusEve
 	// Force update of last cursor interaction time to ensure the caret blinks at the correct frequency regardless of how focus is set.
 	CursorInfo.UpdateLastCursorInteractionTime();
 
+	OwnerWidget->OnBeginTextEdit(BoundText.Get(FText::GetEmpty()));
 	return true;
 }
 
diff --git a/Engine/Source/Runtime/Slate/Public/Framework/SlateDelegates.h b/Engine/Source/Runtime/Slate/Public/Framework/SlateDelegates.h
--- a/Engine/Source/Runtime/Slate/Public/Framework/SlateDelegates.h
+++ b/Engine/Source/Runtime/Slate/Public/Framework/SlateDelegates.h
@@ -85,6 +85,9 @@ DECLARE_DELEGATE_RetVal(
 /** Delegate for hooking up to an inline editable text block 'IsSelected' check. */
 DECLARE_DELEGATE_RetVal( bool, FIsSelected );
 
+/** Delegate for hooking up to an editable text box's 'OnBeginTextEdit' */
+DECLARE_DELEGATE_OneParam(FOnBeginTextEdit, const FText&)
+
 /** Delegate for hooking up to an editable text box's 'OnTextChanged' */
 DECLARE_DELEGATE_OneParam( FOnTextChanged, const FText& );
 
diff --git a/Engine/Source/Runtime/Slate/Public/Widgets/Input/SEditableText.h b/Engine/Source/Runtime/Slate/Public/Widgets/Input/SEditableText.h
--- a/Engine/Source/Runtime/Slate/Public/Widgets/Input/SEditableText.h
+++ b/Engine/Source/Runtime/Slate/Public/Widgets/Input/SEditableText.h
@@ -137,6 +137,9 @@ public:
 		 */
 		SLATE_EVENT( FOnIsTypedCharValid, OnIsTypedCharValid )
 
+		/** Callback when the text starts to be edited */
+		SLATE_EVENT( FOnBeginTextEdit, OnBeginTextEdit )
+
 		/** Called whenever the text is changed programmatically or interactively by the user */
 		SLATE_EVENT( FOnTextChanged, OnTextChanged )
 
@@ -455,6 +458,7 @@ protected:
 	SLATE_API virtual TSharedRef<SWidget> GetSlateWidget() override;
 	SLATE_API virtual TSharedPtr<SWidget> GetSlateWidgetPtr() override;
 	SLATE_API virtual TSharedPtr<SWidget> BuildContextMenuContent() const override;
+	SLATE_API virtual void OnBeginTextEdit(const FText& InText) override;
 	SLATE_API virtual void OnTextChanged(const FText& InText) override;
 	SLATE_API virtual void OnTextCommitted(const FText& InText, const ETextCommit::Type InTextAction) override;
 	SLATE_API virtual void OnCursorMoved(const FTextLocation& InLocation) override;
@@ -514,6 +518,9 @@ protected:
 	/** Called when a character is typed and we want to know if the text field supports typing this character. */
 	FOnIsTypedCharValid OnIsTypedCharValid;
 
+	/** Called when the text starts to be edited */
+	FOnBeginTextEdit OnBeginTextEditCallback;
+
 	/** Called whenever the text is changed programmatically or interactively by the user */
 	FOnTextChanged OnTextChangedCallback;
 
diff --git a/Engine/Source/Runtime/Slate/Public/Widgets/Input/SEditableTextBox.h b/Engine/Source/Runtime/Slate/Public/Widgets/Input/SEditableTextBox.h
--- a/Engine/Source/Runtime/Slate/Public/Widgets/Input/SEditableTextBox.h
+++ b/Engine/Source/Runtime/Slate/Public/Widgets/Input/SEditableTextBox.h
@@ -118,6 +118,9 @@ public:
 		/** Menu extender for the right-click context menu */
 		SLATE_EVENT(FMenuExtensionDelegate, ContextMenuExtender)
 
+		/** Callback when the text starts to be edited */
+		SLATE_EVENT( FOnBeginTextEdit, OnBeginTextEdit )
+
 		/** Called whenever the text is changed programmatically or interactively by the user */
 		SLATE_EVENT( FOnTextChanged, OnTextChanged )
 
diff --git a/Engine/Source/Runtime/Slate/Public/Widgets/Input/SMultiLineEditableTextBox.h b/Engine/Source/Runtime/Slate/Public/Widgets/Input/SMultiLineEditableTextBox.h
--- a/Engine/Source/Runtime/Slate/Public/Widgets/Input/SMultiLineEditableTextBox.h
+++ b/Engine/Source/Runtime/Slate/Public/Widgets/Input/SMultiLineEditableTextBox.h
@@ -165,6 +165,9 @@ public:
 		 */
 		SLATE_EVENT( FOnIsTypedCharValid, OnIsTypedCharValid )
 
+		/** Callback when the text starts to be edited */
+		SLATE_EVENT( FOnBeginTextEdit, OnBeginTextEdit )
+
 		/** Called whenever the text is changed programmatically or interactively by the user */
 		SLATE_EVENT( FOnTextChanged, OnTextChanged )
 
diff --git a/Engine/Source/Runtime/Slate/Public/Widgets/Text/ISlateEditableTextWidget.h b/Engine/Source/Runtime/Slate/Public/Widgets/Text/ISlateEditableTextWidget.h
--- a/Engine/Source/Runtime/Slate/Public/Widgets/Text/ISlateEditableTextWidget.h
+++ b/Engine/Source/Runtime/Slate/Public/Widgets/Text/ISlateEditableTextWidget.h
@@ -266,6 +266,9 @@ public:
 
 	/** Build the context menu content to use for this widget (if any) */
 	virtual TSharedPtr<SWidget> BuildContextMenuContent() const = 0;
+	
+	/** Called when the text starts to be edited */
+    virtual void OnBeginTextEdit(const FText& InText) = 0;
 
 	/** Called when the text has been changed by an edit operation */
 	virtual void OnTextChanged(const FText& InText) = 0;
diff --git a/Engine/Source/Runtime/Slate/Public/Widgets/Text/SInlineEditableTextBlock.h b/Engine/Source/Runtime/Slate/Public/Widgets/Text/SInlineEditableTextBlock.h
--- a/Engine/Source/Runtime/Slate/Public/Widgets/Text/SInlineEditableTextBlock.h
+++ b/Engine/Source/Runtime/Slate/Public/Widgets/Text/SInlineEditableTextBlock.h
@@ -23,8 +23,6 @@ class SHorizontalBox;
 class SMultiLineEditableTextBox;
 class STextBlock;
 
-DECLARE_DELEGATE_OneParam(FOnBeginTextEdit, const FText&)
-
 /**
  * Slate's InlineEditableTextBlock's are double selectable to go from a STextBlock to become SEditableTextBox.
  */
@@ -192,8 +190,6 @@ class SInlineEditableTextBlock: public SCompoundWidget
 	/** Sets the overflow policy for this text block */
 	SLATE_API void SetOverflowPolicy(TOptional<ETextOverflowPolicy> InOverflowPolicy);
 
-	SLATE_API void SetOnBeginTextEdit(FOnBeginTextEdit InOnBeginTextEdit);
-
 	SLATE_API void SetOnTextCommitted(FOnTextCommitted InOnTextCommitted);
 
 	SLATE_API void SetOnEnterEditingMode(FSimpleDelegate InOnEnterEditingMode);
@@ -229,9 +225,6 @@ protected:
 	FSimpleDelegate OnEnterEditingMode;
 	FSimpleDelegate OnExitEditingMode;
 
-	/** Delegate to execute when the text starts to be edited */
-	FOnBeginTextEdit OnBeginTextEditDelegate;
-
 	/** Delegate to execute when editing mode text is committed. */
 	FOnTextCommitted OnTextCommittedDelegate;
 
diff --git a/Engine/Source/Runtime/Slate/Public/Widgets/Text/SMultiLineEditableText.h b/Engine/Source/Runtime/Slate/Public/Widgets/Text/SMultiLineEditableText.h
--- a/Engine/Source/Runtime/Slate/Public/Widgets/Text/SMultiLineEditableText.h
+++ b/Engine/Source/Runtime/Slate/Public/Widgets/Text/SMultiLineEditableText.h
@@ -159,6 +159,9 @@ public:
 		 */
 		SLATE_EVENT(FOnIsTypedCharValid, OnIsTypedCharValid)
 
+		/** Callback when the text starts to be edited */
+		SLATE_EVENT(FOnBeginTextEdit, OnBeginTextEdit )
+
 		/** Called whenever the text is changed programmatically or interactively by the user */
 		SLATE_EVENT(FOnTextChanged, OnTextChanged)
 
@@ -528,6 +531,7 @@ protected:
 	SLATE_API virtual TSharedRef<SWidget> GetSlateWidget() override;
 	SLATE_API virtual TSharedPtr<SWidget> GetSlateWidgetPtr() override;
 	SLATE_API virtual TSharedPtr<SWidget> BuildContextMenuContent() const override;
+	SLATE_API virtual void OnBeginTextEdit(const FText& InText) override;
 	SLATE_API virtual void OnTextChanged(const FText& InText) override;
 	SLATE_API virtual void OnTextCommitted(const FText& InText, const ETextCommit::Type InTextAction) override;
 	SLATE_API virtual void OnCursorMoved(const FTextLocation& InLocation) override;
@@ -576,6 +580,9 @@ protected:
 	/** Called when a character is typed and we want to know if the text field supports typing this character. */
 	FOnIsTypedCharValid OnIsTypedCharValid;
 
+	/** Called when the text starts to be edited */
+	FOnBeginTextEdit OnBeginTextEditCallback;
+
 	/** Called whenever the text is changed programmatically or interactively by the user */
 	FOnTextChanged OnTextChangedCallback;
 
