[HomeScreen] Add possibility for tabs to lock themselves in the TabWell and do not allow also dragging operation on them

#FailedValidations Preflight

diff --git a/Engine/Source/Runtime/Slate/Private/Framework/Docking/SDockingTabStack.cpp b/Engine/Source/Runtime/Slate/Private/Framework/Docking/SDockingTabStack.cpp
--- a/Engine/Source/Runtime/Slate/Private/Framework/Docking/SDockingTabStack.cpp
+++ b/Engine/Source/Runtime/Slate/Private/Framework/Docking/SDockingTabStack.cpp
@@ -377,6 +377,26 @@ void SDockingTabStack::SetTabPinnedInSidebar(const TSharedRef<SDockTab>& InTab,
 	}
 }
 
+bool SDockingTabStack::IsTabLocked(const TSharedRef<const SDockTab>& InTab) const
+{
+	const FTabManager::FTab* Tab = Tabs.FindByPredicate(FTabMatcher(InTab->GetLayoutIdentifier()));
+	if (Tab)
+	{
+		return Tab->bIsLockedInPlace.Get(false);
+	}
+
+	return false;
+}
+
+void SDockingTabStack::SetTabLocked(const TSharedRef<const SDockTab>& InTab, TOptional<bool> bInLocked)
+{
+	FTabManager::FTab* Tab = Tabs.FindByPredicate(FTabMatcher(InTab->GetLayoutIdentifier()));
+	if (Tab)
+	{
+		Tab->bIsLockedInPlace = bInLocked;
+	}
+}
+
 const TSlotlessChildren<SDockTab>& SDockingTabStack::GetTabs() const
 {
 	return TabWell->GetTabs();
@@ -682,6 +702,15 @@ TSharedRef<SWidget> SDockingTabStack::MakeContextMenu()
 				);
 			}
 
+			MenuBuilder.AddMenuEntry(
+				TAttribute<FText>::CreateSP(this, &SDockingTabStack::GetLockTabText),
+				TAttribute<FText>::CreateSP(this, &SDockingTabStack::GetLockTabToolTip),
+				FSlateIcon(),
+				FUIAction(
+					FExecuteAction::CreateSP(this, &SDockingTabStack::ToggleForegroundTabLocked),
+					FCanExecuteAction::CreateSP(this, &SDockingTabStack::CanChangeTabLock)
+				)
+			);
 		}
 		MenuBuilder.EndSection();
 
@@ -1030,6 +1059,26 @@ EVisibility SDockingTabStack::GetUnhideButtonVisibility() const
 		: EVisibility::Collapsed;
 }
 
+FText SDockingTabStack::GetLockTabText() const
+{
+	if (TSharedPtr<SDockTab> ForegroundTabPtr = TabWell->GetForegroundTab())
+	{
+		return IsTabLocked(ForegroundTabPtr.ToSharedRef()) ? LOCTEXT("UnlockTabText", "Unlock Tab") : LOCTEXT("LockTabText", "Lock Tab");
+	}
+
+	return FText::GetEmpty();
+}
+
+FText SDockingTabStack::GetLockTabToolTip() const
+{
+	if (TSharedPtr<SDockTab> ForegroundTabPtr = TabWell->GetForegroundTab())
+	{
+		return IsTabLocked(ForegroundTabPtr.ToSharedRef()) ? LOCTEXT("UnlockTabToolTip", "Unlocks the tab from the current location") : LOCTEXT("LockTabToolTip", "Locks the tab to the current location and prevents moving/closing");
+	}
+
+	return FText::GetEmpty();
+}
+
 void SDockingTabStack::ToggleTabWellVisibility()
 {
 	ShowHideTabWell.Reverse();
@@ -1043,6 +1092,15 @@ void SDockingTabStack::MoveForegroundTabToSidebar()
 	}
 }
 
+void SDockingTabStack::ToggleForegroundTabLocked()
+{
+	if (const TSharedPtr<SDockTab> ForegroundTabPtr = TabWell->GetForegroundTab())
+	{
+		const TSharedRef<SDockTab> ForegroundTabRef = ForegroundTabPtr.ToSharedRef();
+		SetTabLocked(ForegroundTabRef, !IsTabLocked(ForegroundTabRef));
+	}
+}
+
 void SDockingTabStack::MoveTabToSidebar(TSharedRef<SDockTab> Tab)
 {
 	const int32 TabIndex = Tabs.IndexOfByPredicate(FTabMatcher(Tab->GetLayoutIdentifier(), ETabState::OpenedTab));
@@ -1191,12 +1249,23 @@ bool SDockingTabStack::IsTabAllowedInSidebar(TSharedPtr<SDockTab> Tab) const
 	{
 		if (TSharedPtr<FTabManager> TabManager = Tab->GetTabManagerPtr())
 		{
-			return Tab->GetVisualTabRole() != ETabRole::MajorTab && TabManager->IsTabAllowedInSidebar(Tab->GetLayoutIdentifier());
+			return !Tab->IsTabLocked() && Tab->GetVisualTabRole() != ETabRole::MajorTab && TabManager->IsTabAllowedInSidebar(Tab->GetLayoutIdentifier());
 		}
 	}
 	return false;
 }
 
+bool SDockingTabStack::CanChangeTabLock() const
+{
+	if (const TSharedPtr<SDockTab>& ForegroundTabPtr = TabWell->GetForegroundTab())
+	{
+		const FTabId& ForegroundTabId = ForegroundTabPtr->GetLayoutIdentifier();
+		// Currently only the HomeScreen and the LevelEditor tab cannot change their Locked state.
+		return ForegroundTabId.TabType != FName(TEXT("HomeScreen")) && ForegroundTabId.TabType != FName(TEXT("LevelEditor"));
+	}
+	return false;
+}
+
 SSplitter::ESizeRule SDockingTabStack::GetSizeRule() const
 {
 	int32 NumTabs = this->GetNumTabs();
diff --git a/Engine/Source/Runtime/Slate/Private/Framework/Docking/SDockingTabStack.h b/Engine/Source/Runtime/Slate/Private/Framework/Docking/SDockingTabStack.h
--- a/Engine/Source/Runtime/Slate/Private/Framework/Docking/SDockingTabStack.h
+++ b/Engine/Source/Runtime/Slate/Private/Framework/Docking/SDockingTabStack.h
@@ -83,6 +83,9 @@ public:
 	SLATE_API bool IsTabPinnedInSidebar(const TSharedRef<SDockTab>& InTab);
 	SLATE_API void SetTabPinnedInSidebar(const TSharedRef<SDockTab>& InTab, bool bPinnedInSidebar);
 
+	SLATE_API bool IsTabLocked(const TSharedRef<const SDockTab>& InTab) const;
+	SLATE_API void SetTabLocked(const TSharedRef<const SDockTab>& InTab, TOptional<bool> bInLocked);
+
 	/** @return true if the specified tab can be moved to a sidebar */
 	SLATE_API bool CanMoveTabToSideBar(TSharedRef<SDockTab> Tab) const;
 
@@ -242,13 +245,22 @@ private:
 
 	/** Show the stuff needed to unhide the tab well? */
 	SLATE_API EVisibility GetUnhideButtonVisibility() const;
-	
+
+	/** The text to display for the Lock Tab action, dependent on the current state */
+	SLATE_API FText GetLockTabText() const;
+
+	/** The tooltip to display for the Lock Tab action, dependent on the current state  */
+	SLATE_API FText GetLockTabToolTip() const;
+
 	/** Show/Hide the tab well; do it smoothly with an animation */
 	SLATE_API void ToggleTabWellVisibility();
 
 	/** Moves the foreground tab to a sidebar */
 	SLATE_API void MoveForegroundTabToSidebar();
 
+	/** Show/Hide the tab well; do it smoothly with an animation */
+	SLATE_API void ToggleForegroundTabLocked();
+
 	/** Unhides the tab well, revealing all tab headers */
 	SLATE_API FReply UnhideTabWell();
 
@@ -269,6 +281,9 @@ private:
 	/** @return true if the specified tab can be moved to a sidebar */
 	SLATE_API bool IsTabAllowedInSidebar(TSharedPtr<SDockTab> Tab) const;
 
+	/** @return true if the specified tab can ever be locked/unlocked, currently only the HomeScreen and the LevelEditor can't be locked/unlocked from users */
+	SLATE_API bool CanChangeTabLock() const;
+
 	/** The tab well widget shows all tabs, keeps track of the selected tab, allows tab rearranging, etc. */
 	TSharedPtr<class SDockingTabWell> TabWell;
 
diff --git a/Engine/Source/Runtime/Slate/Private/Framework/Docking/TabManager.cpp b/Engine/Source/Runtime/Slate/Private/Framework/Docking/TabManager.cpp
--- a/Engine/Source/Runtime/Slate/Private/Framework/Docking/TabManager.cpp
+++ b/Engine/Source/Runtime/Slate/Private/Framework/Docking/TabManager.cpp
@@ -315,7 +315,14 @@ TSharedRef<FTabManager::FLayoutNode> FTabManager::FLayout::NewFromString_Helper(
 				TabAsJson->TryGetBoolField(TEXT("SidebarPinned"), bPinnedInSidebar);
 			}
 
-			NewStack->AddTab(TabId, TabStateFromString( TabAsJson->GetStringField(TEXT("TabState"))), SidebarLocationFromString(SidebarLocation), SidebarSizeCoefficient, bPinnedInSidebar);
+			bool bLockedInPlace = true;
+			TOptional<bool> bLockedInPlaceOptional = TOptional<bool>();
+			if (TabAsJson->TryGetBoolField(TEXT("LockedInPlace"), bLockedInPlace))
+			{
+				bLockedInPlaceOptional = bLockedInPlace;
+			}
+
+			NewStack->AddTab(TabId, TabStateFromString( TabAsJson->GetStringField(TEXT("TabState"))), SidebarLocationFromString(SidebarLocation), SidebarSizeCoefficient, bPinnedInSidebar, bLockedInPlaceOptional);
 		}
 		return NewStack;
 	}
@@ -488,6 +495,12 @@ TSharedRef<FJsonObject> FTabManager::FLayout::PersistToString_Helper(const TShar
 					TabAsJson->SetBoolField(TEXT("SidebarPinned"), Tab.bPinnedInSidebar);
 				}
 
+				// Save it only if set
+				if (Tab.bIsLockedInPlace.IsSet())
+				{
+					TabAsJson->SetBoolField("LockedInPlace", Tab.bIsLockedInPlace.GetValue());
+				}
+
 				TabsAsJson.Add( MakeShareable( new FJsonValueObject(TabAsJson) ) );
 			}
 		}
@@ -1790,12 +1803,13 @@ TSharedPtr<SDockingNode> FTabManager::RestoreArea_Helper(const TSharedRef<FLayou
 			NewStackWidget->SetSizeCoefficient(LayoutNode->GetSizeCoefficient());
 		}
 		// Open Tabs
-		for (const FTab& SomeTab : NodeAsStack->Tabs)
+		for (FTab& SomeTab : NodeAsStack->Tabs)
 		{
 			if ((SomeTab.TabState == ETabState::OpenedTab || SomeTab.TabState == ETabState::SidebarTab) && IsValidTabForSpawning(SomeTab))
 			{
 				const bool bCanUnrecognizedTabBeNullptr = true;
 				const TSharedPtr<SDockTab> NewTabWidget = SpawnTab(SomeTab.TabId, ParentWindow, bCanUnrecognizedTabBeNullptr);
+				TSharedPtr<FTabSpawnerEntry> Spawner = FindTabSpawnerFor(SomeTab.TabId.TabType);
 
 				if (NewTabWidget.IsValid())
 				{
@@ -1812,6 +1826,13 @@ TSharedPtr<SDockingNode> FTabManager::RestoreArea_Helper(const TSharedRef<FLayou
 						NewStackWidget->SetSizeCoefficient(LayoutNode->GetSizeCoefficient());
 					}
 
+					// If the config didn't have a locked state saved check if the spawner one is set and use that one instead
+					// The config have priority over the Spawner one since it should be used as a default value
+					if (!SomeTab.bIsLockedInPlace.IsSet() && Spawner->bIsLocked.IsSet())
+					{
+						NewStackWidget->SetTabLocked(NewTabWidget.ToSharedRef(), Spawner->bIsLocked.GetValue());
+					}
+
 					if (SomeTab.TabState == ETabState::OpenedTab)
 					{
 						NewStackWidget->AddTabWidget(NewTabWidget.ToSharedRef());
@@ -2642,7 +2663,7 @@ bool FTabManager::CanCloseManager( const TSet< TSharedRef<SDockTab> >& TabsToIgn
 			bCanCloseManager =
 				TabsToIgnore.Contains( AreasTabs[TabIndex] ) ||
 				AreasTabs[TabIndex]->GetTabRole() != ETabRole::MajorTab ||
-				AreasTabs[TabIndex]->CanCloseTab();
+				AreasTabs[TabIndex]->CanCloseTab(/** Ignore locked tabs */ true);
 		}		
 	}
 
diff --git a/Engine/Source/Runtime/Slate/Private/Widgets/Docking/SDockTab.cpp b/Engine/Source/Runtime/Slate/Private/Widgets/Docking/SDockTab.cpp
--- a/Engine/Source/Runtime/Slate/Private/Widgets/Docking/SDockTab.cpp
+++ b/Engine/Source/Runtime/Slate/Private/Widgets/Docking/SDockTab.cpp
@@ -75,6 +75,11 @@ FReply SDockTab::OnMouseButtonDoubleClick( const FGeometry& MyGeometry, const FP
 
 FReply SDockTab::OnDragDetected( const FGeometry& MyGeometry, const FPointerEvent& MouseEvent  )
 {
+	if (IsTabLocked())
+	{
+		return FReply::Unhandled();
+	}
+
 	// Need to remember where within a tab we grabbed
 	const FVector2D TabGrabOffset = MyGeometry.AbsoluteToLocal( MouseEvent.GetScreenSpacePosition() );
 	const FVector2D TabSize = MyGeometry.GetLocalSize();
@@ -988,7 +993,7 @@ EVisibility SDockTab::HandleIsCloseButtonVisible() const
 	TSharedPtr<FTabManager> Manager = MyTabManager.Pin();
 	const bool bIsVisibled = IsHovered() || IsForeground();
 	const bool bIsTabCloseable = Manager == nullptr || Manager->IsTabCloseable(SharedThis(this));
-	return bCanEverClose && bIsVisibled && bIsTabCloseable ? EVisibility::Visible : EVisibility::Hidden;
+	return bCanEverClose && !IsTabLocked() && bIsVisibled && bIsTabCloseable ? EVisibility::Visible : EVisibility::Hidden;
 }
 
 TOptional<FVector2D> SDockTab::GetTabIconSize() const
@@ -997,6 +1002,11 @@ TOptional<FVector2D> SDockTab::GetTabIconSize() const
 }
 
 bool SDockTab::CanCloseTab() const
+{
+	return CanCloseTab(false);
+}
+
+bool SDockTab::CanCloseTab(bool bIgnoreLockedTabs) const
 {
 	// Check the below conditions separately and early out to prevent the OnCanCloseTab delegate
 	// from being called if we already know we cannot close the tab.
@@ -1011,7 +1021,15 @@ bool SDockTab::CanCloseTab() const
 	{
 		return false;
 	}
-	
+
+	if (!bIgnoreLockedTabs)
+	{
+		if (IsTabLocked())
+		{
+			return false;
+		}
+	}
+
 	if (OnCanCloseTab.IsBound() && !OnCanCloseTab.Execute())
 	{
 		return false;
@@ -1045,6 +1063,12 @@ bool SDockTab::RequestCloseTab()
 	return bCanCloseTabNow;
 }
 
+bool SDockTab::IsTabLocked() const
+{
+	const TSharedPtr<SDockingTabStack>& StackTab = GetParentDockTabStack();
+	return StackTab.IsValid() ? StackTab->IsTabLocked(SharedThis(this)) : false;
+}
+
 void SDockTab::PersistVisualState()
 {
 	OnPersistVisualState.ExecuteIfBound();
diff --git a/Engine/Source/Runtime/Slate/Public/Framework/Docking/TabManager.h b/Engine/Source/Runtime/Slate/Public/Framework/Docking/TabManager.h
--- a/Engine/Source/Runtime/Slate/Public/Framework/Docking/TabManager.h
+++ b/Engine/Source/Runtime/Slate/Public/Framework/Docking/TabManager.h
@@ -308,11 +308,26 @@ struct FTabSpawnerEntry : public FWorkspaceItem
 		return *this;
 	}
 
+	/** Sets the tab lock state default value, where a locked tab cannot be undocked or closed.\n
+	 *	Locked state in spawn arg have priority on this one, if you use it in both places the spawn arg one will be the one set.\n
+	 *	This property is also saved and loaded, if saved the loaded value will be the one used.
+	 */
+	FTabSpawnerEntry& SetIsTabLocked(bool bInIsLocked)
+	{
+		bIsLocked = bInIsLocked;
+		return *this;
+	}
+
 	bool CanSidebarTab() const
 	{
 		return bCanSidebarTab;
 	}
 
+	TOptional<bool> IsTabLocked() const
+	{
+		return bIsLocked;
+	}
+
 	bool IsTabNameHidden() const
 	{
 		return bHiddenTabName;
@@ -345,7 +360,10 @@ private:
 	bool bAutoGenerateMenuEntry;
 	/** Whether or not this tab can ever be in a sidebar */
 	bool bCanSidebarTab;
+	/** Whether the tab name should be hidden */
 	bool bHiddenTabName;
+	/** Whether the tab default state should be locked or not */
+	TOptional<bool> bIsLocked;
 
 	TWeakPtr<SDockTab> SpawnedTabPtr;
 
@@ -446,16 +464,18 @@ class FTabManager : public TSharedFromThis<FTabManager>
 				, SidebarLocation(ESidebarLocation::None)
 				, SidebarSizeCoefficient(0.0f)
 				, bPinnedInSidebar(false)
+				, bIsLockedInPlace(TOptional<bool>())
 			{
 				check(InTabState != ETabState::SidebarTab);
 			}
 
-			FTab(const FTabId& InTabId, ETabState::Type InTabState, ESidebarLocation InSidebarLocation, float InSidebarSizeCoefficient, bool bInPinnedInSidebar)
+			FTab(const FTabId& InTabId, ETabState::Type InTabState, ESidebarLocation InSidebarLocation, float InSidebarSizeCoefficient, bool bInPinnedInSidebar, TOptional<bool> bInLockedInPlace)
 				: TabId(InTabId)
 				, TabState(InTabState)
 				, SidebarLocation(InSidebarLocation)
 				, SidebarSizeCoefficient(InSidebarSizeCoefficient)
 				, bPinnedInSidebar(bInPinnedInSidebar)
+				, bIsLockedInPlace(bInLockedInPlace)
 			{
 				check(InTabState != ETabState::SidebarTab || InSidebarLocation != ESidebarLocation::None);
 			}
@@ -470,6 +490,7 @@ class FTabManager : public TSharedFromThis<FTabManager>
 			ESidebarLocation SidebarLocation;
 			float SidebarSizeCoefficient;
 			bool bPinnedInSidebar;
+			TOptional<bool> bIsLockedInPlace;
 		};
 
 		class FStack : public FLayoutNode
@@ -495,18 +516,18 @@ class FTabManager : public TSharedFromThis<FTabManager>
 					return SharedThis(this);
 				}
 
-				TSharedRef<FStack> AddTab(const FName TabType, ETabState::Type InTabState, ESidebarLocation InSidebarLocation, float SidebarSizeCoefficient, bool bPinnedInSidebar=false)
+				TSharedRef<FStack> AddTab(const FName TabType, ETabState::Type InTabState, ESidebarLocation InSidebarLocation, float SidebarSizeCoefficient, bool bPinnedInSidebar=false, TOptional<bool> bLockedInPlace = TOptional<bool>())
 				{
 					check(InTabState != ETabState::SidebarTab || InSidebarLocation != ESidebarLocation::None);
-					Tabs.Add(FTab(FTabId(TabType), InTabState, InSidebarLocation, SidebarSizeCoefficient, bPinnedInSidebar));
+					Tabs.Add(FTab(FTabId(TabType), InTabState, InSidebarLocation, SidebarSizeCoefficient, bPinnedInSidebar, bLockedInPlace));
 					return SharedThis(this);
 				}
 
-				TSharedRef<FStack> AddTab(const FTabId TabId, ETabState::Type InTabState, ESidebarLocation InSidebarLocation, float SidebarSizeCoefficient, bool bPinnedInSidebar=false)
+				TSharedRef<FStack> AddTab(const FTabId TabId, ETabState::Type InTabState, ESidebarLocation InSidebarLocation, float SidebarSizeCoefficient, bool bPinnedInSidebar=false, TOptional<bool> bLockedInPlace = TOptional<bool>())
 				{
 					check(InTabState != ETabState::SidebarTab || InSidebarLocation != ESidebarLocation::None);
 
-					Tabs.Add(FTab(TabId, InTabState, InSidebarLocation, SidebarSizeCoefficient, bPinnedInSidebar));
+					Tabs.Add(FTab(TabId, InTabState, InSidebarLocation, SidebarSizeCoefficient, bPinnedInSidebar, bLockedInPlace));
 
 					return SharedThis(this);
 				}
diff --git a/Engine/Source/Runtime/Slate/Public/Widgets/Docking/SDockTab.h b/Engine/Source/Runtime/Slate/Public/Widgets/Docking/SDockTab.h
--- a/Engine/Source/Runtime/Slate/Public/Widgets/Docking/SDockTab.h
+++ b/Engine/Source/Runtime/Slate/Public/Widgets/Docking/SDockTab.h
@@ -242,9 +242,15 @@ public:
 	/** @return true if the tab can be closed right now. Example: Callback could ask user via dialog. */
 	SLATE_API bool CanCloseTab() const;
 
+	/** @return true if the tab can be closed right now. Example: Callback could ask user via dialog. */
+	SLATE_API bool CanCloseTab(bool bIgnoreLockedTabs) const;
+
 	/** Requests that the tab be closed.  Tabs may prevent closing depending on their state */	
 	SLATE_API bool RequestCloseTab();
 
+	/** @return True if this tab is currently locked */
+	SLATE_API bool IsTabLocked() const;
+
 	/** A chance for the tab's content to save any internal layout info */
 	SLATE_API void PersistVisualState();
 
