[Slate] TabManager and DockTabStack comment fix

diff --git a/Engine/Source/Runtime/Slate/Private/Framework/Docking/SDockingTabStack.h b/Engine/Source/Runtime/Slate/Private/Framework/Docking/SDockingTabStack.h
--- a/Engine/Source/Runtime/Slate/Private/Framework/Docking/SDockingTabStack.h
+++ b/Engine/Source/Runtime/Slate/Private/Framework/Docking/SDockingTabStack.h
@@ -258,7 +258,7 @@ private:
 	/** Moves the foreground tab to a sidebar */
 	SLATE_API void MoveForegroundTabToSidebar();
 
-	/** Show/Hide the tab well; do it smoothly with an animation */
+	/** Toggle the locked state of tabs */
 	SLATE_API void ToggleForegroundTabLocked();
 
 	/** Unhides the tab well, revealing all tab headers */
@@ -281,7 +281,7 @@ private:
 	/** @return true if the specified tab can be moved to a sidebar */
 	SLATE_API bool IsTabAllowedInSidebar(TSharedPtr<SDockTab> Tab) const;
 
-	/** @return true if the specified tab can ever be locked/unlocked, currently only the HomeScreen and the LevelEditor can't be locked/unlocked from users */
+	/** @return true if the current foreground tab can ever be locked/unlocked, currently only the HomeScreen and the LevelEditor can't be locked/unlocked from users */
 	SLATE_API bool CanChangeTabLock() const;
 
 	/** The tab well widget shows all tabs, keeps track of the selected tab, allows tab rearranging, etc. */
diff --git a/Engine/Source/Runtime/Slate/Public/Framework/Docking/TabManager.h b/Engine/Source/Runtime/Slate/Public/Framework/Docking/TabManager.h
--- a/Engine/Source/Runtime/Slate/Public/Framework/Docking/TabManager.h
+++ b/Engine/Source/Runtime/Slate/Public/Framework/Docking/TabManager.h
@@ -309,7 +309,6 @@ struct FTabSpawnerEntry : public FWorkspaceItem
 	}
 
 	/** Sets the tab lock state default value, where a locked tab cannot be undocked or closed.\n
-	 *	Locked state in spawn arg have priority on this one, if you use it in both places the spawn arg one will be the one set.\n
 	 *	This property is also saved and loaded, if saved the loaded value will be the one used.
 	 */
 	FTabSpawnerEntry& SetIsTabLocked(bool bInIsLocked)
