[TabManager] Hide Locked feature from users and remove API

diff --git a/Engine/Source/Editor/LevelEditor/Private/LevelEditor.cpp b/Engine/Source/Editor/LevelEditor/Private/LevelEditor.cpp
--- a/Engine/Source/Editor/LevelEditor/Private/LevelEditor.cpp
+++ b/Engine/Source/Editor/LevelEditor/Private/LevelEditor.cpp
@@ -317,7 +317,6 @@ void FLevelEditorModule::StartupModule()
 	const IWorkspaceMenuStructure& MenuStructure = WorkspaceMenu::GetMenuStructure();
 
 	FGlobalTabmanager::Get()->RegisterTabSpawner("LevelEditor", FOnSpawnTab::CreateRaw(this, &FLevelEditorModule::SpawnLevelEditor))
-		.SetIsTabLocked(true)
 		.SetDisplayName(NSLOCTEXT("LevelEditor", "LevelEditorTab", "Level Editor"))
 		.SetAutoGenerateMenuEntry(false);
 
diff --git a/Engine/Source/Editor/MainFrame/Private/MainFrameModule.cpp b/Engine/Source/Editor/MainFrame/Private/MainFrameModule.cpp
--- a/Engine/Source/Editor/MainFrame/Private/MainFrameModule.cpp
+++ b/Engine/Source/Editor/MainFrame/Private/MainFrameModule.cpp
@@ -686,7 +686,6 @@ void FMainFrameModule::StartupModule( )
     		FGlobalTabmanager::Get()->RegisterTabSpawner(
     			"HomeScreen", FOnSpawnTab::CreateRaw(this, &FMainFrameModule::SpawnHomeScreenWindow))
     			.SetAutoGenerateMenuEntry(false)
-				.SetIsTabLocked(true)
     			.SetMenuType(ETabSpawnerMenuType::Hidden)
     			.SetDisplayName(FText::GetEmpty());
     
diff --git a/Engine/Source/Runtime/Slate/Private/Framework/Docking/SDockingTabStack.cpp b/Engine/Source/Runtime/Slate/Private/Framework/Docking/SDockingTabStack.cpp
--- a/Engine/Source/Runtime/Slate/Private/Framework/Docking/SDockingTabStack.cpp
+++ b/Engine/Source/Runtime/Slate/Private/Framework/Docking/SDockingTabStack.cpp
@@ -701,16 +701,6 @@ TSharedRef<SWidget> SDockingTabStack::MakeContextMenu()
 					)
 				);
 			}
-
-			MenuBuilder.AddMenuEntry(
-				TAttribute<FText>::CreateSP(this, &SDockingTabStack::GetLockTabText),
-				TAttribute<FText>::CreateSP(this, &SDockingTabStack::GetLockTabToolTip),
-				FSlateIcon(),
-				FUIAction(
-					FExecuteAction::CreateSP(this, &SDockingTabStack::ToggleForegroundTabLocked),
-					FCanExecuteAction::CreateSP(this, &SDockingTabStack::CanChangeTabLock)
-				)
-			);
 		}
 		MenuBuilder.EndSection();
 
@@ -1059,26 +1049,6 @@ EVisibility SDockingTabStack::GetUnhideButtonVisibility() const
 		: EVisibility::Collapsed;
 }
 
-FText SDockingTabStack::GetLockTabText() const
-{
-	if (TSharedPtr<SDockTab> ForegroundTabPtr = TabWell->GetForegroundTab())
-	{
-		return IsTabLocked(ForegroundTabPtr.ToSharedRef()) ? LOCTEXT("UnlockTabText", "Unlock Tab") : LOCTEXT("LockTabText", "Lock Tab");
-	}
-
-	return FText::GetEmpty();
-}
-
-FText SDockingTabStack::GetLockTabToolTip() const
-{
-	if (TSharedPtr<SDockTab> ForegroundTabPtr = TabWell->GetForegroundTab())
-	{
-		return IsTabLocked(ForegroundTabPtr.ToSharedRef()) ? LOCTEXT("UnlockTabToolTip", "Unlocks the tab from the current location") : LOCTEXT("LockTabToolTip", "Locks the tab to the current location and prevents moving/closing");
-	}
-
-	return FText::GetEmpty();
-}
-
 void SDockingTabStack::ToggleTabWellVisibility()
 {
 	ShowHideTabWell.Reverse();
@@ -1092,15 +1062,6 @@ void SDockingTabStack::MoveForegroundTabToSidebar()
 	}
 }
 
-void SDockingTabStack::ToggleForegroundTabLocked()
-{
-	if (const TSharedPtr<SDockTab> ForegroundTabPtr = TabWell->GetForegroundTab())
-	{
-		const TSharedRef<SDockTab> ForegroundTabRef = ForegroundTabPtr.ToSharedRef();
-		SetTabLocked(ForegroundTabRef, !IsTabLocked(ForegroundTabRef));
-	}
-}
-
 void SDockingTabStack::MoveTabToSidebar(TSharedRef<SDockTab> Tab)
 {
 	const int32 TabIndex = Tabs.IndexOfByPredicate(FTabMatcher(Tab->GetLayoutIdentifier(), ETabState::OpenedTab));
@@ -1255,17 +1216,6 @@ bool SDockingTabStack::IsTabAllowedInSidebar(TSharedPtr<SDockTab> Tab) const
 	return false;
 }
 
-bool SDockingTabStack::CanChangeTabLock() const
-{
-	if (const TSharedPtr<SDockTab>& ForegroundTabPtr = TabWell->GetForegroundTab())
-	{
-		const FTabId& ForegroundTabId = ForegroundTabPtr->GetLayoutIdentifier();
-		// Currently only the HomeScreen and the LevelEditor tab cannot change their Locked state.
-		return ForegroundTabId.TabType != FName(TEXT("HomeScreen")) && ForegroundTabId.TabType != FName(TEXT("LevelEditor"));
-	}
-	return false;
-}
-
 SSplitter::ESizeRule SDockingTabStack::GetSizeRule() const
 {
 	int32 NumTabs = this->GetNumTabs();
diff --git a/Engine/Source/Runtime/Slate/Private/Framework/Docking/SDockingTabStack.h b/Engine/Source/Runtime/Slate/Private/Framework/Docking/SDockingTabStack.h
--- a/Engine/Source/Runtime/Slate/Private/Framework/Docking/SDockingTabStack.h
+++ b/Engine/Source/Runtime/Slate/Private/Framework/Docking/SDockingTabStack.h
@@ -246,21 +246,12 @@ private:
 	/** Show the stuff needed to unhide the tab well? */
 	SLATE_API EVisibility GetUnhideButtonVisibility() const;
 
-	/** The text to display for the Lock Tab action, dependent on the current state */
-	SLATE_API FText GetLockTabText() const;
-
-	/** The tooltip to display for the Lock Tab action, dependent on the current state  */
-	SLATE_API FText GetLockTabToolTip() const;
-
 	/** Show/Hide the tab well; do it smoothly with an animation */
 	SLATE_API void ToggleTabWellVisibility();
 
 	/** Moves the foreground tab to a sidebar */
 	SLATE_API void MoveForegroundTabToSidebar();
 
-	/** Toggle the locked state of tabs */
-	SLATE_API void ToggleForegroundTabLocked();
-
 	/** Unhides the tab well, revealing all tab headers */
 	SLATE_API FReply UnhideTabWell();
 
@@ -281,9 +272,6 @@ private:
 	/** @return true if the specified tab can be moved to a sidebar */
 	SLATE_API bool IsTabAllowedInSidebar(TSharedPtr<SDockTab> Tab) const;
 
-	/** @return true if the current foreground tab can ever be locked/unlocked, currently only the HomeScreen and the LevelEditor can't be locked/unlocked from users */
-	SLATE_API bool CanChangeTabLock() const;
-
 	/** The tab well widget shows all tabs, keeps track of the selected tab, allows tab rearranging, etc. */
 	TSharedPtr<class SDockingTabWell> TabWell;
 
diff --git a/Engine/Source/Runtime/Slate/Private/Framework/Docking/TabManager.cpp b/Engine/Source/Runtime/Slate/Private/Framework/Docking/TabManager.cpp
--- a/Engine/Source/Runtime/Slate/Private/Framework/Docking/TabManager.cpp
+++ b/Engine/Source/Runtime/Slate/Private/Framework/Docking/TabManager.cpp
@@ -315,14 +315,8 @@ TSharedRef<FTabManager::FLayoutNode> FTabManager::FLayout::NewFromString_Helper(
 				TabAsJson->TryGetBoolField(TEXT("SidebarPinned"), bPinnedInSidebar);
 			}
 
-			bool bLockedInPlace = true;
-			TOptional<bool> bLockedInPlaceOptional = TOptional<bool>();
-			if (TabAsJson->TryGetBoolField(TEXT("LockedInPlace"), bLockedInPlace))
-			{
-				bLockedInPlaceOptional = bLockedInPlace;
-			}
-
-			NewStack->AddTab(TabId, TabStateFromString( TabAsJson->GetStringField(TEXT("TabState"))), SidebarLocationFromString(SidebarLocation), SidebarSizeCoefficient, bPinnedInSidebar, bLockedInPlaceOptional);
+			// For now always have the tab unlocked until the feature is made public, not assigning a value to it will by default keep them unlocked but also allow us internally to have the LevelEditor and HomeScreen tab decide their own state
+			NewStack->AddTab(TabId, TabStateFromString( TabAsJson->GetStringField(TEXT("TabState"))), SidebarLocationFromString(SidebarLocation), SidebarSizeCoefficient, bPinnedInSidebar);
 		}
 		return NewStack;
 	}
@@ -495,12 +489,6 @@ TSharedRef<FJsonObject> FTabManager::FLayout::PersistToString_Helper(const TShar
 					TabAsJson->SetBoolField(TEXT("SidebarPinned"), Tab.bPinnedInSidebar);
 				}
 
-				// Save it only if set
-				if (Tab.bIsLockedInPlace.IsSet())
-				{
-					TabAsJson->SetBoolField("LockedInPlace", Tab.bIsLockedInPlace.GetValue());
-				}
-
 				TabsAsJson.Add( MakeShareable( new FJsonValueObject(TabAsJson) ) );
 			}
 		}
diff --git a/Engine/Source/Runtime/Slate/Public/Framework/Docking/TabManager.h b/Engine/Source/Runtime/Slate/Public/Framework/Docking/TabManager.h
--- a/Engine/Source/Runtime/Slate/Public/Framework/Docking/TabManager.h
+++ b/Engine/Source/Runtime/Slate/Public/Framework/Docking/TabManager.h
@@ -235,11 +235,18 @@ struct FTabSpawnerEntry : public FWorkspaceItem
 		, SpawnedTabPtr()
 		, ReadOnlyBehavior(ETabReadOnlyBehavior::Disabled)
 	{
-		// Special treatment for the HomeScreen tab to not have its name displayed, this is currently intended to not be a public feature
+		// Special treatment for the HomeScreen tab to not have its name displayed and to be locked, this is currently intended to not be a public feature
 		if (InTabType == FName(TEXT("HomeScreen")))
 		{
 			bHiddenTabName = true;
+			bIsLocked = true;
 		}
+		
+		// Special treatment for the LevelEditor tab to be locked, this is currently intended to not be a public feature
+        if (InTabType == FName(TEXT("LevelEditor")))
+        {
+        	bIsLocked = true;
+        }
 	}
 
 	FTabSpawnerEntry& SetReadOnlyBehavior( const ETabReadOnlyBehavior& InReadOnlyBehavior)
@@ -308,15 +315,6 @@ struct FTabSpawnerEntry : public FWorkspaceItem
 		return *this;
 	}
 
-	/** Sets the tab lock state default value, where a locked tab cannot be undocked or closed.\n
-	 *	This property is also saved and loaded, if saved the loaded value will be the one used.
-	 */
-	FTabSpawnerEntry& SetIsTabLocked(bool bInIsLocked)
-	{
-		bIsLocked = bInIsLocked;
-		return *this;
-	}
-
 	bool CanSidebarTab() const
 	{
 		return bCanSidebarTab;
@@ -468,13 +466,13 @@ class FTabManager : public TSharedFromThis<FTabManager>
 				check(InTabState != ETabState::SidebarTab);
 			}
 
-			FTab(const FTabId& InTabId, ETabState::Type InTabState, ESidebarLocation InSidebarLocation, float InSidebarSizeCoefficient, bool bInPinnedInSidebar, TOptional<bool> bInLockedInPlace)
+			FTab(const FTabId& InTabId, ETabState::Type InTabState, ESidebarLocation InSidebarLocation, float InSidebarSizeCoefficient, bool bInPinnedInSidebar)
 				: TabId(InTabId)
 				, TabState(InTabState)
 				, SidebarLocation(InSidebarLocation)
 				, SidebarSizeCoefficient(InSidebarSizeCoefficient)
 				, bPinnedInSidebar(bInPinnedInSidebar)
-				, bIsLockedInPlace(bInLockedInPlace)
+				, bIsLockedInPlace(TOptional<bool>())
 			{
 				check(InTabState != ETabState::SidebarTab || InSidebarLocation != ESidebarLocation::None);
 			}
@@ -515,18 +513,18 @@ class FTabManager : public TSharedFromThis<FTabManager>
 					return SharedThis(this);
 				}
 
-				TSharedRef<FStack> AddTab(const FName TabType, ETabState::Type InTabState, ESidebarLocation InSidebarLocation, float SidebarSizeCoefficient, bool bPinnedInSidebar=false, TOptional<bool> bLockedInPlace = TOptional<bool>())
+				TSharedRef<FStack> AddTab(const FName TabType, ETabState::Type InTabState, ESidebarLocation InSidebarLocation, float SidebarSizeCoefficient, bool bPinnedInSidebar=false)
 				{
 					check(InTabState != ETabState::SidebarTab || InSidebarLocation != ESidebarLocation::None);
-					Tabs.Add(FTab(FTabId(TabType), InTabState, InSidebarLocation, SidebarSizeCoefficient, bPinnedInSidebar, bLockedInPlace));
+					Tabs.Add(FTab(FTabId(TabType), InTabState, InSidebarLocation, SidebarSizeCoefficient, bPinnedInSidebar));
 					return SharedThis(this);
 				}
 
-				TSharedRef<FStack> AddTab(const FTabId TabId, ETabState::Type InTabState, ESidebarLocation InSidebarLocation, float SidebarSizeCoefficient, bool bPinnedInSidebar=false, TOptional<bool> bLockedInPlace = TOptional<bool>())
+				TSharedRef<FStack> AddTab(const FTabId TabId, ETabState::Type InTabState, ESidebarLocation InSidebarLocation, float SidebarSizeCoefficient, bool bPinnedInSidebar=false)
 				{
 					check(InTabState != ETabState::SidebarTab || InSidebarLocation != ESidebarLocation::None);
 
-					Tabs.Add(FTab(TabId, InTabState, InSidebarLocation, SidebarSizeCoefficient, bPinnedInSidebar, bLockedInPlace));
+					Tabs.Add(FTab(TabId, InTabState, InSidebarLocation, SidebarSizeCoefficient, bPinnedInSidebar));
 
 					return SharedThis(this);
 				}
