Disable slate drag and drops from conflicting with each other

#rnx
#FailedValidations Preflight

diff --git a/Engine/Plugins/Fab/Source/Fab/Private/FabBrowser.cpp b/Engine/Plugins/Fab/Source/Fab/Private/FabBrowser.cpp
--- a/Engine/Plugins/Fab/Source/Fab/Private/FabBrowser.cpp
+++ b/Engine/Plugins/Fab/Source/Fab/Private/FabBrowser.cpp
@@ -368,6 +368,10 @@ TSharedRef<SDockTab> FFabBrowser::OpenTab(const FSpawnTabArgs&)
 		WebBrowserInstance.ToSharedRef()
 	];
 	
+#if PLATFORM_MAC && !WITH_CEF3
+	FSlateApplication::Get().ToggleDisableLastDragOnDragEnter(true);
+#endif
+	
 	WebBrowserWindow->SetParentDockTab(DockTab);
 
 	return DockTab.ToSharedRef();
@@ -405,6 +409,10 @@ void FFabBrowser::OnPluginTabClosed(TSharedRef<SDockTab> InParentTab)
 			WebBrowserInstance.Reset();
 			WebBrowserWindow.Reset();
 			DockTab.Reset();
+			
+#if PLATFORM_MAC && !WITH_CEF3
+			FSlateApplication::Get().ToggleDisableLastDragOnDragEnter(false);
+#endif
 		});
 	});
 }
diff --git a/Engine/Source/Runtime/Slate/Private/Framework/Application/SlateApplication.cpp b/Engine/Source/Runtime/Slate/Private/Framework/Application/SlateApplication.cpp
--- a/Engine/Source/Runtime/Slate/Private/Framework/Application/SlateApplication.cpp
+++ b/Engine/Source/Runtime/Slate/Private/Framework/Application/SlateApplication.cpp
@@ -3361,7 +3361,7 @@ void FSlateApplication::ProcessReply( const FWidgetPath& CurrentEventPath, const
 		}
 	}
 
-	if (TheReply.ShouldEndDragDrop())
+	if ((bDisableLastDragOnDragEnter && bStartingDragDrop && SlateUser->IsDragDropping()) || TheReply.ShouldEndDragDrop())
 	{
 		SlateUser->CancelDragDrop();
 	}
@@ -3986,6 +3986,16 @@ void FSlateApplication::SetDragTriggerDistance( float ScreenPixels )
 	DragTriggerDistance = ScreenPixels;
 }
 
+void FSlateApplication::ToggleDisableLastDragOnDragEnter( bool bInDisableLastDragOnDragEnter )
+{
+	bDisableLastDragOnDragEnter = bInDisableLastDragOnDragEnter;
+}
+
+bool FSlateApplication::IsLastDragOnDragEnterDisabled() const
+{
+	return bDisableLastDragOnDragEnter;
+}
+
 bool FSlateApplication::RegisterInputPreProcessor(TSharedPtr<IInputProcessor> InputProcessor)
 {
 	return RegisterInputPreProcessor(InputProcessor, FInputPreprocessorRegistrationKey());
diff --git a/Engine/Source/Runtime/Slate/Public/Framework/Application/SlateApplication.h b/Engine/Source/Runtime/Slate/Public/Framework/Application/SlateApplication.h
--- a/Engine/Source/Runtime/Slate/Public/Framework/Application/SlateApplication.h
+++ b/Engine/Source/Runtime/Slate/Public/Framework/Application/SlateApplication.h
@@ -1485,6 +1485,9 @@ public:
 
 	/** @return the deadzone size squared for dragging in screen pixels (aka virtual desktop pixels) */
 	SLATE_API float GetDragTriggerDistanceSquared() const;
+	
+	/** @return Whether we disable the last drag on a new drag enter event to prevent double drag conflicts. */
+	SLATE_API bool IsLastDragOnDragEnterDisabled() const;
 
 	/** @return true if the difference between the ScreenSpaceOrigin and the ScreenSpacePosition is larger than the trigger distance for dragging in Slate. */
 	SLATE_API bool HasTraveledFarEnoughToTriggerDrag(const FPointerEvent& PointerEvent, const UE::Slate::FDeprecateVector2DParameter ScreenSpaceOrigin) const;
@@ -1492,6 +1495,9 @@ public:
 
 	/** Set the size of the deadzone for dragging in screen pixels */
 	SLATE_API void SetDragTriggerDistance( float ScreenPixels );
+	
+	/** Toggle disabling last drag on a new drag enter event to prevent double drag conflicts. */
+	SLATE_API void ToggleDisableLastDragOnDragEnter( bool bInDisableLastDragOnDragEnter );
 
 
 	/**
@@ -2169,6 +2175,9 @@ private:
 	
 	/** Are we ticking the SlateApplication. */
 	bool bIsTicking = false;
+	
+	/** Are we disabling the last drag on a new drag enter event to prevent double drag conflicts?  */
+	bool bDisableLastDragOnDragEnter = false;
 
 	/** Platform mouse movement event count. */
 	uint64 PlatformMouseMovementEvents = 0;
