[HomeScreen] Add HomeScreen as part of the DefaultLayout of the Engine

#FailedValidations Preflight

diff --git a/Engine/Source/Editor/MainFrame/Private/MainFrameModule.cpp b/Engine/Source/Editor/MainFrame/Private/MainFrameModule.cpp
--- a/Engine/Source/Editor/MainFrame/Private/MainFrameModule.cpp
+++ b/Engine/Source/Editor/MainFrame/Private/MainFrameModule.cpp
@@ -964,14 +964,12 @@ TSharedRef<SDockTab> FMainFrameModule::SpawnHomeScreenWindow(const FSpawnTabArgs
 		.IconColor(FStyleColors::White.GetSpecifiedColor())
 		.ContentPadding(FMargin(0));
 
-	if (!HomeScreenTabManager.IsValid())
-	{
-		HomeScreenTabManager = FGlobalTabmanager::Get()->NewTabManager(HomeScreenTab);
-		HomeScreenTabManager->SetAllowWindowMenuBar(true);
-		FToolMenuContext EmptyContext;
-		MakeMainMenu(HomeScreenTabManager, "HomeScreen.MainMenu", EmptyContext);
-		HomeScreenTab->SetOnTabClosed(SDockTab::FOnTabClosedCallback::CreateLambda([this] (TSharedRef<SDockTab>) { HomeScreenTabManager.Reset(); }));
-	}
+	// Always recreate the TabManager and bind it to the new HomeScreen tab, when loading layout for example the OnTabClosed is not called.
+	HomeScreenTabManager = FGlobalTabmanager::Get()->NewTabManager(HomeScreenTab);
+	HomeScreenTabManager->SetAllowWindowMenuBar(true);
+	FToolMenuContext EmptyContext;
+	MakeMainMenu(HomeScreenTabManager, "HomeScreen.MainMenu", EmptyContext);
+	HomeScreenTab->SetOnTabClosed(SDockTab::FOnTabClosedCallback::CreateLambda([this] (TSharedRef<SDockTab>) { HomeScreenTabManager.Reset(); }));
 
 	if (!WebObject.IsValid())
 	{
diff --git a/Engine/Source/Runtime/Slate/Private/Framework/Docking/TabManager.cpp b/Engine/Source/Runtime/Slate/Private/Framework/Docking/TabManager.cpp
--- a/Engine/Source/Runtime/Slate/Private/Framework/Docking/TabManager.cpp
+++ b/Engine/Source/Runtime/Slate/Private/Framework/Docking/TabManager.cpp
@@ -40,6 +40,23 @@ TAutoConsoleVariable<bool> CVarPanelDrawerToggle(
 	TEXT("Enables or disables the panel drawer. When the drawer is off invocations will just create a normal tab.")
 );
 
+namespace UE::Editor::HomeScreen
+{
+	static bool IsHomeScreenEnabled()
+	{
+		static bool bIsHomeScreenEnabled = false;
+		UE_CALL_ONCE([&]()
+		{
+			if (const IConsoleVariable* CVar = IConsoleManager::Get().FindConsoleVariable(TEXT("HomeScreen.EnableHomeScreen")))
+			{
+				ensureAlwaysMsgf(!EnumHasAnyFlags(CVar->GetFlags(), ECVF_Default), TEXT("The CVar should have already been set from commandline."));
+				bIsHomeScreenEnabled = CVar->GetBool();
+			}
+		});
+		return bIsHomeScreenEnabled;
+	}
+}
+
 #define LOCTEXT_NAMESPACE "TabManager"
 
 static const FString UE_TABMANAGER_OPENED_TAB_STRING = TEXT("OpenedTab");
@@ -1049,6 +1066,15 @@ TSharedPtr<SWidget> FTabManager::RestoreFrom(const TSharedRef<FLayout>& Layout,
 		// Set all InvalidTab tabs to OpenedTab so the Editor tries to load them. All non-recognized tabs will be set to InvalidTab later.
 		SetTabsTo(ThisArea, ETabState::OpenedTab, ETabState::InvalidTab);
 		const bool bIsPrimaryArea = ThisArea->WindowPlacement == FArea::Placement_NoWindow;
+
+		// Force the HomeScreen to be always added when restoring layouts in the primary area if enabled or just force the LevelEditor to be the first tab
+		// For now only fix the layout if the CVar is enabled
+		if (bIsPrimaryArea && UE::Editor::HomeScreen::IsHomeScreenEnabled())
+		{
+			FixLayoutLoadingPrimaryArea(ThisArea);
+		}
+
+		// Check the Area after it was cleared of possible HomeScreen tab that remained during testing in the layout config
 		const bool bShouldCreate = bIsPrimaryArea || HasValidTabs(ThisArea);
 
 		if ( bShouldCreate )
@@ -2951,6 +2977,73 @@ void FTabManager::GetRecordableStats( int32& OutTabCount, TArray<TSharedPtr<SWin
 	}
 }
 
+void FTabManager::FixLayoutLoadingPrimaryArea(const TSharedRef<FArea>& InPrimaryArea)
+{
+	FName LevelEditorTabId = FName(TEXT("LevelEditor"));
+	FName HomeScreenTabId = FName(TEXT("HomeScreen"));
+	const bool bIsHomeScreenEnabled = UE::Editor::HomeScreen::IsHomeScreenEnabled();
+
+	for (const TSharedRef<FLayoutNode>& Node : InPrimaryArea->GetChildNodes())
+	{
+		if (const TSharedPtr<FStack>& StackNode = Node->AsStack())
+		{
+			// Get the LevelEditor tab
+			FTab* LevelEditorTabPtr = StackNode->Tabs.FindByPredicate([LevelEditorTabId] (const FTab& InTab)
+				{
+					return InTab.TabId.TabType == LevelEditorTabId;
+				});
+
+			if (!LevelEditorTabPtr)
+			{
+				// LevelEditor tabs is not part of this PrimaryArea, could be that we are restoring the internal tabs of the LevelEditor area so the tab won't be here but the area is still the PrimaryArea
+				continue;
+			}
+
+			// Get the actual LevelEditor tab since the pointer may point to other tab later since we are moving tabs
+			const FTab LevelEditorTab = *LevelEditorTabPtr;
+			const FTab NewHomeScreenTab = FTab(FTabId(HomeScreenTabId), ETabState::OpenedTab);
+			const int32 LevelEditorExpectedIndex = bIsHomeScreenEnabled ? /** Always second tab with HomeScreen */ 1 : /** Always first without HomeScreen */ 0;
+
+			// Exit immediately if the position is already correct, either HomeScreen first and LevelEditor second if the HomeScreen is enabled or LevelEditor first if not
+			if (bIsHomeScreenEnabled && StackNode->Tabs.IndexOfByKey(LevelEditorTab) == LevelEditorExpectedIndex && StackNode->Tabs.IndexOfByKey(NewHomeScreenTab) == 0)
+			{
+				return;
+			}
+
+			if (!bIsHomeScreenEnabled && StackNode->Tabs.IndexOfByKey(LevelEditorTab) == LevelEditorExpectedIndex)
+			{
+				return;
+			}
+
+			// Remove all HomeScreen tab from the current stack and insert a new one in the correct position.
+			// While the change was hidden and tested there may have been cases where the HomeScreen was already saved in the PrimaryArea and we don't want 2 entries.
+			StackNode->Tabs.RemoveAll([HomeScreenTabId] (const FTab& InTab)
+				{
+					return InTab.TabId.TabType == HomeScreenTabId;
+				});
+
+			if (bIsHomeScreenEnabled)
+			{
+				// Insert the HomeScreen tab as the first tab
+				StackNode->Tabs.Insert(NewHomeScreenTab, 0);
+			}
+
+			const int32 LevelEditorIndex = StackNode->Tabs.IndexOfByKey(LevelEditorTab);
+
+			// Do not touch the LevelEditor position if already correct, the order should be HomeScreen->LevelEditor->Other tabs... or LevelEditor->Other tabs... if the HomeScreen is not enabled
+			if (LevelEditorIndex != LevelEditorExpectedIndex)
+			{
+				// Re-Insert the LevelEditor as the first or second tab
+				StackNode->Tabs.RemoveAt(LevelEditorIndex);
+				StackNode->Tabs.Insert(LevelEditorTab, LevelEditorExpectedIndex);
+			}
+
+			// Do not continue since we already fixed the area we wanted to contain the HomeScreen and the LevelEditor next to each other.
+			return;
+		}
+	}
+}
+
 const TSharedRef<FGlobalTabmanager>& FGlobalTabmanager::Get()
 {
 	static const TSharedRef<FGlobalTabmanager> Instance = FGlobalTabmanager::New();
diff --git a/Engine/Source/Runtime/Slate/Public/Framework/Docking/TabManager.h b/Engine/Source/Runtime/Slate/Public/Framework/Docking/TabManager.h
--- a/Engine/Source/Runtime/Slate/Public/Framework/Docking/TabManager.h
+++ b/Engine/Source/Runtime/Slate/Public/Framework/Docking/TabManager.h
@@ -1271,6 +1271,9 @@ protected:
 		/** Checks all dock areas and adds up the number of open tabs and unique parent windows in the manager */
 		SLATE_API void GetRecordableStats( int32& OutTabCount, TArray<TSharedPtr<SWindow>>& OutUniqueParentWindows ) const;
 
+	/** If the HomeScreen tab is enabled this will force the primary area to contain the HomeScreen and the LevelEditor next to each other otherwise it will just force the LevelEditor to be the first tab */
+	void FixLayoutLoadingPrimaryArea(const TSharedRef<FArea>& InPrimaryArea);
+
 	TSharedPtr<SDockingTabStack> AttemptToOpenTab( const FTabId& ClosedTabId, bool bForceOpenWindowIfNeeded = false );
 
 	/**
