[HomeScreen] Make the HomeScreen and LevelEditor position fixed, making other tabs unable to be dragged before them

diff --git a/Engine/Source/Editor/MainFrame/Private/MainFrameModule.cpp b/Engine/Source/Editor/MainFrame/Private/MainFrameModule.cpp
--- a/Engine/Source/Editor/MainFrame/Private/MainFrameModule.cpp
+++ b/Engine/Source/Editor/MainFrame/Private/MainFrameModule.cpp
@@ -80,21 +80,6 @@ namespace UE::Editor::HomeScreen
 		{
 			SetEnableHomeScreenFromCmdLine();
 		});
-
-	static bool IsHomeScreenEnabled()
-	{
-		static bool bIsHomeScreenEnabled = false;
-		UE_CALL_ONCE([&]()
-		{
-			if (const IConsoleVariable* CVar = IConsoleManager::Get().FindConsoleVariable(TEXT("HomeScreen.EnableHomeScreen")))
-			{
-				ensureAlwaysMsgf(!EnumHasAnyFlags(CVar->GetFlags(), ECVF_Default), TEXT("The CVar should have already been set from commandline."));
-				bIsHomeScreenEnabled = CVar->GetBool();
-			}
-		});
-
-		return bIsHomeScreenEnabled;
-	}
 }
 
 static FAutoConsoleCommand ResizeMainFrameCommand(
diff --git a/Engine/Source/Runtime/Slate/Private/Framework/Docking/SDockingTabWell.cpp b/Engine/Source/Runtime/Slate/Private/Framework/Docking/SDockingTabWell.cpp
--- a/Engine/Source/Runtime/Slate/Private/Framework/Docking/SDockingTabWell.cpp
+++ b/Engine/Source/Runtime/Slate/Private/Framework/Docking/SDockingTabWell.cpp
@@ -124,9 +124,18 @@ void SDockingTabWell::OnArrangeChildren( const FGeometry& AllottedGeometry, FArr
 			return ChildSize;
 		};
 
+	int32 ExpectedTabDropIndex = INDEX_NONE;
+	if (TabBeingDragged.IsValid())
+	{
+		// Get the expected DropIndex of the tab being dragged
+		ExpectedTabDropIndex = ComputeChildDropIndex(AllottedGeometry, TabBeingDragged.ToSharedRef());
+	}
+
 	// Arrange all the tabs left to right.
 	float XOffset = 0;
 
+	int32 TabDropIndexVisualGap = INDEX_NONE;
+
 	for( int32 TabIndex=0; TabIndex < NumChildren; ++TabIndex )
 	{
 		const TSharedRef<SDockTab> CurTab = Tabs[TabIndex];
@@ -143,9 +152,16 @@ void SDockingTabWell::OnArrangeChildren( const FGeometry& AllottedGeometry, FArr
 		}
 
 		// Is this spot reserved from the tab that is being dragged?
-		if ( TabBeingDragged.IsValid() && XOffset <= DraggedChildCenter && DraggedChildCenter < (XOffset + ChildWidthWithOverlap) )
+		if (TabBeingDragged.IsValid() && XOffset <= DraggedChildCenter && DraggedChildCenter < (XOffset + ChildWidthWithOverlap))
+		{
+			// If the expected tab index is greater than the current tab index it means that we are not allowed to place it in this current TabIndex, so visually don't add the gap, but force it at the same Expected Index
+			// Otherwise add it at the same TabIndex
+			TabDropIndexVisualGap = ExpectedTabDropIndex > TabIndex ? ExpectedTabDropIndex : TabIndex;
+		}
+
+		// If the gap is expected at this TabIndex add the offset here
+		if (TabDropIndexVisualGap == TabIndex)
 		{
-			// if so, leave some room to signify that this is where the dragged tab would end up
 			XOffset += ChildWidthWithOverlap;
 		}
 
@@ -333,7 +349,7 @@ float SDockingTabWell::ComputeDraggedTabOffset( const FGeometry& MyGeometry, con
 }
 
 
-int32 SDockingTabWell::ComputeChildDropIndex(const FGeometry& MyGeometry, const TSharedRef<SDockTab>& TabBeingDragged)
+int32 SDockingTabWell::ComputeChildDropIndex(const FGeometry& MyGeometry, const TSharedRef<SDockTab>& TabBeingDragged) const
 {
 	const float ChildWidth = ComputeChildSize(MyGeometry).X;
 	const float ChildWidthWithOverlap = ChildWidth - TabBeingDragged->GetOverlapWidth();
@@ -362,7 +378,25 @@ int32 SDockingTabWell::ComputeChildDropIndex(const FGeometry& MyGeometry, const
 			}
 		}
 	}
-	const int32 DropLocationIndex = FMath::Clamp(static_cast<int32>(DraggedChildCenter / ChildWidthWithOverlap), 0, Tabs.Num());
+
+	// If this is the LevelEditor Area other Tabs are not allowed to be placed in the first and/or second position as those position are fixed for the HomeScreen (if enabled) and the LevelEditor
+	TArray<TSharedRef<SDockTab>> TabsArray = Tabs.AsArrayCopy();
+	TSharedRef<SDockTab>* LevelEditorTab = TabsArray.FindByPredicate([](TSharedRef<SDockTab> InDockTab)
+		{
+			return InDockTab->GetLayoutIdentifier().TabType == FName(TEXT("LevelEditor"));
+		});
+
+	// if we found the LevelEditor tab than this is the LevelEditor PrimaryArea
+	const bool bIsLevelEditorPrimaryArea = LevelEditorTab != nullptr;
+	const bool bIsHomeScreenEnabled = UE::Editor::HomeScreen::IsHomeScreenEnabled();
+
+	int32 MinClamp = 0;
+	if (bIsLevelEditorPrimaryArea)
+	{
+		MinClamp = bIsHomeScreenEnabled ? /** First 2 are locked from the HomeScreen and the LevelEditor */ 2 : /** First is locked from the LevelEditor */ 1;
+	}
+
+	const int32 DropLocationIndex = FMath::Clamp(static_cast<int32>(DraggedChildCenter / ChildWidthWithOverlap), MinClamp, Tabs.Num());
 	return DropLocationIndex;
 }
 
diff --git a/Engine/Source/Runtime/Slate/Private/Framework/Docking/SDockingTabWell.h b/Engine/Source/Runtime/Slate/Private/Framework/Docking/SDockingTabWell.h
--- a/Engine/Source/Runtime/Slate/Private/Framework/Docking/SDockingTabWell.h
+++ b/Engine/Source/Runtime/Slate/Private/Framework/Docking/SDockingTabWell.h
@@ -24,7 +24,6 @@ struct FDockingConstants
 	static const FVector2D SLATE_API GetMaxTabSizeFor( ETabRole TabRole );
 };
 
-
 /**
  * TabWell is a panel that shows dockable tabs.
  * Tabs can be re-arranged and dragged out of the TabStack.
@@ -130,7 +129,7 @@ private:
 	 * The user is dropping a tab on this TabWell, and the TabWell's geometry is currently MyGeometry.
 	 * @return the index of the slot into which the tab shoudl go.
 	 */
-	int32 ComputeChildDropIndex(const FGeometry& MyGeometry, const TSharedRef<SDockTab>& TabBeingDragged);
+	int32 ComputeChildDropIndex(const FGeometry& MyGeometry, const TSharedRef<SDockTab>& TabBeingDragged) const;
 
 private:
 
diff --git a/Engine/Source/Runtime/Slate/Private/Framework/Docking/TabManager.cpp b/Engine/Source/Runtime/Slate/Private/Framework/Docking/TabManager.cpp
--- a/Engine/Source/Runtime/Slate/Private/Framework/Docking/TabManager.cpp
+++ b/Engine/Source/Runtime/Slate/Private/Framework/Docking/TabManager.cpp
@@ -40,23 +40,6 @@ TAutoConsoleVariable<bool> CVarPanelDrawerToggle(
 	TEXT("Enables or disables the panel drawer. When the drawer is off invocations will just create a normal tab.")
 );
 
-namespace UE::Editor::HomeScreen
-{
-	static bool IsHomeScreenEnabled()
-	{
-		static bool bIsHomeScreenEnabled = false;
-		UE_CALL_ONCE([&]()
-		{
-			if (const IConsoleVariable* CVar = IConsoleManager::Get().FindConsoleVariable(TEXT("HomeScreen.EnableHomeScreen")))
-			{
-				ensureAlwaysMsgf(!EnumHasAnyFlags(CVar->GetFlags(), ECVF_Default), TEXT("The CVar should have already been set from commandline."));
-				bIsHomeScreenEnabled = CVar->GetBool();
-			}
-		});
-		return bIsHomeScreenEnabled;
-	}
-}
-
 #define LOCTEXT_NAMESPACE "TabManager"
 
 static const FString UE_TABMANAGER_OPENED_TAB_STRING = TEXT("OpenedTab");
@@ -1068,8 +1051,7 @@ TSharedPtr<SWidget> FTabManager::RestoreFrom(const TSharedRef<FLayout>& Layout,
 		const bool bIsPrimaryArea = ThisArea->WindowPlacement == FArea::Placement_NoWindow;
 
 		// Force the HomeScreen to be always added when restoring layouts in the primary area if enabled or just force the LevelEditor to be the first tab
-		// For now only fix the layout if the CVar is enabled
-		if (bIsPrimaryArea && UE::Editor::HomeScreen::IsHomeScreenEnabled())
+		if (bIsPrimaryArea)
 		{
 			FixLayoutLoadingPrimaryArea(ThisArea);
 		}
@@ -3016,7 +2998,6 @@ void FTabManager::FixLayoutLoadingPrimaryArea(const TSharedRef<FArea>& InPrimary
 			}
 
 			// Remove all HomeScreen tab from the current stack and insert a new one in the correct position.
-			// While the change was hidden and tested there may have been cases where the HomeScreen was already saved in the PrimaryArea and we don't want 2 entries.
 			StackNode->Tabs.RemoveAll([HomeScreenTabId] (const FTab& InTab)
 				{
 					return InTab.TabId.TabType == HomeScreenTabId;
diff --git a/Engine/Source/Runtime/Slate/Public/Framework/Docking/TabManager.h b/Engine/Source/Runtime/Slate/Public/Framework/Docking/TabManager.h
--- a/Engine/Source/Runtime/Slate/Public/Framework/Docking/TabManager.h
+++ b/Engine/Source/Runtime/Slate/Public/Framework/Docking/TabManager.h
@@ -24,6 +24,22 @@ class FLayoutExtender;
 struct FTabMatcher;
 struct FSidebarTabLists;
 
+namespace UE::Editor::HomeScreen
+{
+	static bool IsHomeScreenEnabled()
+	{
+		static bool bIsHomeScreenEnabled = false;
+		UE_CALL_ONCE([&]()
+		{
+			if (const IConsoleVariable* CVar = IConsoleManager::Get().FindConsoleVariable(TEXT("HomeScreen.EnableHomeScreen")))
+			{
+				ensureAlwaysMsgf(!EnumHasAnyFlags(CVar->GetFlags(), ECVF_Default), TEXT("The CVar should have already been set from commandline."));
+				bIsHomeScreenEnabled = CVar->GetBool();
+			}
+		});
+		return bIsHomeScreenEnabled;
+	}
+}
 namespace UE::Slate::Private
 {
 	class SPanelDrawerArea;
