[HomeScreen] Fix static startup CVar getter in TabManager

diff --git a/Engine/Source/Runtime/Slate/Public/Framework/Docking/TabManager.h b/Engine/Source/Runtime/Slate/Public/Framework/Docking/TabManager.h
--- a/Engine/Source/Runtime/Slate/Public/Framework/Docking/TabManager.h
+++ b/Engine/Source/Runtime/Slate/Public/Framework/Docking/TabManager.h
@@ -26,7 +26,7 @@ struct FSidebarTabLists;
 
 namespace UE::Editor::HomeScreen
 {
-	static bool IsHomeScreenEnabled()
+	inline bool IsHomeScreenEnabled()
 	{
 		static bool bIsHomeScreenEnabled = false;
 		UE_CALL_ONCE([&]()
