Moved the ownership of the PanelDrawerto from the TabManager to the SDockingArea.
This doesn't change the apparence of the PanelDrawer, but it help with the serialization of the layout and will make easier to add PanelDrawers to other unsupported nomad tab windows and sub asset windows.
Added the layout persistence for the PanelDrawer.
Also restoring a layout in a asset edtior don't close the PannelDrawer anymore.

Restricted the experimental api for the PanelDrawer to only accept nomad tabs. On paper it could accept other tab type such as a minor tab, but this lack testing and might put the system in a bad state.
Changed the experimental public api of the PanelDrawer via the TabManagers to be a bit more future proof.
Cleaned a bit the new code in the tab manager.

DockTab can now have a DockingArea without having a parent DockingTabWell.
Changed the default size of PanelDrawer when opening. It now take 1/4 of the host area instead of 1/5 by default.

#FailedValidations Preflight

diff --git a/Engine/Plugins/Experimental/AIAssistant/Source/AIAssistant/Private/AIAssistant.cpp b/Engine/Plugins/Experimental/AIAssistant/Source/AIAssistant/Private/AIAssistant.cpp
--- a/Engine/Plugins/Experimental/AIAssistant/Source/AIAssistant/Private/AIAssistant.cpp
+++ b/Engine/Plugins/Experimental/AIAssistant/Source/AIAssistant/Private/AIAssistant.cpp
@@ -143,7 +143,8 @@ void FAIAssistantModule::OnTogglePluginTabBasedOnCursorLocation()
 
 	if (TabManager)
 	{
-		TabManager->TryToggleTabInPanelDrawer(AIAssistantTabName);
+		// Leave the parent window to null so it will use the primary area
+		TabManager->TryToggleTabInPanelDrawer(AIAssistantTabName, {});
 	}
 	else
 	{
@@ -162,7 +163,8 @@ void FAIAssistantModule::OnOpenPluginTab(const FWidgetPath& InWidgetPath)
 
 	if (TabManager)
 	{
-		TabManager->TryOpenTabInPanelDrawer(AIAssistantTabName);
+		// Leave the parent window to null so it will use the primary area
+		TabManager->TryOpenTabInPanelDrawer(AIAssistantTabName, {});
 	}
 	else
 	{
diff --git a/Engine/Source/Editor/LevelEditor/Private/SLevelEditor.cpp b/Engine/Source/Editor/LevelEditor/Private/SLevelEditor.cpp
--- a/Engine/Source/Editor/LevelEditor/Private/SLevelEditor.cpp
+++ b/Engine/Source/Editor/LevelEditor/Private/SLevelEditor.cpp
@@ -1775,7 +1775,7 @@ TSharedRef<SWidget> SLevelEditor::RestoreContentArea( const TSharedRef<SDockTab>
 			ContentAreaWidget.ToSharedRef()
 		];
 
-	if (TSharedPtr<SWidget> PanelDrawer = LevelEditorTabManager->RestorePanelDrawerFrom(DefaultLayout, LevelEditorContent, OwnerWindow))
+	if (TSharedPtr<SWidget> PanelDrawer = LevelEditorTabManager->RestorePanelDrawer(LevelEditorContent, OwnerWindow))
 	{
 		LevelEditorContent = PanelDrawer.ToSharedRef();
 	}
diff --git a/Engine/Source/Editor/StatusBar/Private/SStatusBar.cpp b/Engine/Source/Editor/StatusBar/Private/SStatusBar.cpp
--- a/Engine/Source/Editor/StatusBar/Private/SStatusBar.cpp
+++ b/Engine/Source/Editor/StatusBar/Private/SStatusBar.cpp
@@ -811,10 +811,10 @@ TSharedPtr<SWidget> SStatusBar::MakeSummonPanelDrawer(TArray<UStatusBarSubsystem
 	{
 		return {};
 	}
+	TSharedPtr<SDockTab> ParentTabPtr = ParentTab.Pin();
+	TSharedPtr<FTabManager> TabManager = FGlobalTabmanager::Get()->GetTabManagerForMajorTab(ParentTabPtr);
 
-	TSharedPtr<FTabManager> TabManager = FGlobalTabmanager::Get()->GetTabManagerForMajorTab(ParentTab.Pin());
-
-	if (!TabManager || !TabManager->HasPanelDrawer())
+	if (!TabManager || !TabManager->HasPanelDrawer(ParentTabPtr->GetParentWindow()))
 	{
 		return {};
 	}
@@ -1093,7 +1093,8 @@ void SStatusBar::DismissProgressBar()
 void SStatusBar::UpdateSummonPanelDrawerArrowRotation(const FTabManager& TabManager)
 {
 	// Todo add its own style when more mature and less a proof of concept
-	if (TabManager.IsPanelDrawerOpen())
+	TSharedPtr<SDockTab> ParentTabPtr = ParentTab.Pin();
+	if (TabManager.IsPanelDrawerOpen(ParentTabPtr->GetParentWindow()))
 	{
 		PanelDrawerArrow->SetImage(FAppStyle::Get().GetBrush(TEXT("Icons.ChevronRight")));
 	}
@@ -1105,9 +1106,10 @@ void SStatusBar::UpdateSummonPanelDrawerArrowRotation(const FTabManager& TabMana
 
 FReply SStatusBar::ToggleTabInPanelDrawer(FName TabId)
 {
-	if (TSharedPtr<FTabManager> TabManager = FGlobalTabmanager::Get()->GetTabManagerForMajorTab(ParentTab.Pin()))
+	TSharedPtr<SDockTab> ParentTabPtr = ParentTab.Pin();
+	if (TSharedPtr<FTabManager> TabManager = FGlobalTabmanager::Get()->GetTabManagerForMajorTab(ParentTabPtr))
 	{
-		TabManager->TryToggleTabInPanelDrawer(TabId);
+		TabManager->TryToggleTabInPanelDrawer(TabId, ParentTabPtr->GetParentWindow());
 	}
 
 	return FReply::Handled();
@@ -1115,9 +1117,10 @@ FReply SStatusBar::ToggleTabInPanelDrawer(FName TabId)
 
 FReply SStatusBar::ClosePanelDrawer()
 {
-	if (TSharedPtr<FTabManager> TabManager = FGlobalTabmanager::Get()->GetTabManagerForMajorTab(ParentTab.Pin()))
+	TSharedPtr<SDockTab> ParentTabPtr = ParentTab.Pin();
+	if (TSharedPtr<FTabManager> TabManager = FGlobalTabmanager::Get()->GetTabManagerForMajorTab(ParentTabPtr))
 	{
-		TabManager->ClosePanelDrawer();
+		TabManager->ClosePanelDrawer(ParentTabPtr->GetParentWindow());
 	}
 
 	return FReply::Handled();
diff --git a/Engine/Source/Editor/StatusBar/Private/StatusBarSubsystem.cpp b/Engine/Source/Editor/StatusBar/Private/StatusBarSubsystem.cpp
--- a/Engine/Source/Editor/StatusBar/Private/StatusBarSubsystem.cpp
+++ b/Engine/Source/Editor/StatusBar/Private/StatusBarSubsystem.cpp
@@ -106,7 +106,7 @@ namespace UE
 					TabManager = FGlobalTabmanager::Get()->GetTabManagerForMajorTab(InParentDockTab);
 				}
 
-				if (!TabManager || !TabManager->HasPanelDrawer())
+				if (!TabManager || !TabManager->HasPanelDrawer(InParentDockTab->GetParentWindow()))
 				{
 					return {};
 				}
diff --git a/Engine/Source/Editor/UnrealEd/Private/Toolkits/SStandaloneAssetEditorToolkitHost.cpp b/Engine/Source/Editor/UnrealEd/Private/Toolkits/SStandaloneAssetEditorToolkitHost.cpp
--- a/Engine/Source/Editor/UnrealEd/Private/Toolkits/SStandaloneAssetEditorToolkitHost.cpp
+++ b/Engine/Source/Editor/UnrealEd/Private/Toolkits/SStandaloneAssetEditorToolkitHost.cpp
@@ -182,9 +182,11 @@ void SStandaloneAssetEditorToolkitHost::RestoreFromLayout( const TSharedRef<FTab
 	const TSharedRef<SDockTab> HostTab = HostTabPtr.Pin().ToSharedRef();
 	HostTab->SetOnTabClosed(SDockTab::FOnTabClosedCallback::CreateSP(this, &SStandaloneAssetEditorToolkitHost::OnTabClosed));
 
-	this->ChildSlot[SNullWidget::NullWidget];
+	// Let the manager handle some logic on the Areas before destroying
 	MyTabManager->CloseAllAreas();
 
+	this->ChildSlot[SNullWidget::NullWidget];
+
 	TSharedPtr<SWindow> ParentWindow = FSlateApplication::Get().FindWidgetWindow( HostTab );
 	TSharedPtr<SWidget> RestoredUI = MyTabManager->RestoreFrom( NewLayout, ParentWindow );
 
@@ -213,7 +215,7 @@ void SStandaloneAssetEditorToolkitHost::RestoreFromLayout( const TSharedRef<FTab
 	 * The ordering is important here, otherwise the status bar won't have a summon tab in PanelDrawer widget.
 	 * The layout and panel drawer area must be restored before creating the status bar.
 	 */
-	if (TSharedPtr<SWidget> PanelDrawer = MyTabManager->RestorePanelDrawerFrom(NewLayout, AssetEditorContent, ParentWindow))
+	if (TSharedPtr<SWidget> PanelDrawer = MyTabManager->RestorePanelDrawer(AssetEditorContent, ParentWindow.ToSharedRef()))
 	{
 		AssetEditorContent = PanelDrawer.ToSharedRef();
 	}
diff --git a/Engine/Source/Runtime/Slate/Private/Framework/Docking/SDockingArea.cpp b/Engine/Source/Runtime/Slate/Private/Framework/Docking/SDockingArea.cpp
--- a/Engine/Source/Runtime/Slate/Private/Framework/Docking/SDockingArea.cpp
+++ b/Engine/Source/Runtime/Slate/Private/Framework/Docking/SDockingArea.cpp
@@ -170,6 +170,13 @@ void SDockingArea::Construct( const FArguments& InArgs, const TSharedRef<FTabMan
 	{
 		AddChildNode( InArgs._InitialContent.ToSharedRef() );
 	}
+
+	for (const FTabManager::FPanelDrawerTab& InactivePanelDrawerTab : PersistentNode->GetPanelDrawerInactiveTabs())
+	{
+		TSharedRef<UE::Slate::Private::FPanelDrawerData> PanelDrawerData = MakeShared<UE::Slate::Private::FPanelDrawerData>();
+		PanelDrawerData->Size = InactivePanelDrawerTab.Size;
+		InactivePanelDrawerTabs.Add(InactivePanelDrawerTab.TabId, MoveTemp(PanelDrawerData));
+	}
 }
 
 void SDockingArea::OnDragEnter( const FGeometry& MyGeometry, const FDragDropEvent& DragDropEvent )
@@ -381,6 +388,53 @@ TSharedPtr<FTabManager::FLayoutNode> SDockingArea::GatherPersistentLayout() cons
 		}
 	}
 
+	// PanelDrawer layout data
+	{
+		FTabManager::FPanelDrawerTab PanelDrawerTab;
+		if (const TSharedPtr<UE::Slate::Private::SPanelDrawerArea> PanelDrawer = PanelDrawerPtr.Pin())
+		{
+			if (const TSharedPtr<UE::Slate::Private::FPanelDrawerData> PanelDrawerData = PanelDrawer->GetHostedPanelDrawerData())
+			{
+				if (PanelDrawerData->HostedTab)
+				{
+					// Only store the nomad tab in the area that manager their window
+					ETabRole HostedTabRole = PanelDrawerData->HostedTab->GetTabRole();
+					if ((HostedTabRole == ETabRole::NomadTab && GetParentWindow())
+						|| (HostedTabRole != ETabRole::NomadTab && !GetParentWindow()))
+					{
+						PanelDrawerTab.TabId = PanelDrawerData->HostedTab->GetLayoutIdentifier();
+						const FTabManager::FPanelDrawerSize& Size = PanelDrawerData->Size;
+						PanelDrawerTab.Size = Size;
+					}
+				}
+			}
+		}
+
+		if (PanelDrawerTab.TabId.TabType.IsNone())
+		{
+			if (HiddenPanelDrawerTabToReopenOnRestore && HiddenPanelDrawerTabToReopenOnRestore->HostedTab)
+			{
+				PanelDrawerTab.TabId = HiddenPanelDrawerTabToReopenOnRestore->HostedTab->GetLayoutIdentifier();
+				const FTabManager::FPanelDrawerSize& Size = HiddenPanelDrawerTabToReopenOnRestore->Size;
+				PanelDrawerTab.Size = Size;
+			}
+		}
+
+		if (!PanelDrawerTab.TabId.TabType.IsNone())
+		{
+			PersistentNode->SetPanelDrawerActiveTab(MoveTemp(PanelDrawerTab));
+		}
+	}
+
+	for (const TPair<FTabId, TSharedPtr<UE::Slate::Private::FPanelDrawerData>>& Pair : InactivePanelDrawerTabs)
+	{
+		FTabManager::FPanelDrawerTab PanelDrawerTab;
+		PanelDrawerTab.TabId = Pair.Key;
+		const FTabManager::FPanelDrawerSize& Size = Pair.Value->Size;
+		PanelDrawerTab.Size = Size;
+		PersistentNode->AddPanelDrawerInactiveTab(MoveTemp(PanelDrawerTab));
+	}
+
 	if( !bHaveLayoutData )
 	{
 		PersistentNode.Reset();
@@ -495,12 +549,367 @@ TArray<TSharedRef<SDockTab>> SDockingArea::GetAllSidebarTabs() const
 	return AllSidebarTabs;
 }
 
+bool SDockingArea::HostTabIntoPanelDrawer(const TSharedRef<SDockTab>& InTab)
+{
+	TSharedPtr<UE::Slate::Private::SPanelDrawerArea> PanelDrawer = PanelDrawerPtr.Pin();
+	if (!PanelDrawer)
+	{
+		return false;
+	}
+
+	TSharedPtr<UE::Slate::Private::FPanelDrawerData> PanelDrawerData = PanelDrawer->GetHostedPanelDrawerData();
+	if (PanelDrawerData && PanelDrawerData->HostedTab)
+	{
+		InactivePanelDrawerTabs.Add(PanelDrawerData->HostedTab->GetLayoutIdentifier(), MoveTemp(PanelDrawerData));
+	}
+
+
+	if (!InactivePanelDrawerTabs.RemoveAndCopyValue(InTab->GetLayoutIdentifier(), PanelDrawerData))
+	{
+		PanelDrawerData = MakeShared<UE::Slate::Private::FPanelDrawerData>();
+	}
+
+	PanelDrawerData->HostedTab = InTab;
+	InTab->SetTabManager(GetTabManager());
+	InTab->SetParentDockingArea(AsSharedSubobject(this));
+	HiddenPanelDrawerTabToReopenOnRestore.Reset();
+
+	constexpr bool bWithAnimation = true;
+	PanelDrawer->OpenPanel(bWithAnimation, PanelDrawerData.ToSharedRef());
+
+	return true;
+}
+
+void SDockingArea::ClosePanelDrawer()
+{
+	if (TSharedPtr<UE::Slate::Private::SPanelDrawerArea> PanelDrawer = PanelDrawerPtr.Pin())
+	{
+		TSharedPtr<UE::Slate::Private::FPanelDrawerData> PanelDrawerData = PanelDrawer->GetHostedPanelDrawerData();
+		if (PanelDrawerData && PanelDrawerData->HostedTab)
+		{
+			InactivePanelDrawerTabs.Add(PanelDrawerData->HostedTab->GetLayoutIdentifier(), MoveTemp(PanelDrawerData));
+		}
+
+		constexpr bool bWithAnimation = true;
+		PanelDrawer->ClosePanel(bWithAnimation);
+
+		HiddenPanelDrawerTabToReopenOnRestore.Reset();
+	}
+}
+
+TSharedPtr<SDockTab> SDockingArea::ClosePanelDrawerForTransfer()
+{
+	TSharedPtr<SDockTab> HostedTab;
+	if (TSharedPtr<UE::Slate::Private::SPanelDrawerArea> PanelDrawer = PanelDrawerPtr.Pin())
+	{
+		TSharedPtr<UE::Slate::Private::FPanelDrawerData> PanelDrawerData = PanelDrawer->GetHostedPanelDrawerData();
+		if (PanelDrawerData && PanelDrawerData->HostedTab)
+		{
+			// Purposely null the stored Hosted Tab
+			HostedTab = MoveTemp(PanelDrawerData->HostedTab);
+			InactivePanelDrawerTabs.Add(HostedTab->GetLayoutIdentifier(), MoveTemp(PanelDrawerData));
+		}
+
+		constexpr bool bWithAnimation = true;
+		constexpr bool bForTransfer = true;
+		PanelDrawer->ClosePanel(bWithAnimation, bForTransfer);
+
+		HiddenPanelDrawerTabToReopenOnRestore.Reset();
+	}
+
+	return HostedTab;
+}
+
+bool SDockingArea::IsPanelDrawerOpen() const
+{
+	if (TSharedPtr<UE::Slate::Private::SPanelDrawerArea> PanelDrawer = PanelDrawerPtr.Pin())
+	{
+		return PanelDrawer->IsOpen();
+	}
+	return false;
+}
+
+bool SDockingArea::HasPanelDrawer() const
+{
+	return PanelDrawerPtr.IsValid();
+}
+
+TSharedPtr<SDockTab> SDockingArea::GetPanelDrawerSystemHostedTab(const FTabId& TabId) const
+{
+	if (HiddenPanelDrawerTabToReopenOnRestore && HiddenPanelDrawerTabToReopenOnRestore->HostedTab)
+	{
+		if (HiddenPanelDrawerTabToReopenOnRestore->HostedTab->GetLayoutIdentifier() == TabId)
+		{
+			return HiddenPanelDrawerTabToReopenOnRestore->HostedTab;
+		}
+	}
+
+	if (const TSharedPtr<UE::Slate::Private::FPanelDrawerData>* PanelDrawerData = InactivePanelDrawerTabs.Find(TabId))
+	{
+		return (*PanelDrawerData)->HostedTab;
+	}
+
+	if (TSharedPtr<UE::Slate::Private::SPanelDrawerArea> PanelDrawer = PanelDrawerPtr.Pin())
+	{
+		if (TSharedPtr<SDockTab> HostedTab = PanelDrawer->GetHostedTab())
+		{
+			if (HostedTab->GetLayoutIdentifier() == TabId)
+			{
+				return HostedTab;
+			}
+		}
+	}
+
+	return {};
+}
+
+TSharedPtr<SDockTab> SDockingArea::GetPanelDrawerHostedTab() const
+{
+	// Handle the case of this local area hosting a nomad tab
+	if (TSharedPtr<UE::Slate::Private::SPanelDrawerArea> PanelDrawer = PanelDrawerPtr.Pin())
+	{
+		// Note: Consider only returning the data when open if this cause some issue
+		return PanelDrawer->GetHostedTab();
+	}
+
+	return {};
+}
+
+void SDockingArea::SetPanelDrawerArea(const TSharedPtr<UE::Slate::Private::SPanelDrawerArea>& InPanelDrawerArea)
+{
+	if (TSharedPtr<UE::Slate::Private::SPanelDrawerArea> PanelDrawer = PanelDrawerPtr.Pin())
+	{
+		if (!InPanelDrawerArea)
+		{
+			HiddenPanelDrawerTabToReopenOnRestore = PanelDrawer->GetHostedPanelDrawerData();
+		}
+
+		if (InPanelDrawerArea != PanelDrawer)
+		{
+			PanelDrawer->ClosePanel(false);
+		}
+	}
+
+	PanelDrawerPtr = InPanelDrawerArea;
+
+	if (InPanelDrawerArea)
+	{
+		bool bShouldRestore = true;
+		if (TSharedPtr<UE::Slate::Private::FPanelDrawerData> PanelDrawerData = InPanelDrawerArea->GetHostedPanelDrawerData())
+		{
+			if (PanelDrawerData->HostedTab)
+			{
+				if (PanelDrawerData->HostedTab->GetTabRole() == ETabRole::NomadTab)
+				{
+					bShouldRestore = false;
+					HiddenPanelDrawerTabToReopenOnRestore.Reset();
+					PanelDrawerData->HostedTab->SetTabManager(GetTabManager());
+					PanelDrawerData->HostedTab->SetParentDockingArea(AsSharedSubobject(this));
+					InactivePanelDrawerTabs.Remove(PanelDrawerData->HostedTab->GetLayoutIdentifier());
+				}
+			}
+		}
+
+		if (bShouldRestore)
+		{
+			RestorePanelDrawerArea();
+		}
+	}
+}
+
+void SDockingArea::DetachPanelDrawerArea()
+{
+	PanelDrawerPtr.Reset();
+}
+
+bool SDockingArea::RestorePanelDrawerArea()
+{
+	return RestorePanelDrawerArea(TSharedPtr<UE::Slate::Private::SPanelDrawerArea>());
+}
+
+bool SDockingArea::RestorePanelDrawerArea(const TSharedPtr<UE::Slate::Private::SPanelDrawerArea>& InPanelDrawerAreaOverride)
+{
+	TSharedPtr<UE::Slate::Private::SPanelDrawerArea> PanelDrawer;
+	if (InPanelDrawerAreaOverride)
+	{
+		PanelDrawer = InPanelDrawerAreaOverride;
+	}
+	else
+	{
+		PanelDrawer = PanelDrawerPtr.Pin();
+	}
+
+	if (!PanelDrawer)
+	{
+		return false;
+	}
+
+
+	if (TSharedPtr<SDockingArea> MajorDockingArea = GetTopLevelDockingArea())
+	{
+		// This Area live under another area. Let the top level Area select the drawer content first.
+		MajorDockingArea->SetPanelDrawerArea(PanelDrawer);
+		if (PanelDrawer->IsOpen())
+		{
+			HiddenPanelDrawerTabToReopenOnRestore.Reset();
+			return true;
+		}
+	}
+	else if (!GetParentWindow())
+	{
+		return false;
+	}
+
+	// Consume the queued Restore
+	TSharedPtr<UE::Slate::Private::FPanelDrawerData> PanelDrawerData = MoveTemp(HiddenPanelDrawerTabToReopenOnRestore);
+	if (PanelDrawerData && PanelDrawerData->HostedTab)
+	{
+		PanelDrawerData->HostedTab->SetTabManager(GetTabManager());
+		PanelDrawerData->HostedTab->SetParentDockingArea(AsSharedSubobject(this));
+		InactivePanelDrawerTabs.Remove(PanelDrawerData->HostedTab->GetLayoutIdentifier());
+
+		constexpr bool bWithAnimation = false;
+		PanelDrawer->OpenPanel(bWithAnimation, PanelDrawerData.ToSharedRef());
+		
+		return true;
+	}
+
+	return false;
+}
+
+TMap<FTabId, TSharedRef<SDockTab>> SDockingArea::GetPanelDrawerKeepAliveTabs() const
+{
+	TMap<FTabId, TSharedRef<SDockTab>> Output;
+	Output.Reserve(InactivePanelDrawerTabs.Num() + 1);
+
+	if (TSharedPtr < UE::Slate::Private::SPanelDrawerArea> PanelDrawer = PanelDrawerPtr.Pin())
+	{
+		TSharedPtr<UE::Slate::Private::FPanelDrawerData> PanelDrawerData = PanelDrawer->GetHostedPanelDrawerData();
+		if (PanelDrawerData && PanelDrawerData->HostedTab)
+		{
+			ETabRole HostedTabRole = PanelDrawerData->HostedTab->GetTabRole();
+			// Only keep alive the nomad tab if we own the window
+			if (HostedTabRole != ETabRole::NomadTab || GetParentWindow())
+			{
+				Output.Add(PanelDrawerData->HostedTab->GetLayoutIdentifier().TabType, PanelDrawerData->HostedTab.ToSharedRef());
+			}
+		}
+		else if (HiddenPanelDrawerTabToReopenOnRestore && HiddenPanelDrawerTabToReopenOnRestore->HostedTab)
+		{
+			Output.Add(HiddenPanelDrawerTabToReopenOnRestore->HostedTab->GetLayoutIdentifier().TabType, HiddenPanelDrawerTabToReopenOnRestore->HostedTab.ToSharedRef());
+		}
+
+	}
+
+	for (const TPair<FTabId, TSharedPtr<UE::Slate::Private::FPanelDrawerData>>& Pair : InactivePanelDrawerTabs)
+	{
+		if (Pair.Value && Pair.Value->HostedTab)
+		{
+			Output.Add(Pair.Key, Pair.Value->HostedTab.ToSharedRef());
+		}
+	}
+
+	return Output;
+}
+
+TSharedPtr<UE::Slate::Private::SPanelDrawerArea> SDockingArea::GetPanelDrawerArea() const
+{
+	return PanelDrawerPtr.Pin();
+}
+
+bool SDockingArea::RemoveHiddenInactivePanelDrawerTab(const TSharedPtr<SDockTab>& TabToRemove)
+{
+	if (TabToRemove)
+	{
+		for (TPair<FTabId, TSharedPtr<UE::Slate::Private::FPanelDrawerData>>& Pair : InactivePanelDrawerTabs)
+		{
+			if (Pair.Value && Pair.Value->HostedTab == TabToRemove)
+			{
+				Pair.Value->HostedTab.Reset();
+				return true;
+			}
+		}
+
+		if (HiddenPanelDrawerTabToReopenOnRestore && HiddenPanelDrawerTabToReopenOnRestore->HostedTab == TabToRemove)
+		{
+			HiddenPanelDrawerTabToReopenOnRestore.Reset();
+			return true;
+		}
+	}
+
+	return false;
+}
+
+void SDockingArea::CleanPanelDrawer()
+{
+	// Push our nomad tab to top docking area
+	if (TSharedPtr<SDockingArea> GlobalDockArea = GetTopLevelDockingArea())
+	{
+		if (TSharedPtr<UE::Slate::Private::SPanelDrawerArea> PanelDrawer = PanelDrawerPtr.Pin())
+		{
+			if (TSharedPtr<UE::Slate::Private::FPanelDrawerData> PanelDrawerData = PanelDrawer->GetHostedPanelDrawerData())
+			{
+				if (PanelDrawerData->HostedTab && PanelDrawerData->HostedTab->GetTabRole() == ETabRole::NomadTab)
+				{
+					GlobalDockArea->HiddenPanelDrawerTabToReopenOnRestore = PanelDrawer->GetHostedPanelDrawerData();
+				}
+			}
+		}
+
+		if (!GlobalDockArea->HiddenPanelDrawerTabToReopenOnRestore)
+		{
+			if (HiddenPanelDrawerTabToReopenOnRestore)
+			{
+				if (HiddenPanelDrawerTabToReopenOnRestore->HostedTab && HiddenPanelDrawerTabToReopenOnRestore->HostedTab->GetTabRole() == ETabRole::NomadTab)
+				{
+					GlobalDockArea->HiddenPanelDrawerTabToReopenOnRestore = MoveTemp(HiddenPanelDrawerTabToReopenOnRestore);
+				}
+			}
+		}
+
+
+	}
+
+	HiddenPanelDrawerTabToReopenOnRestore.Reset();
+	InactivePanelDrawerTabs.Empty();
+}
+
+void SDockingArea::SetPanelDrawerHiddenActiveTab(TSharedRef<UE::Slate::Private::FPanelDrawerData>&& InPanelDrawerData)
+{
+	if (!PanelDrawerPtr.IsValid())
+	{
+		if (InPanelDrawerData->HostedTab)
+		{
+			InPanelDrawerData->HostedTab->SetTabManager(GetTabManager());
+			InPanelDrawerData->HostedTab->SetParentDockingArea(SharedThis(this));
+			HiddenPanelDrawerTabToReopenOnRestore = MoveTemp(InPanelDrawerData);
+		}
+	}
+}
+
 SDockingNode::ECleanupRetVal SDockingArea::CleanUpNodes()
 {
 	SDockingNode::ECleanupRetVal ReturnValue = SDockingSplitter::CleanUpNodes();
 	return ReturnValue;
 }
 
+TSharedPtr<SDockingArea> SDockingArea::GetTopLevelDockingArea() const
+{
+	if (!ParentWindowPtr.IsValid())
+	{
+		TSharedRef<FGlobalTabmanager> GlobalManager = FGlobalTabmanager::Get();
+		if (TSharedPtr<SDockTab> MajorTab = GlobalManager->GetMajorTabForTabManager(GetTabManager()))
+		{
+			if (TSharedPtr<SWindow> GlobalWindow = MajorTab->GetParentWindow())
+			{
+				return GlobalManager->GetPrivateApi().GetDockingAreaForPanelDrawer(GlobalWindow);
+			}
+		}
+	}
+
+	return {};
+}
+
 EVisibility SDockingArea::TargetCrossVisibility() const
 {
 	return (bIsOverlayVisible && !bIsCenterTargetVisible)
diff --git a/Engine/Source/Runtime/Slate/Private/Framework/Docking/SDockingArea.h b/Engine/Source/Runtime/Slate/Private/Framework/Docking/SDockingArea.h
--- a/Engine/Source/Runtime/Slate/Private/Framework/Docking/SDockingArea.h
+++ b/Engine/Source/Runtime/Slate/Private/Framework/Docking/SDockingArea.h
@@ -11,6 +11,7 @@
 #include "Framework/Docking/TabManager.h"
 #include "Framework/Docking/SDockingNode.h"
 #include "Framework/Docking/SDockingSplitter.h"
+#include "Framework/Docking/SPanelDrawerArea.h"
 
 class STabSidebar;
 
@@ -136,10 +137,63 @@ public:
 	SLATE_API TArray<TSharedRef<SDockTab>> GetAllSidebarTabs() const;
 
 	bool CanHaveSidebar() const { return bCanHaveSidebar; }
+
+	/** If we have a PanelDrawer the tab host the tab in it */
+	bool HostTabIntoPanelDrawer(const TSharedRef<SDockTab>& InTab);
+
+	/** Close the PanelDrawer associated to this area */
+	void ClosePanelDrawer();
+
+	/** Close the PanelDrawer associated to this area but set it up for a tab transfer */
+	TSharedPtr<SDockTab> ClosePanelDrawerForTransfer();
+
+	/** Return true if the area as an open PanelDrawer */
+	bool IsPanelDrawerOpen() const;
+
+	/** Does this area have a DrawerPanel associated with it */
+	bool HasPanelDrawer() const;
+
+	/** Check if the tab is hosted by the PanelDrawer system associated to this area (Include previously hosted tabs ready to reopened) */
+	TSharedPtr<SDockTab> GetPanelDrawerSystemHostedTab(const FTabId& TabId) const;
+
+	/** Get the currently active tab in the PanelDrawer */
+	TSharedPtr<SDockTab> GetPanelDrawerHostedTab() const;
+
+	/** Associate a PanelDrawer to this area */
+	void SetPanelDrawerArea(const TSharedPtr<UE::Slate::Private::SPanelDrawerArea>& InPanelDrawerArea);
+
+	/** Detach the PanelDrawer but doesn't setup or use the restore mechanism */
+	void DetachPanelDrawerArea();
+
+	/** Used to set the state of the PanelDrawer associated to this area */
+	bool RestorePanelDrawerArea();
+
+	/** Get list of tabs this PanelDrawer is keeping alive */
+	TMap<FTabId, TSharedRef<SDockTab>> GetPanelDrawerKeepAliveTabs() const;
+
+	/** Get the PanelDrawer currently associated with this area */
+	TSharedPtr<UE::Slate::Private::SPanelDrawerArea> GetPanelDrawerArea() const;
+
+	/**
+	 * Remove a hidden but keep alive tab by the PanelDrawer.
+	 * This may include the PanelDrawer hidden active tab if the area currently doesn't have an associated PanelDrawer
+	 */
+	bool RemoveHiddenInactivePanelDrawerTab(const TSharedPtr<SDockTab>& TabToRemove);
+
+	/** Remove any data that could cause this Paneldrawer to stay alive. Also transfer some state to the top level area if needed */
+	void CleanPanelDrawer();
+
+	/** Use when restoring a layout do not use outside of that */
+	void SetPanelDrawerHiddenActiveTab(TSharedRef<UE::Slate::Private::FPanelDrawerData>&& InPanelDrawerData);
+
 protected:
 	SLATE_API virtual SDockingNode::ECleanupRetVal CleanUpNodes() override;
 
 private:
+
+	/** Return the top docking area that is owning the window in which this docking area live */
+	TSharedPtr<SDockingArea> GetTopLevelDockingArea() const;
+
 	SLATE_API EVisibility TargetCrossVisibility() const;
 	SLATE_API EVisibility TargetCrossCenterVisibility() const;
 	/** Dock a tab along the outer edge of this DockArea */
@@ -161,6 +215,8 @@ private:
 	 * We also update the sidebar to account for major tabs.  Docking areas for major tabs do not have a sidebar.
 	 */
 	SLATE_API void UpdateWindowChromeAndSidebar();
+
+	bool RestorePanelDrawerArea(const TSharedPtr<UE::Slate::Private::SPanelDrawerArea>& InPanelDrawerAreaOverride);
 private:	
 	/** Left and right sidebar widgets */
 	TSharedPtr<STabSidebar> LeftSidebar;
@@ -193,4 +249,11 @@ private:
 
 	/** True if this area can ever show sidebars (minor tab areas only) */
 	bool bCanHaveSidebar;
+
+	/** Optional Panel drawer for this area */
+	TWeakPtr<UE::Slate::Private::SPanelDrawerArea> PanelDrawerPtr;
+
+	/** Layout data for the PanelDrawer, also include the keep alive tabs */
+	TSharedPtr<UE::Slate::Private::FPanelDrawerData> HiddenPanelDrawerTabToReopenOnRestore;
+	TMap<FTabId, TSharedPtr<UE::Slate::Private::FPanelDrawerData>> InactivePanelDrawerTabs;
 };
diff --git a/Engine/Source/Runtime/Slate/Private/Framework/Docking/SDockingTabWell.cpp b/Engine/Source/Runtime/Slate/Private/Framework/Docking/SDockingTabWell.cpp
--- a/Engine/Source/Runtime/Slate/Private/Framework/Docking/SDockingTabWell.cpp
+++ b/Engine/Source/Runtime/Slate/Private/Framework/Docking/SDockingTabWell.cpp
@@ -430,13 +430,19 @@ FReply SDockingTabWell::StartDraggingTab( TSharedRef<SDockTab> TabToStartDraggin
 		// On Mac we need to activate the app as we may be dragging a window that is set to be invisible if the app is inactive
 		FPlatformApplicationMisc::ActivateApplication();
 #endif
+		TSharedPtr<SDockingArea> DockArea = GetDockArea();
+
+		if (TabToStartDragging->GetTabRole() == ETabRole::MajorTab)
+		{
+			DockArea->DetachPanelDrawerArea();
+		}
 
 		// Start dragging.
 		TSharedRef<FDockingDragOperation> DragDropOperation =
 			FDockingDragOperation::New(
 				TabToStartDragging,
 				InTabGrabOffsetFraction,
-				GetDockArea().ToSharedRef(),
+				DockArea.ToSharedRef(),
 				ParentTabStackPtr.Pin()->GetTabStackGeometry().GetLocalSize()
 			);
 
@@ -487,6 +493,19 @@ void SDockingTabWell::OnDragLeave( const FDragDropEvent& DragDropEvent )
 			// Update the DragAndDrop operation based on this change.
 			const int32 LastForegroundTabIndex = Tabs.Find(TabBeingDragged.ToSharedRef());
 
+			if ( Tabs.Num() > 1 )
+			{
+				// Also stop showing its content; switch to the next tab that was active.
+				if (LastForegroundTabIndex + 1 < Tabs.Num())
+				{
+					BringTabToFront(LastForegroundTabIndex + 1);
+				}
+				else
+				{
+					BringTabToFront(FMath::Max(LastForegroundTabIndex - 1, 0));
+				}
+			}
+
 			// Remove the tab from the well when it is dragged out
 			Tabs.Remove(TabBeingDraggedPtr.ToSharedRef());
 
@@ -498,16 +517,11 @@ void SDockingTabWell::OnDragLeave( const FDragDropEvent& DragDropEvent )
 			this->TabBeingDraggedPtr.Reset();
 
 			// We may have removed the last tab that this DockNode had.
-			if ( Tabs.Num() == 0 )
+			if (Tabs.Num() == 0)
 			{
 				// Let the DockNode know that it is no longer needed.
 				ParentTabStack->OnLastTabRemoved();
 			}
-			else
-			{
-				// Also stop showing its content; switch to the last tab that was active.
-				BringTabToFront(FMath::Max(LastForegroundTabIndex - 1, 0));
-			}
 
 			GetDockArea()->CleanUp( SDockingNode::TabRemoval_DraggedOut );
 
@@ -645,7 +659,9 @@ void SDockingTabWell::BringTabToFront( int32 TabIndexToActivate )
 	
 	if ( bActiveIndexChanging )
 	{
-		const int32 LastForegroundTabIndex = FMath::Min(ForegroundTabIndex, Tabs.Num()-1);
+		const int32 LastForegroundTabIndex = ForegroundTabIndex >= Tabs.Num() 
+			? INDEX_NONE
+			: ForegroundTabIndex;
 
 		// For positive indexes, don't go out of bounds on the array.
 		ForegroundTabIndex = FMath::Min(TabIndexToActivate, Tabs.Num()-1);
@@ -660,10 +676,20 @@ void SDockingTabWell::BringTabToFront( int32 TabIndexToActivate )
 			const TSharedPtr<SDockTab> NewForegroundTab = (ForegroundTabIndex == INDEX_NONE)
 				? TSharedPtr<SDockTab>()
 				: Tabs[ForegroundTabIndex];
-			
-			MyDockArea->GetTabManager()->GetPrivateApi().OnTabForegrounded(NewForegroundTab, PreviousForegroundTab);
 
-			FGlobalTabmanager::Get()->GetPrivateApi().OnTabForegrounded(NewForegroundTab, PreviousForegroundTab);
+			// Avoid useless broadcast
+			if (PreviousForegroundTab != NewForegroundTab)
+			{
+				TSharedRef<FGlobalTabmanager> GlobalTabManager = FGlobalTabmanager::Get();
+				TSharedPtr<FTabManager> LocalTabManager = MyDockArea->GetTabManager();
+
+				if (GlobalTabManager != LocalTabManager)
+				{
+					LocalTabManager->GetPrivateApi().OnTabForegrounded(NewForegroundTab, PreviousForegroundTab);
+				}
+
+				GlobalTabManager->GetPrivateApi().OnTabForegrounded(NewForegroundTab, PreviousForegroundTab);
+			}
 		}
 	}
 
diff --git a/Engine/Source/Runtime/Slate/Private/Framework/Docking/SPanelDrawerArea.cpp b/Engine/Source/Runtime/Slate/Private/Framework/Docking/SPanelDrawerArea.cpp
--- a/Engine/Source/Runtime/Slate/Private/Framework/Docking/SPanelDrawerArea.cpp
+++ b/Engine/Source/Runtime/Slate/Private/Framework/Docking/SPanelDrawerArea.cpp
@@ -21,7 +21,7 @@ constexpr float AnimationLength = 0.15f;
 void SPanelDrawerArea::Construct(const FArguments& InArgs, const TSharedRef<SWidget>& InContent)
 {
 	AreaContent = InContent;
-
+	PanelDrawerData = MakeShared<UE::Slate::Private::FPanelDrawerData>();
 	OpenCloseAnimation = FCurveSequence(0.f, AnimationLength, ECurveEaseFunction::QuadOut);
 
 	// Add Restoring logic here
@@ -36,7 +36,10 @@ SPanelDrawerArea::~SPanelDrawerArea()
 void SPanelDrawerArea::Tick(const FGeometry& AllottedGeometry, const double InCurrentTime, const float InDeltaTime)
 {
 	// Compute the desired final size of the panel (Always cached to avoid issue jumping issue)
-	DesiredOpenWidth = (AllottedGeometry.GetLocalSize().X - SSplitterHandleWith) * Size.OpenPanelDrawerSize / (Size.ContentSlotSize + Size.OpenPanelDrawerSize);
+
+	const FTabManager::FPanelDrawerSize& Size = PanelDrawerData->Size;
+
+	DesiredOpenWidth = (AllottedGeometry.GetLocalSize().X - SSplitterHandleWith) * Size.PanelDrawerCoefficient / (Size.MainContentCoefficient + Size.PanelDrawerCoefficient);
 
 	if (OpenCloseAnimation.IsPlaying())
 	{
@@ -44,24 +47,10 @@ void SPanelDrawerArea::Tick(const FGeometry& AllottedGeometry, const double InCu
 	}
 }
 
-void SPanelDrawerArea::TogglePanel(TSharedPtr<SDockTab> InDockTab, bool bWithAnimation)
+void SPanelDrawerArea::OpenPanel(bool bWithAnimation, const TSharedRef<UE::Slate::Private::FPanelDrawerData>& InData)
 {
-	HostedTab = InDockTab;
-	if (bIsOpen)
-	{
-		ClosePanel(bWithAnimation);
-	}
-	else
-	{
-		if (HostedTab)
-		{
-			OpenPanel(bWithAnimation);
-		}
-	}
-}
+	PanelDrawerData = InData;
 
-void SPanelDrawerArea::OpenPanel(bool bWithAnimation)
-{
 	if (bWithAnimation)
 	{
 		constexpr bool bIsOpening = true;
@@ -75,12 +64,24 @@ void SPanelDrawerArea::OpenPanel(bool bWithAnimation)
 	if (!bIsOpen)
 	{
 		bIsOpen = true;
+		RequestSaveLayout();
 		OnExternalStateChange.ExecuteIfBound();
 	}
 }
 
-void SPanelDrawerArea::ClosePanel(bool bWithAnimation)
+void SPanelDrawerArea::ClosePanel(bool bWithAnimation, bool bIsTabBeingTransfered)
 {
+	if (!bIsOpen)
+	{
+		return;
+	}
+
+	if (bIsTabBeingTransfered && PanelDrawerData->HostedTab)
+	{
+		PanelDrawerData = MakeShared<UE::Slate::Private::FPanelDrawerData>(*PanelDrawerData.Get());
+		PanelDrawerData->HostedTab.Reset();
+	}
+
 	if (bWithAnimation)
 	{
 		constexpr bool bIsOpening = false;
@@ -93,6 +94,7 @@ void SPanelDrawerArea::ClosePanel(bool bWithAnimation)
 
 	if (bIsOpen)
 	{
+		RequestSaveLayout();
 		bIsOpen = false;
 		OnExternalStateChange.ExecuteIfBound();
 	}
@@ -105,7 +107,7 @@ void SPanelDrawerArea::SetupAnimationLayout()
 		SNew(SSplitter)
 		.Orientation(Orient_Horizontal)
 		+ SSplitter::Slot()
-		.Value(bIsOpen ? 1.f : Size.ContentSlotSize)
+		.Value(bIsOpen ? 1.f : PanelDrawerData->Size.MainContentCoefficient)
 		[
 			SNew(SHorizontalBox)
 			+SHorizontalBox::Slot()
@@ -131,7 +133,7 @@ void SPanelDrawerArea::SetupAnimationLayout()
 			]
 		]
 		+ SSplitter::Slot()
-		.Value(bIsOpen ? 0.0f : Size.OpenPanelDrawerSize)
+		.Value(bIsOpen ? 0.0f : PanelDrawerData->Size.PanelDrawerCoefficient)
 		.SizeRule(SSplitter::ESizeRule::SizeToContent)
 		[
 			SNew(SBox)
@@ -146,7 +148,7 @@ void SPanelDrawerArea::SetupAnimationLayout()
 				.Offset(FMargin(0.f, 0.f, DesiredOpenWidth, 0.f))
 				.Alignment(FVector2D(0.f, 0.f))
 				[
-					HostedTab ? HostedTab->GetContent() : SNullWidget::NullWidget
+					PanelDrawerData->HostedTab ? PanelDrawerData->HostedTab->GetContent() : SNullWidget::NullWidget
 				]
 			]
 		]
@@ -159,6 +161,12 @@ void SPanelDrawerArea::SetupClosedLayout()
 	[
 		AreaContent.ToSharedRef()
 	];
+
+	if (PanelDrawerData->HostedTab)
+	{
+		PanelDrawerData = MakeShared<UE::Slate::Private::FPanelDrawerData>(*PanelDrawerData.Get());
+		PanelDrawerData->HostedTab.Reset();
+	}
 }
 
 void SPanelDrawerArea::SetupOpenedLayout()
@@ -168,8 +176,8 @@ void SPanelDrawerArea::SetupOpenedLayout()
 		SNew(SSplitter)
 		.Orientation(Orient_Horizontal)
 		+ SSplitter::Slot()
-		.OnSlotResized(this, &SPanelDrawerArea::OnContentSlotResized)
-		.Value(this, &SPanelDrawerArea::GetContentSlotSize)
+		.OnSlotResized(this, &SPanelDrawerArea::OnMainContentSlotResized)
+		.Value(this, &SPanelDrawerArea::GetMainContentCoefficient)
 		[
 			// Add Clipping for content that isn't responsive to dynamic sizes
 			SNew(SBox)
@@ -181,32 +189,41 @@ void SPanelDrawerArea::SetupOpenedLayout()
 			]
 		]
 		+ SSplitter::Slot()
-		.OnSlotResized(this, &SPanelDrawerArea::OnPanelDrawerResized)
-		.Value(this, &SPanelDrawerArea::GetPanelDrawerSize)
+		.OnSlotResized(this, &SPanelDrawerArea::OnPanelDrawerSlotResized)
+		.Value(this, &SPanelDrawerArea::GetPanelDrawerCoefficient)
 		[
-			HostedTab ? HostedTab->GetContent() : SNullWidget::NullWidget
+			// Add Clipping for content that isn't responsive to dynamic sizes
+			SNew(SBox)
+			.HAlign(HAlign_Fill)
+			.VAlign(VAlign_Fill)
+			.Clipping(EWidgetClipping::OnDemand)
+			[
+				PanelDrawerData->HostedTab ? PanelDrawerData->HostedTab->GetContent() : SNullWidget::NullWidget
+			]
 		]
 	];
 }
 
-void SPanelDrawerArea::OnContentSlotResized(float InContentSlotSize)
+void SPanelDrawerArea::OnMainContentSlotResized(float InMainContentCoefficient)
 {
-	Size.ContentSlotSize = InContentSlotSize;
+	PanelDrawerData->Size.MainContentCoefficient = InMainContentCoefficient;
+	RequestSaveLayout();
 }
 
-void SPanelDrawerArea::OnPanelDrawerResized(float InOpenPanelDrawerSize)
+void SPanelDrawerArea::OnPanelDrawerSlotResized(float InPanelDrawerCoefficient)
 {
-	Size.OpenPanelDrawerSize = InOpenPanelDrawerSize;
+	PanelDrawerData->Size.PanelDrawerCoefficient = InPanelDrawerCoefficient;
+	RequestSaveLayout();
 }
 
-float SPanelDrawerArea::GetContentSlotSize() const
+float SPanelDrawerArea::GetMainContentCoefficient() const
 {
-	return Size.ContentSlotSize;
+	return PanelDrawerData->Size.MainContentCoefficient;
 }
 
-float SPanelDrawerArea::GetPanelDrawerSize() const
+float SPanelDrawerArea::GetPanelDrawerCoefficient() const
 {
-	return Size.OpenPanelDrawerSize;
+	return PanelDrawerData->Size.PanelDrawerCoefficient;
 }
 
 void SPanelDrawerArea::PlayAnimation(bool bIsOpening)
@@ -290,64 +307,48 @@ bool SPanelDrawerArea::IsOpen() const
 
 bool SPanelDrawerArea::IsHostingTab(TSharedPtr<SDockTab> InDockTab) const
 {
-	return InDockTab == HostedTab;
-}
-
-void SPanelDrawerArea::SetHostedTab(TSharedPtr<SDockTab> InDockTab)
-{
-	HostedTab = InDockTab;
-}
-
-TSharedPtr<SDockTab> SPanelDrawerArea::GetHostedTab() const
-{
-	return HostedTab;
-}
+	if (PanelDrawerData)
+	{
+		return InDockTab == PanelDrawerData->HostedTab;
+	}
 
-const FPanelDrawerSize SPanelDrawerArea::GetSize() const
-{
-	return Size;
+	return false;
 }
 
-void SPanelDrawerArea::SetSize(const FPanelDrawerSize& InSize)
+FSimpleDelegate& SPanelDrawerArea::GetOnExternalStateChanged()
 {
-	Size = InSize;
+	return OnExternalStateChange;
 }
 
-FPanelDrawerData SPanelDrawerArea::TakeOwnershipOfHostedTab(bool bWithAnimation)
+TSharedPtr<SDockTab> SPanelDrawerArea::GetHostedTab() const
 {
-	FPanelDrawerData HostedData{ Size, bIsOpen, MoveTemp(HostedTab) };
-	if (bIsOpen)
+	if (PanelDrawerData)
 	{
-		ClosePanel(bWithAnimation);
+		return PanelDrawerData->HostedTab;
 	}
-	return HostedData;
+
+	return {};
 }
 
-void SPanelDrawerArea::TransferOwnershipOfHostedTab(FPanelDrawerData&& InDrawerData, bool bWithAnimation)
+TSharedPtr<UE::Slate::Private::FPanelDrawerData> SPanelDrawerArea::GetHostedPanelDrawerData() const
 {
-	Size = InDrawerData.Size;
-	HostedTab = MoveTemp(InDrawerData.HostedTab);
-
-	if (InDrawerData.bIsOpen)
+	if (bIsOpen)
 	{
-		OpenPanel(bWithAnimation);
-
+		return PanelDrawerData;
 	}
-}
 
-FSimpleDelegate& SPanelDrawerArea::GetOnExternalStateChanged()
-{
-	return OnExternalStateChange;
+	return {};
 }
 
-bool SPanelDrawerArea::IsHostingNomadTab() const
+void SPanelDrawerArea::RequestSaveLayout()
 {
-	if (HostedTab)
+	if (PanelDrawerData && PanelDrawerData->HostedTab)
 	{
-		return HostedTab->GetTabRole() == ETabRole::NomadTab;
+		if (TSharedPtr<FTabManager> TabManager = PanelDrawerData->HostedTab->GetTabManagerPtr())
+		{
+			TabManager->RequestSavePersistentLayout();
+		}
 	}
-
-	return false;
 }
 
 }
diff --git a/Engine/Source/Runtime/Slate/Private/Framework/Docking/SPanelDrawerArea.h b/Engine/Source/Runtime/Slate/Private/Framework/Docking/SPanelDrawerArea.h
--- a/Engine/Source/Runtime/Slate/Private/Framework/Docking/SPanelDrawerArea.h
+++ b/Engine/Source/Runtime/Slate/Private/Framework/Docking/SPanelDrawerArea.h
@@ -3,24 +3,19 @@
 #pragma once
 
 #include "Animation/CurveSequence.h"
-#include "Widgets/SCompoundWidget.h"
 #include "Application/ThrottleManager.h"
+#include "Framework/Docking/TabManager.h"
+#include "Widgets/SCompoundWidget.h"
 
+class SDockingArea;
 class SDockTab;
 
 namespace UE::Slate::Private
 {
 
-struct FPanelDrawerSize
-{
-	float ContentSlotSize = 0.8f;
-	float OpenPanelDrawerSize = 0.2f;
-};
-
 struct FPanelDrawerData
 {
-	FPanelDrawerSize Size;
-	bool bIsOpen = false;
+	FTabManager::FPanelDrawerSize Size;
 	TSharedPtr<SDockTab> HostedTab;
 };
 
@@ -39,36 +34,28 @@ public:
 
 	virtual void Tick(const FGeometry& AllottedGeometry, const double InCurrentTime, const float InDeltaTime) override;
 
-	void TogglePanel(TSharedPtr<SDockTab> InDockTab, bool bWithAnimation);
-
-	void OpenPanel(bool bWithAnimation);
-	void ClosePanel(bool bWithAnimation);
+	void OpenPanel(bool bWithAnimation, const TSharedRef<UE::Slate::Private::FPanelDrawerData>& InData);
+	void ClosePanel(bool bWithAnimation, bool bIsTabBeingTransfered = false);
 
 	bool IsOpen() const;
-	bool IsHostingNomadTab() const;
 	bool IsHostingTab(TSharedPtr<SDockTab> InDockTab) const;
-	void SetHostedTab(TSharedPtr<SDockTab> InDockTab);
 	TSharedPtr<SDockTab> GetHostedTab() const;
-
-	const FPanelDrawerSize GetSize() const;
-	void SetSize(const FPanelDrawerSize& InSize);
-
-	FPanelDrawerData TakeOwnershipOfHostedTab(bool bWithAnimation);
-	void TransferOwnershipOfHostedTab(FPanelDrawerData&& InDrawerData, bool bWithAnimation);
-
 	FSimpleDelegate& GetOnExternalStateChanged();
+	TSharedPtr<UE::Slate::Private::FPanelDrawerData> GetHostedPanelDrawerData() const;
 
 private:
+	void RequestSaveLayout();
+
 	void SetupAnimationLayout();
 
 	void SetupOpenedLayout();
 	void SetupClosedLayout();
 
-	void OnContentSlotResized(float InContentSlotSize);
-	void OnPanelDrawerResized(float InOpenPanelDrawerSize);
+	void OnMainContentSlotResized(float InContentSlotSize);
+	void OnPanelDrawerSlotResized(float InOpenPanelDrawerSize);
 
-	float GetContentSlotSize() const;
-	float GetPanelDrawerSize() const;
+	float GetMainContentCoefficient() const;
+	float GetPanelDrawerCoefficient() const;
 
 	void PlayAnimation(bool bIsOpening);
 	EActiveTimerReturnType UpdateAnimation(const double InCurrentTime, const float InDeltaTime);
@@ -82,11 +69,8 @@ private:
 
 	bool bIsOpen = false;
 
-	FPanelDrawerSize Size;
-
 	TSharedPtr<SWidget> AreaContent;
-	TSharedPtr<SDockTab> HostedTab;
-
+	TSharedPtr<UE::Slate::Private::FPanelDrawerData> PanelDrawerData;
 
 	// Animation Data
 	bool bRequestedAnimation = false;
diff --git a/Engine/Source/Runtime/Slate/Private/Framework/Docking/TabManager.cpp b/Engine/Source/Runtime/Slate/Private/Framework/Docking/TabManager.cpp
--- a/Engine/Source/Runtime/Slate/Private/Framework/Docking/TabManager.cpp
+++ b/Engine/Source/Runtime/Slate/Private/Framework/Docking/TabManager.cpp
@@ -116,6 +116,30 @@ static ETabState::Type TabStateFromString(const FString& AsString)
 	}
 }
 
+namespace UE::Slate::Private
+{
+	TSharedRef<SWindow> GetTabManagerTopWindow(const TSharedRef<SWindow>& InWindow)
+	{
+		TSharedRef<SWindow> TopLevelWindow = InWindow;
+
+		EWindowType WindowType = TopLevelWindow->GetType();
+		while (!(WindowType == EWindowType::Normal || WindowType == EWindowType::GameWindow))
+		{
+			// The window might be a menu or a tooltip if this the case check the parent window instead
+			if (TSharedPtr<SWindow> ParentWindow = TopLevelWindow->GetParentWindow())
+			{
+				TopLevelWindow = ParentWindow.ToSharedRef();
+			}
+			else
+			{
+				break;
+			}
+			WindowType = TopLevelWindow->GetType();
+		}
+
+		return TopLevelWindow;
+	}
+}
 
 //////////////////////////////////////////////////////////////////////////
 // 
@@ -266,6 +290,33 @@ TSharedRef<FTabManager::FLayoutNode> FTabManager::FLayout::NewFromString_Helper(
 		NewArea->SetSizeCoefficient((float)JsonObject->GetNumberField( TEXT("SizeCoefficient") ) );
 		NewArea->SetOrientation( local::OrientationFromString( JsonObject->GetStringField(TEXT("Orientation")) ) );
 
+		// Panel Drawer data
+		{
+			FString ActivePanelDrawerTabType;
+			if (JsonObject->TryGetStringField(TEXT("ActivePanelDrawerTab"), ActivePanelDrawerTabType))
+			{
+				FPanelDrawerTab ActivePanelDrawerTab;
+				ActivePanelDrawerTab.TabId = FTabId(*ActivePanelDrawerTabType);
+				ActivePanelDrawerTab.Size.MainContentCoefficient = JsonObject->GetNumberField(TEXT("MainContentCoefficient"));
+				ActivePanelDrawerTab.Size.PanelDrawerCoefficient = JsonObject->GetNumberField(TEXT("PanelDrawerCoefficient"));
+				NewArea->SetPanelDrawerActiveTab(MoveTemp(ActivePanelDrawerTab));
+			}
+
+			const TArray<TSharedPtr<FJsonValue>>* InactivePanelDrawerTabsPtr = nullptr;
+			if (JsonObject->TryGetArrayField(TEXT("InactivePanelDrawerTabs"), InactivePanelDrawerTabsPtr))
+			{
+				for (const TSharedPtr<FJsonValue>& PanelDrawerTabJsonValue : *InactivePanelDrawerTabsPtr)
+				{
+					TSharedPtr<FJsonObject> PanelDrawerTabJson = PanelDrawerTabJsonValue->AsObject();
+					FPanelDrawerTab PanelDrawerTab;
+					PanelDrawerTab.TabId = FTabId(*PanelDrawerTabJson->GetStringField(TEXT("TabId")));
+					PanelDrawerTab.Size.MainContentCoefficient = PanelDrawerTabJson->GetNumberField(TEXT("MainContentCoefficient"));
+					PanelDrawerTab.Size.PanelDrawerCoefficient = PanelDrawerTabJson->GetNumberField(TEXT("PanelDrawerCoefficient"));
+					NewArea->AddPanelDrawerInactiveTab(MoveTemp(PanelDrawerTab));
+				}
+			}
+		}
+
 		TArray< TSharedPtr<FJsonValue> > ChildNodeValues = JsonObject->GetArrayField(TEXT("nodes"));
 		for( int32 ChildIndex=0; ChildIndex < ChildNodeValues.Num(); ++ChildIndex )
 		{
@@ -442,7 +493,29 @@ TSharedRef<FJsonObject> FTabManager::FLayout::PersistToString_Helper(const TShar
 			JsonObj->SetNumberField( TEXT("WindowSize_Y"), NodeAsArea->UnscaledWindowSize.Y );
 			JsonObj->SetBoolField( TEXT("bIsMaximized"), NodeAsArea->bIsMaximized );
 		}
-				
+
+		if (!NodeAsArea->ActivePanelDrawerTab.TabId.TabType.IsNone())
+		{
+			JsonObj->SetStringField(TEXT("ActivePanelDrawerTab"), NodeAsArea->ActivePanelDrawerTab.TabId.ToString());
+			JsonObj->SetNumberField(TEXT("MainContentCoefficient"), NodeAsArea->ActivePanelDrawerTab.Size.MainContentCoefficient);
+			JsonObj->SetNumberField(TEXT("PanelDrawerCoefficient"), NodeAsArea->ActivePanelDrawerTab.Size.PanelDrawerCoefficient);
+		}
+
+		TArray<TSharedPtr<FJsonValue>> InactivePanelDrawerTabs;
+		for (const FPanelDrawerTab& PanelDrawerTab : NodeAsArea->InactivePanelDrawerTabs)
+		{
+			TSharedRef<FJsonObject> PanelDrawerTabJson = MakeShared<FJsonObject>();
+			PanelDrawerTabJson->SetStringField(TEXT("TabId"), PanelDrawerTab.TabId.ToString());
+			PanelDrawerTabJson->SetNumberField(TEXT("MainContentCoefficient"), PanelDrawerTab.Size.MainContentCoefficient);
+			PanelDrawerTabJson->SetNumberField(TEXT("PanelDrawerCoefficient"), PanelDrawerTab.Size.PanelDrawerCoefficient);
+			InactivePanelDrawerTabs.Add(MakeShared<FJsonValueObject>(MoveTemp(PanelDrawerTabJson)));
+		}
+
+		if (!InactivePanelDrawerTabs.IsEmpty())
+		{
+			JsonObj->SetArrayField(TEXT("InactivePanelDrawerTabs"), InactivePanelDrawerTabs);
+		}
+
 		TArray< TSharedPtr<FJsonValue> > Nodes;
 		for ( const TSharedRef<FLayoutNode>& ChildNode : NodeAsArea->GetChildNodes() )
 		{
@@ -709,7 +782,8 @@ void FTabManager::FPrivateApi::OnDockAreaClosing( const TSharedRef<SDockingArea>
 		TabManager.CollapsedDockAreas.Add( PersistentDockAreaLayout.ToSharedRef() );
 	}
 
-	TabManager.QueuedPanelDrawerTransfers.Remove(DockAreaThatIsClosing->GetParentWindow());
+	TabManager.HandleClosingAreaPanelDrawerData(DockAreaThatIsClosing);
+
 }
 
 void FTabManager::FPrivateApi::OnTabManagerClosing()
@@ -772,6 +846,11 @@ void FTabManager::FPrivateApi::SetCanDoDeferredLayoutSave(bool bInCanDoDeferredL
 	TabManager.bCanDoDeferredLayoutSave = bInCanDoDeferredLayoutSave;
 }
 
+TSharedPtr<SDockingArea> FTabManager::FPrivateApi::GetDockingAreaForPanelDrawer(const TSharedPtr<SWindow>& InWindow) const
+{
+	return TabManager.GetDockingAreaForPanelDrawer(InWindow);
+}
+
 FTabManager::FPrivateApi& FTabManager::GetPrivateApi()
 {
 	return *PrivateApi;
@@ -910,21 +989,21 @@ void FTabManager::CloseAllAreas()
 	for ( int32 LiveAreaIndex=0; LiveAreaIndex < DockAreas.Num(); ++LiveAreaIndex )
 	{
 		const TSharedPtr<SDockingArea> SomeDockArea = DockAreas[LiveAreaIndex].Pin();
-		const TSharedPtr<SWindow> ParentWindow = (SomeDockArea.IsValid())
-			? SomeDockArea->GetParentWindow()
-			: TSharedPtr<SWindow>();
-
-		if (ParentWindow.IsValid())
+		if (SomeDockArea)
 		{
-			ParentWindow->RequestDestroyWindow();
+			HandleClosingAreaPanelDrawerData(SomeDockArea);
+
+			const TSharedPtr<SWindow> ParentWindow = SomeDockArea->GetParentWindow();
+			if (ParentWindow.IsValid())
+			{
+				ParentWindow->RequestDestroyWindow();
+			}
 		}
 	}
 	DockAreas.Empty();
 
 	CollapsedDockAreas.Empty();
 	InvalidDockAreas.Empty();
-
-	QueuedPanelDrawerTransfers.Empty();
 }
 
 
@@ -1051,7 +1130,8 @@ TSharedPtr<SWidget> FTabManager::RestoreFrom(const TSharedRef<FLayout>& Layout,
 		const bool bIsPrimaryArea = ThisArea->WindowPlacement == FArea::Placement_NoWindow;
 
 		// Force the HomeScreen to be always added when restoring layouts in the primary area if enabled or just force the LevelEditor to be the first tab
-		if (bIsPrimaryArea)
+		// For now only fix the layout if the CVar is enabled
+		if (bIsPrimaryArea && UE::Editor::HomeScreen::IsHomeScreenEnabled())
 		{
 			FixLayoutLoadingPrimaryArea(ThisArea);
 		}
@@ -1106,46 +1186,49 @@ TSharedPtr<SWidget> FTabManager::RestoreFrom(const TSharedRef<FLayout>& Layout,
 	return PrimaryDockArea;
 }
 
-TSharedPtr<SWidget> FTabManager::RestorePanelDrawerFrom(const TSharedRef<FLayout>& Layout, const TSharedRef<SWidget>& InContent, const TSharedPtr<SWindow>& ParentWindow)
+TSharedPtr<SWidget> FTabManager::RestorePanelDrawer(const TSharedRef<SWidget>& InContent, const TSharedRef<SWindow>& ParentWindow)
 {
-	TSharedRef<UE::Slate::Private::SPanelDrawerArea> NewPanelDrawerArea = SNew(UE::Slate::Private::SPanelDrawerArea, InContent);
-	PanelDrawerArea = NewPanelDrawerArea;
-
-	NewPanelDrawerArea->GetOnExternalStateChanged().BindSP(this, &FTabManager::OnPanelDrawerStateChanged);
-
-	// Use a queued transfer if we have one
-	TUniquePtr<UE::Slate::Private::FPanelDrawerData> PtrToPreviousPanelData = GetQueuedNomadDrawerPanelForWindow(ParentWindow);
-	if (PtrToPreviousPanelData)
+	TSharedPtr<SDockingArea> HostArea = GetDockingAreaForPanelDrawer(ParentWindow);
+	
+	if (!ensureAlwaysMsgf(HostArea, TEXT("Cannot create a drawer in a window that is not managed by the docking system")))
 	{
-		TransferNomadPanelDrawer(MoveTemp(*PtrToPreviousPanelData.Get()));
+		return {};
 	}
 
+	TSharedRef<UE::Slate::Private::SPanelDrawerArea> NewPanelDrawerArea = SNew(UE::Slate::Private::SPanelDrawerArea, InContent);
+	NewPanelDrawerArea->GetOnExternalStateChanged().BindSP(this, &FTabManager::OnPanelDrawerStateChanged);
+	HostArea->SetPanelDrawerArea(NewPanelDrawerArea);
+
 	return MoveTemp(NewPanelDrawerArea);
 }
 
-bool FTabManager::HasPanelDrawer() const
+bool FTabManager::HasPanelDrawer(const TSharedPtr<SWindow>& ParentWindow) const
 {
-	return PanelDrawerArea.IsValid();
+	if (TSharedPtr<SDockingArea> DockingArea = GetDockingAreaForPanelDrawer(ParentWindow))
+	{
+		return DockingArea->HasPanelDrawer();
+	}
+
+	return false;
 }
 
-bool FTabManager::IsPanelDrawerOpen() const
+bool FTabManager::IsPanelDrawerOpen(const TSharedPtr<SWindow>& ParentWindow) const
 {
-	if (UE::Slate::Private::SPanelDrawerArea* PanelDrawer = PanelDrawerArea.Pin().Get())
+	if (TSharedPtr<SDockingArea> DockingArea = GetDockingAreaForPanelDrawer(ParentWindow))
 	{
-		return PanelDrawer->IsOpen();
+		return DockingArea->IsPanelDrawerOpen();
 	}
 
 	return false;
 }
 
-void FTabManager::ClosePanelDrawer()
+void FTabManager::ClosePanelDrawer(const TSharedPtr<SWindow>& ParentWindow)
 {
-	if (UE::Slate::Private::SPanelDrawerArea* PanelDrawer = PanelDrawerArea.Pin().Get())
+	if (TSharedPtr<SDockingArea> DockingArea = GetDockingAreaForPanelDrawer(ParentWindow))
 	{
-		if (PanelDrawer->IsOpen())
+		if (DockingArea->IsPanelDrawerOpen())
 		{
-			constexpr bool bWithAnimation = true;
-			PanelDrawer->ClosePanel(bWithAnimation);
+			DockingArea->ClosePanelDrawer();
 		}
 	}
 }
@@ -1377,20 +1460,16 @@ void FTabManager::DrawAttention( const TSharedRef<SDockTab>& TabToHighlight )
 
 		if (!DockingArea->TryOpenSidebarDrawer(TabToHighlight))
 		{
-			TabToHighlight->GetParentDockTabStack()->BringToFront(TabToHighlight);
+			if (TSharedPtr<SDockingTabStack> DockingTabStack = TabToHighlight->GetParentDockTabStack())
+			{
+				DockingTabStack->BringToFront(TabToHighlight);
+			}
 		}
 
 		TabToHighlight->FlashTab();
 
 		FGlobalTabmanager::Get()->UpdateMainMenu(TabToHighlight, true);
 	}
-	else if (const TSharedPtr<FTabManager> ManagerOfTabToHighlight = TabToHighlight->GetTabManagerPtr())
-	{
-
-
-		FGlobalTabmanager::Get()->DrawAttentionToTabManager(ManagerOfTabToHighlight.ToSharedRef());
-		TabToHighlight->FlashTab();
-	}
 }
 
 void FTabManager::InsertNewDocumentTab(FName PlaceholderId, FName NewTabId, const FSearchPreference& SearchPreference, const TSharedRef<SDockTab>& UnmanagedTab)
@@ -1470,7 +1549,7 @@ TSharedPtr<SDockTab> FTabManager::TryInvokeTab(const FTabId& TabId, bool bInvoke
 	return NewTab;
 }
 
-bool FTabManager::InvokeTab_SharedInternal(const FTabId& TabId, bool bInvokeAsInactive, TSharedPtr<SDockTab>& OutExistingTab)
+bool FTabManager::InvokeTab_CanInvokeTab(const FTabId& TabId) const
 {
 	if (!IsAllowedTab(TabId))
 	{
@@ -1478,7 +1557,12 @@ bool FTabManager::InvokeTab_SharedInternal(const FTabId& TabId, bool bInvokeAsIn
 		return false;
 	}
 
-	TSharedPtr<FTabSpawnerEntry> Spawner = FindTabSpawnerFor(TabId.TabType);
+	return true;
+}
+
+TSharedPtr<SDockTab> FTabManager::InvokeTab_FindOrReuseExistingTab(const FTabId& TabId) const
+{
+	TSharedPtr<const FTabSpawnerEntry> Spawner = FindTabSpawnerFor(TabId.TabType);
 
 	if (!Spawner.IsValid())
 	{
@@ -1486,46 +1570,49 @@ bool FTabManager::InvokeTab_SharedInternal(const FTabId& TabId, bool bInvokeAsIn
 	}
 	else
 	{
-		TSharedPtr<SDockTab> ExistingTab = Spawner->OnFindTabToReuse.IsBound()
+		return Spawner->OnFindTabToReuse.IsBound()
 			? Spawner->OnFindTabToReuse.Execute(TabId)
 			: Spawner->SpawnedTabPtr.Pin();
+	}
 
-		if (ExistingTab.IsValid())
+	return {};
+}
+
+void FTabManager::InvokeTab_DrawAttentionToTab(const TSharedPtr<SDockTab>& InvokedTab)
+{
+	if (InvokedTab)
+	{
+		TSharedPtr<SDockTab> MajorTab;
+		if (TSharedPtr<FTabManager> ExistingTabManager = InvokedTab->GetTabManagerPtr())
 		{
-			TSharedPtr<SDockTab> MajorTab;
-			if (TSharedPtr<FTabManager> ExistingTabManager = ExistingTab->GetTabManagerPtr())
-			{
-				MajorTab = FGlobalTabmanager::Get()->GetMajorTabForTabManager(ExistingTabManager.ToSharedRef());
-			}
+			MajorTab = FGlobalTabmanager::Get()->GetMajorTabForTabManager(ExistingTabManager.ToSharedRef());
+		}
+
+		// Rules for drawing attention to a tab:
+		// 1. Tab is not active
+		// 2. Tab's owning major tab is not in the foreground (making the tab we want to draw attention to is not visible)
+		// 3. Tab is nomad and is not in the foreground
+		// 4. Tab is not in a closed DrawerPanel
+		// If the tab is not active or the tabs major tab is not in the foreground, activate it
+		if (!InvokedTab->IsActive() || (MajorTab && !MajorTab->IsForeground()) || !InvokedTab->IsForeground())
+		{
+			bool bShouldDrawAttention = true;
 
-			// Rules for drawing attention to a tab:
-			// 1. Tab is not active
-			// 2. Tab's owning major tab is not in the foreground (making the tab we want to draw attention to is not visible)
-			// 3. Tab is nomad and is not in the foreground
-			// 4. Tab is not in closed DrawerPanel
-			// If the tab is not active or the tabs major tab is not in the foreground, activate it
-			if (!bInvokeAsInactive && (!ExistingTab->IsActive() || (MajorTab && !MajorTab->IsForeground()) || !ExistingTab->IsForeground()))
+			if (TSharedPtr<SDockingArea> DockArea = InvokedTab->GetDockArea())
 			{
-				if (TSharedPtr<UE::Slate::Private::SPanelDrawerArea> PanelDrawer = PanelDrawerArea.Pin())
-				{
-					if (!PanelDrawer->IsHostingTab(ExistingTab) || PanelDrawer->IsOpen())
-					{
-						// Draw attention to this tab if it didn't already have focus
-						DrawAttention(ExistingTab.ToSharedRef());
-					}
-				}
-				else
+				if (DockArea->GetPanelDrawerSystemHostedTab(InvokedTab->GetLayoutIdentifier()) && DockArea->GetPanelDrawerHostedTab() != InvokedTab)
 				{
-					// Draw attention to this tab if it didn't already have focus
-					DrawAttention(ExistingTab.ToSharedRef());
+					bShouldDrawAttention = false;
 				}
 			}
-			OutExistingTab = ExistingTab;
+
+			if (bShouldDrawAttention)
+			{
+				// Draw attention to this tab if it didn't already have focus
+				DrawAttention(InvokedTab.ToSharedRef());
+			}
 		}
 	}
-
-	return true;
-
 }
 
 TSharedPtr<SDockTab> FTabManager::InvokeTab_Internal(const FTabId& TabId, bool bInvokeAsInactive, bool bForceOpenWindowIfNeeded)
@@ -1547,49 +1634,42 @@ TSharedPtr<SDockTab> FTabManager::InvokeTab_Internal(const FTabId& TabId, bool b
 	//         v
 	//     * Spawn in a new window.
 
+	if (!InvokeTab_CanInvokeTab(TabId))
 	{
-		TSharedPtr<SDockTab> ExistingTab;
-		if (!InvokeTab_SharedInternal(TabId, bInvokeAsInactive, ExistingTab))
-		{
-			return {};
-		}
+		return {};
+	}
 
-		if (ExistingTab.IsValid())
+	TSharedPtr<SDockTab> Tab = InvokeTab_FindOrReuseExistingTab(TabId);
+	if (Tab.IsValid())
+	{
+		if (!RemoveFromHiddenPanelDrawerTabs(Tab))
 		{
-			bool bShouldReturnExistingTab = true;
-
-			if (TSharedPtr<FTabManager> ExistingTabManager = ExistingTab->GetTabManagerPtr())
+			if (!bInvokeAsInactive)
 			{
-				// Will destroy the hidden drawer tab 
-				if (ExistingTabManager->TakeOwnershipOfHiddenTabInDrawerPanel(ExistingTab))
-				{
-					// Let the tab die on purpose it will be recreated in a normal tab
-					bShouldReturnExistingTab = false;
-				}
+				InvokeTab_DrawAttentionToTab(Tab);
 			}
 
-			if (bShouldReturnExistingTab)
-			{
-				return ExistingTab;
-			}
+			return Tab;
 		}
 	}
 
 	// Tab is not live. Figure out where to spawn it.
 	TSharedPtr<SDockingTabStack> StackToSpawnIn = bForceOpenWindowIfNeeded ? AttemptToOpenTab( TabId, true ) : FindPotentiallyClosedTab( TabId );
-
 	if (StackToSpawnIn.IsValid())
 	{
-		const TSharedPtr<SDockTab> NewTab = SpawnTab(TabId, GetPrivateApi().GetParentWindow());
+		if (!Tab)
+		{
+			Tab = SpawnTab(TabId, GetPrivateApi().GetParentWindow());
+		}
 
-		if (NewTab.IsValid())
+		if (Tab.IsValid())
 		{
-			StackToSpawnIn->OpenTab(NewTab.ToSharedRef(), INDEX_NONE, bInvokeAsInactive);
-			NewTab->PlaySpawnAnim();
-			FGlobalTabmanager::Get()->UpdateMainMenu(NewTab.ToSharedRef(), false);
+			StackToSpawnIn->OpenTab(Tab.ToSharedRef(), INDEX_NONE, bInvokeAsInactive);
+			Tab->PlaySpawnAnim();
+			FGlobalTabmanager::Get()->UpdateMainMenu(Tab.ToSharedRef(), false);
 		}
 
-		return NewTab;
+		return Tab;
 	}
 	else if ( FGlobalTabmanager::Get() != SharedThis(this) && NomadTabSpawner->Contains(TabId.TabType) )
 	{
@@ -1598,14 +1678,16 @@ TSharedPtr<SDockTab> FTabManager::InvokeTab_Internal(const FTabId& TabId, bool b
 	}
 	else
 	{
-		const TSharedRef<FArea> NewAreaForTab = GetAreaForTabId(TabId);
+		// Not yet able to reuse inactive but hidden tab in this scenario
+		Tab.Reset();
 
+		const TSharedRef<FArea> NewAreaForTab = GetAreaForTabId(TabId);
 		NewAreaForTab
-		->Split
-		(
-			FTabManager::NewStack()
-			->AddTab( TabId, ETabState::OpenedTab )
-		);
+			->Split
+			(
+				FTabManager::NewStack()
+				->AddTab(TabId, ETabState::OpenedTab)
+			);
 
 		TSharedPtr<SDockingArea> DockingArea = RestoreArea(NewAreaForTab, GetPrivateApi().GetParentWindow());
 		if (DockingArea && DockingArea->GetAllChildTabs().Num() > 0)
@@ -1614,11 +1696,9 @@ TSharedPtr<SDockTab> FTabManager::InvokeTab_Internal(const FTabId& TabId, bool b
 			check(NewlyOpenedTab.IsValid());
 			return NewlyOpenedTab.ToSharedRef();
 		}
-		else
-		{
-			return nullptr;
-		}
 	}
+
+	return nullptr;
 }
 
 TSharedPtr<SDockingTabStack> FTabManager::FindPotentiallyClosedTab( const FTabId& ClosedTabId )
@@ -2005,6 +2085,24 @@ TSharedPtr<SDockingNode> FTabManager::RestoreArea_Helper(const TSharedRef<FLayou
 			}
 		}
 		
+		if (!NodeAsArea->ActivePanelDrawerTab.TabId.TabType.IsNone())
+		{
+			TSharedPtr<SDockTab> InvokedTab;
+			if (CanInvokeInPanelDrawer(NodeAsArea->ActivePanelDrawerTab.TabId, NewDockAreaWidget, InvokedTab))
+			{
+				TSharedPtr<SDockingArea> OtherDockingArea = InvokedTab->GetDockArea();
+
+				if (!OtherDockingArea || OtherDockingArea->GetPanelDrawerHostedTab() != InvokedTab)
+				{
+					// Restore Panel Drawer active tab
+					TSharedRef<UE::Slate::Private::FPanelDrawerData> PanelDrawerData = MakeShared<UE::Slate::Private::FPanelDrawerData>();
+					PanelDrawerData->Size = NodeAsArea->ActivePanelDrawerTab.Size;
+					PanelDrawerData->HostedTab = MoveTemp(InvokedTab);
+					NewDockAreaWidget->SetPanelDrawerHiddenActiveTab(MoveTemp(PanelDrawerData));
+				}
+			}
+		}
+
 		return NewDockAreaWidget;
 	}
 	else if ( NodeAsSplitter.IsValid() ) 
@@ -2317,114 +2415,167 @@ TSharedPtr<SDockTab> FTabManager::FindExistingLiveTab( const FTabId& TabId ) con
 					return ChildTabs[ChildTabIndex];
 				}
 			}
+
+			if (TSharedPtr<SDockTab> Tab = SomeDockArea->GetPanelDrawerSystemHostedTab(TabId))
+			{
+				return Tab;
+			}
 		}
 	}
 
-	if (TSharedPtr<UE::Slate::Private::SPanelDrawerArea> PanelDrawer = PanelDrawerArea.Pin())
+	if (TabId.InstanceId == INDEX_NONE)
 	{
-		if (PanelDrawer->IsOpen())
+		if (const TSharedRef<SDockTab>* TabPtr = PanelDrawerActiveHiddenTabs.Find(TabId))
 		{
-			if (TSharedPtr<SDockTab> HostedTab = PanelDrawer->GetHostedTab())
-			{
-				if (TabId == HostedTab->GetLayoutIdentifier())
-				{
-					return HostedTab;
-				}
-			}
+			return *TabPtr;
 		}
 	}
 
 	return TSharedPtr<SDockTab>();
 }
 
-TSharedPtr<SDockTab> FTabManager::TryOpenTabInPanelDrawer(const FTabId & TabId)
+TSharedPtr<SDockTab> FTabManager::TryOpenTabInPanelDrawer(const FTabId& TabId, const TSharedPtr<SWindow>& ParentWindow)
 {
-	return InvokeTabInPanelDrawer(TabId, false);
+	return InvokeTabInPanelDrawer(TabId, false, ParentWindow);
 }
 
-TSharedPtr<SDockTab> FTabManager::TryToggleTabInPanelDrawer(const FTabId& TabId)
+TSharedPtr<SDockTab> FTabManager::TryToggleTabInPanelDrawer(const FTabId& TabId, const TSharedPtr<SWindow>& ParentWindow)
 {
-	return InvokeTabInPanelDrawer(TabId, true);
+	return InvokeTabInPanelDrawer(TabId, true, ParentWindow);
 }
 
-TSharedPtr<SDockTab> FTabManager::InvokeTabInPanelDrawer(const FTabId& TabId, bool bToggleIfActiveInDrawer)
+TSharedPtr<SDockTab> FTabManager::InvokeTabInPanelDrawer(const FTabId& TabId, bool bToggleIfActiveInDrawer, const TSharedPtr<SWindow>& ParentWindow)
 {
-	TSharedPtr<UE::Slate::Private::SPanelDrawerArea> PanelDrawerAreaPtr = PanelDrawerArea.Pin();
-	if (CVarPanelDrawerToggle.GetValueOnGameThread() && PanelDrawerAreaPtr)
+	TSharedPtr<SDockingArea> DockingArea = GetDockingAreaForPanelDrawer(ParentWindow);
+	if ( DockingArea && DockingArea->HasPanelDrawer())
 	{
-		TSharedPtr<SDockTab> ExistingTab;
-		// Grab the existing without drawing attention
-		if (!InvokeTab_SharedInternal(TabId, true, ExistingTab))
+		TSharedPtr<SDockTab> InvokedTab;
+		if (CanInvokeInPanelDrawer(TabId, DockingArea, InvokedTab))
 		{
-			return {};
-		}
-
-		if (ExistingTab.IsValid())
-		{
-			if (PanelDrawerAreaPtr->IsOpen() && PanelDrawerAreaPtr->IsHostingTab(ExistingTab))
+			if (DockingArea->IsPanelDrawerOpen() && DockingArea->GetPanelDrawerHostedTab() == InvokedTab)
 			{
 				if (bToggleIfActiveInDrawer)
 				{
 					// Close the open panel
-					PanelDrawerAreaPtr->ClosePanel(true);
+					if (TSharedPtr<SDockingArea> OtherDockingArea = InvokedTab->GetDockArea())
+					{
+						OtherDockingArea->ClosePanelDrawer();
+					}
+					else
+					{
+						ensure(false);
+						DockingArea->ClosePanelDrawer();
+					}
 				}
 				else
 				{
-					// Draw attention to the tab
-					InvokeTab_SharedInternal(TabId, false, ExistingTab);
+					InvokeTab_DrawAttentionToTab(InvokedTab);
 				}
-				return ExistingTab;
+
+				return InvokedTab;
 			}
 
-			if (TSharedPtr<FTabManager> OtherTabManager = ExistingTab->GetTabManagerPtr())
+			if (TSharedPtr<SDockingArea> OtherDockingArea = InvokedTab->GetDockArea())
 			{
-				bool bShouldOpenDrawer = false;
-				if (TSharedPtr<UE::Slate::Private::SPanelDrawerArea> OtherPanelDrawerAreaPtr = OtherTabManager->PanelDrawerArea.Pin())
+				if (OtherDockingArea->GetPanelDrawerHostedTab() == InvokedTab)
 				{
-					if (OtherPanelDrawerAreaPtr->IsHostingTab(ExistingTab))
-					{
-						if (OtherPanelDrawerAreaPtr->IsOpen())
-						{
-							OtherPanelDrawerAreaPtr->SetHostedTab(TSharedPtr<SDockTab>());
-							OtherPanelDrawerAreaPtr->ClosePanel(true);
-							bShouldOpenDrawer = true;
-						}
-					}
-				}
-				if (OtherTabManager->TakeOwnershipOfHiddenTabInDrawerPanel(ExistingTab))
-				{
-					bShouldOpenDrawer = true;
+					OtherDockingArea->ClosePanelDrawerForTransfer();
 				}
+			}
 
-				if (bShouldOpenDrawer)
-				{
-					// Open our panel
-					PanelDrawerAreaPtr->TogglePanel(ExistingTab, true);
+			SetTabInPanelDrawer(InvokedTab, DockingArea, ParentWindow);
+			return InvokedTab;
+		}
+	}
+
+	// If no drawer is available fall be back to a normal tab
+	return TryInvokeTab(TabId);
+}
+
+bool FTabManager::CanInvokeInPanelDrawer(const FTabId& TabId, const TSharedPtr<SDockingArea>& DockingArea, TSharedPtr<SDockTab>& OutTab)
+{
+	if (!CVarPanelDrawerToggle.GetValueOnGameThread())
+	{
+		return false;
+	}
+
+	if (!InvokeTab_CanInvokeTab(TabId))
+	{
+		return {};
+	}
+
+	TSharedPtr<SDockTab> ExistingTab = InvokeTab_FindOrReuseExistingTab(TabId);
 
-					ExistingTab->SetTabManager(AsShared());
-					return ExistingTab;
+	if (ExistingTab.IsValid())
+	{
+		OutTab = ExistingTab;
+
+		// Only support nomad tab for now
+		if (ExistingTab->GetTabRole() == ETabRole::NomadTab)
+		{
+			bool bShouldOpenDrawer = ExistingTab->GetTabManagerPtr()->RemoveFromHiddenPanelDrawerTabs(ExistingTab);
+
+			if (TSharedPtr<SDockingArea> OtherDockingArea = ExistingTab->GetDockArea())
+			{
+				if (OtherDockingArea->GetPanelDrawerHostedTab() == ExistingTab)
+				{
+					bShouldOpenDrawer = true;
 				}
 			}
 
-			// Draw attention to the tab
-			InvokeTab_SharedInternal(TabId, false, ExistingTab);
-			return ExistingTab;
+			return bShouldOpenDrawer;
 		}
+	}
 
 
-		const TSharedPtr<SDockTab> NewTab = SpawnTab(TabId, GetPrivateApi().GetParentWindow());
+	const TSharedPtr<SDockTab> NewTab = SpawnTab(TabId, GetPrivateApi().GetParentWindow());
+	OutTab = NewTab;
+	if (NewTab.IsValid() && NewTab->GetTabRole() == ETabRole::NomadTab)
+	{
+		return true;
+	}
+	
+	return false;
+}
+
+void FTabManager::SetTabInPanelDrawer(const TSharedPtr<SDockTab>& InTab, const TSharedPtr<SDockingArea>& TargetDockingArea, const TSharedPtr<SWindow>& InWindow)
+{
+	if (InTab && TargetDockingArea)
+	{
+		RemoveFromHiddenPanelDrawerTabs(InTab);
 
-		if (NewTab.IsValid())
+		if (InTab->GetTabRole() == ETabRole::NomadTab)
 		{
-			PanelDrawerAreaPtr->TogglePanel(NewTab, true);
-		}
+			TSharedRef<FGlobalTabmanager> TargetManager = FGlobalTabmanager::Get();
+			TSharedPtr<SDockingArea> MajorDockingArea = TargetManager->GetDockingAreaForPanelDrawer(InWindow);
+			if (!MajorDockingArea)
+			{
+				MajorDockingArea = TargetManager->GetDockingAreaForPanelDrawer(
+					TargetManager->GetMajorTabForTabManager(TargetDockingArea->GetTabManager())->GetParentWindow()
+					);
+			}
 
-		return NewTab;
+			if (!MajorDockingArea)
+			{
+				MajorDockingArea = TargetManager->GetDockingAreaForPanelDrawer(TargetManager->GetRootWindow());
+			}
 
-	}
+			// Shouldn't happen but just in case
+			if (!ensure(MajorDockingArea))
+			{
+				return;
+			}
 
-	// If no drawer is available fall be back to a normal tab
-	return TryInvokeTab(TabId);
+			TSharedPtr<UE::Slate::Private::SPanelDrawerArea> TargetPanelDrawer = TargetDockingArea->GetPanelDrawerArea();
+			check(TargetPanelDrawer);
+			MajorDockingArea->SetPanelDrawerArea(TargetPanelDrawer.ToSharedRef());
+			MajorDockingArea->HostTabIntoPanelDrawer(InTab.ToSharedRef());
+		}
+		else
+		{
+			TargetDockingArea->HostTabIntoPanelDrawer(InTab.ToSharedRef());
+		}
+	}
 }
 
 void FTabManager::OnPanelDrawerStateChanged() const
@@ -2432,6 +2583,38 @@ void FTabManager::OnPanelDrawerStateChanged() const
 	OnPanelDrawerStateChangedDelegate.Broadcast(*this);
 }
 
+bool FTabManager::RemoveFromHiddenPanelDrawerTabs(TSharedPtr<SDockTab> ExistingTab)
+{
+	bool bWasRemoved = false;
+	if (TSharedPtr<SDockingArea> DockingArea = ExistingTab->GetDockArea())
+	{
+		bWasRemoved |= DockingArea->RemoveHiddenInactivePanelDrawerTab(ExistingTab);
+	}
+
+	if (TSharedPtr<FTabManager> TabManager = ExistingTab->GetTabManagerPtr())
+	{
+		bWasRemoved |=  TabManager->PanelDrawerActiveHiddenTabs.Remove(ExistingTab->GetLayoutIdentifier()) > 0;
+	}
+
+	return bWasRemoved;
+}
+
+void FTabManager::HandleClosingAreaPanelDrawerData(const TSharedPtr<SDockingArea>& ClosingDockingArea)
+{
+	if (ClosingDockingArea)
+	{
+		TMap<FTabId, TSharedRef<SDockTab>> PanelDrawerAliveTabs = ClosingDockingArea->GetPanelDrawerKeepAliveTabs();
+		for (TPair<FTabId, TSharedRef<SDockTab>>& Pair : PanelDrawerAliveTabs)
+		{
+			Pair.Value->SetTabManager(AsShared());
+			Pair.Value->SetParentDockingArea({});
+		}
+
+		ClosingDockingArea->CleanPanelDrawer();
+		PanelDrawerActiveHiddenTabs.Append(MoveTemp(PanelDrawerAliveTabs));
+	}
+}
+
 FTabManager::~FTabManager()
 {
 	ClearPendingLayoutSave();
@@ -2840,99 +3023,46 @@ void FTabManager::UpdateStats()
 	StaticCastSharedRef<FTabManager>(FGlobalTabmanager::Get())->UpdateStats();
 }
 
-void FTabManager::TransferNomadPanelDrawerOnRestore(UE::Slate::Private::FPanelDrawerData&& HostedTabData, const TSharedPtr<SWindow>& OwningWindow)
-{
-	if (OwningWindow && HostedTabData.HostedTab && HostedTabData.HostedTab->GetTabRole() == ETabRole::NomadTab)
-	{
-		HostedTabData.HostedTab->SetTabManager(AsShared());
-		QueuedPanelDrawerTransfers.Add(OwningWindow, MakeUnique<UE::Slate::Private::FPanelDrawerData>(MoveTemp(HostedTabData)));
-	}
-}
-
-void FTabManager::TransferNomadPanelDrawer(UE::Slate::Private::FPanelDrawerData&& HostedTabData)
-{
-	if (!HostedTabData.HostedTab || HostedTabData.HostedTab->GetTabRole() != ETabRole::NomadTab)
-	{
-		return;
-	}
-
-	if (TSharedPtr<UE::Slate::Private::SPanelDrawerArea> PanelDrawer = PanelDrawerArea.Pin())
-	{
-		HostedTabData.HostedTab->SetTabManager(AsShared());
-		constexpr bool bWithAnimation = false;
-		PanelDrawer->TransferOwnershipOfHostedTab(MoveTemp(HostedTabData), bWithAnimation);
-	}
-}
-
-TUniquePtr<UE::Slate::Private::FPanelDrawerData> FTabManager::GetQueuedNomadDrawerPanelForWindow(const TSharedPtr<SWindow>& InWindow)
+TSharedPtr<SDockingArea> FTabManager::GetDockingAreaForWindow(const TSharedRef<SWindow>& InWindow) const
 {
-	TUniquePtr<UE::Slate::Private::FPanelDrawerData> QueuedPanelDrawerData;
-	QueuedPanelDrawerTransfers.RemoveAndCopyValue(InWindow, QueuedPanelDrawerData);
-	if (!QueuedPanelDrawerData)
+	for (const TWeakPtr<SDockingArea>& WeakDockArea : DockAreas)
 	{
-		FTabManager& GlobalTabManager = FGlobalTabmanager::Get().Get();
-		if (this != &GlobalTabManager)
+		if (TSharedPtr<SDockingArea> DockArea = WeakDockArea.Pin())
 		{
-			return GlobalTabManager.GetQueuedNomadDrawerPanelForWindow(InWindow);
+			if (DockArea->GetParentWindow() == InWindow)
+			{
+				return DockArea;
+			}
 		}
 	}
 
-	return QueuedPanelDrawerData;
+	return {};
 }
 
-TUniquePtr<UE::Slate::Private::FPanelDrawerData> FTabManager::TakeOwnershipOfHiddenTabInDrawerPanel(const TSharedPtr<SDockTab>& InTab)
+TSharedPtr<SDockingArea> FTabManager::GetDockingAreaForPanelDrawer(const TSharedPtr<SWindow>& InWindow) const
 {
-	if (TSharedPtr<UE::Slate::Private::SPanelDrawerArea> PanelDrawer = PanelDrawerArea.Pin())
+	TSharedPtr<SDockingArea> PrimaryArea;
+	for (TWeakPtr<SDockingArea> LiveArea : DockAreas)
 	{
-		if (!PanelDrawer->IsOpen() && PanelDrawer->IsHostingTab(InTab))
+		if (TSharedPtr<SDockingArea> LiveDockingArea = LiveArea.Pin())
 		{
-			constexpr bool bWithAnimation = false;
-			return MakeUnique<UE::Slate::Private::FPanelDrawerData>(PanelDrawer->TakeOwnershipOfHostedTab(bWithAnimation));
-		}
-	}
-
-	TWeakPtr<SWindow> FoundInWindow;
-	for (const TPair<TWeakPtr<SWindow>, TUniquePtr<UE::Slate::Private::FPanelDrawerData>>& Pair : QueuedPanelDrawerTransfers)
-	{
-		if (Pair.Value->HostedTab == InTab)
-		{
-			FoundInWindow = Pair.Key;
-		}
-	}
-
-	TUniquePtr<UE::Slate::Private::FPanelDrawerData> QueuedPanelDrawerData;
-	if (FoundInWindow.IsValid())
-	{
-		QueuedPanelDrawerTransfers.RemoveAndCopyValue(FoundInWindow, QueuedPanelDrawerData);
-
-	}
-
-	if (!QueuedPanelDrawerData)
-	{
-		FTabManager& GlobalTabManager = FGlobalTabmanager::Get().Get();
-		if (this != &GlobalTabManager)
-		{
-			return GlobalTabManager.TakeOwnershipOfHiddenTabInDrawerPanel(InTab);
-		}
-	}
-
-	return QueuedPanelDrawerData;
-}
+			// Found the primary area (SubTabmanager primary area never own their windows)
+			TSharedPtr<SWindow> ParentWindow = LiveDockingArea->GetParentWindow();
+			if (ParentWindow == InWindow)
+			{
+				return LiveDockingArea;
+			}
 
-TSharedPtr<SDockingArea> FTabManager::GetDockingAreaForWindow(const TSharedRef<SWindow>& InWindow) const
-{
-	for (const TWeakPtr<SDockingArea>& WeakDockArea : DockAreas)
-	{
-		if (TSharedPtr<SDockingArea> DockArea = WeakDockArea.Pin())
-		{
-			if (DockArea->GetParentWindow() == InWindow)
+			// Found the primary area (SubTabmanager primary area never own their windows)
+			if (!ParentWindow)
 			{
-				return DockArea;
+				ensure(!PrimaryArea);
+				PrimaryArea = MoveTemp(LiveDockingArea);
 			}
 		}
 	}
 
-	return {};
+	return PrimaryArea;
 }
 
 void FTabManager::GetRecordableStats( int32& OutTabCount, TArray<TSharedPtr<SWindow>>& OutUniqueParentWindows ) const
@@ -2998,6 +3128,7 @@ void FTabManager::FixLayoutLoadingPrimaryArea(const TSharedRef<FArea>& InPrimary
 			}
 
 			// Remove all HomeScreen tab from the current stack and insert a new one in the correct position.
+			// While the change was hidden and tested there may have been cases where the HomeScreen was already saved in the PrimaryArea and we don't want 2 entries.
 			StackNode->Tabs.RemoveAll([HomeScreenTabId] (const FTab& InTab)
 				{
 					return InTab.TabId.TabType == HomeScreenTabId;
@@ -3180,22 +3311,7 @@ TSharedPtr<FTabManager> FGlobalTabmanager::GetTabManagerForMajorTab(const TShare
 
 TSharedPtr<FTabManager> FGlobalTabmanager::GetSubTabManagerForWindow(const TSharedRef<SWindow>& InWindow) const
 {
-	TSharedRef<SWindow> SelectedWindow = InWindow;
-
-	EWindowType WindowType = SelectedWindow->GetType();
-	while (!(WindowType == EWindowType::Normal || WindowType == EWindowType::GameWindow))
-	{
-		// The window might be a menu or a tooltip if this the case check the parent window instead
-		if (TSharedPtr<SWindow> ParentWindow = SelectedWindow->GetParentWindow())
-		{
-			SelectedWindow = ParentWindow.ToSharedRef();
-		}
-		else
-		{
-			break;
-		}
-		WindowType = SelectedWindow->GetType();
-	}
+	TSharedRef<SWindow> SelectedWindow = UE::Slate::Private::GetTabManagerTopWindow(InWindow);
 
 	for (const FSubTabManager& SubTabManager : SubTabManagers)
 	{
@@ -3350,84 +3466,25 @@ void FGlobalTabmanager::OnTabForegrounded( const TSharedPtr<SDockTab>& NewForegr
 }
 
 void FGlobalTabmanager::RelocatePanelDrawerNomadTab(FTabAndManagerForDrawer ForegroundTabAndSubManager, FTabAndManagerForDrawer BackgroundedTabAndSubManager)
-{
-	/** 
-	 * 1. Check if the background tab has a drawer with nomad tab.
-	 * 2. If not used the queued operation.
-	 * 3. Check if the new tab can receive the drawer nomad tab.
-	 * 4. If not queue the transfer.
-	 * 5. If yes then do the transfer 
-	 */
 
-	FTabManager* BackgroundedTabDrawerManager = BackgroundedTabAndSubManager.Value.Get();
-	if (!BackgroundedTabDrawerManager)
-	{
-		BackgroundedTabDrawerManager = this;
-	}
-
-	constexpr bool bWithAnimation = false;
-	UE::Slate::Private::FPanelDrawerData PreviousPanelData;
-	if (BackgroundedTabDrawerManager)
-	{
-		if (TSharedPtr<UE::Slate::Private::SPanelDrawerArea> PreviousPanelDrawer = BackgroundedTabDrawerManager->PanelDrawerArea.Pin())
-		{
-			if (PreviousPanelDrawer->IsHostingNomadTab())
-			{
-				PreviousPanelData = PreviousPanelDrawer->TakeOwnershipOfHostedTab(bWithAnimation);
-			}
-		}
-	}
-
-	TSharedPtr<SWindow> CurrentWindow;
-	if (ForegroundTabAndSubManager.Key)
-	{
-		CurrentWindow = ForegroundTabAndSubManager.Key->GetParentWindow();
-	}
-	else if (BackgroundedTabAndSubManager.Key)
-	{
-		CurrentWindow = BackgroundedTabAndSubManager.Key->GetParentWindow();
-	}
-	else
+{
+	
+	SDockTab* BackgroundedTab = BackgroundedTabAndSubManager.Key.Get();
+	FTabManager* BackgroundedTabManager = BackgroundedTabAndSubManager.Value.Get();
+	if (BackgroundedTabManager && BackgroundedTab)
 	{
-		return;
+		BackgroundedTab->GetDockArea()->SetPanelDrawerArea({});
 	}
 
-	if (!PreviousPanelData.HostedTab)
+	SDockTab* ForegroundTab = ForegroundTabAndSubManager.Key.Get();
+	FTabManager* ForegroundTabManager = ForegroundTabAndSubManager.Value.Get();
+	if (ForegroundTab && ForegroundTabManager)
 	{
-		// Use a queued transfer if we have one
-		TUniquePtr<UE::Slate::Private::FPanelDrawerData> PtrToPreviousPanelData = GetQueuedNomadDrawerPanelForWindow(CurrentWindow);
-		if (PtrToPreviousPanelData)
+		if (TSharedPtr<SDockingArea> Primary = ForegroundTabManager->GetDockingAreaForPanelDrawer(ForegroundTab->GetParentWindow()))
 		{
-			PreviousPanelData = MoveTemp(*PtrToPreviousPanelData.Get());
+			ForegroundTab->GetDockArea()->SetPanelDrawerArea(Primary->GetPanelDrawerArea());
 		}
 	}
-
-	if (!PreviousPanelData.HostedTab)
-	{
-		return;
-	}
-
-	// Check if the tab is nomad (we only transfer nomad tabs)
-	if (!ensure(PreviousPanelData.HostedTab->GetTabRole() == ETabRole::NomadTab))
-	{
-		return;
-	}
-
-	FTabManager* ForegroundTabDrawerManager = ForegroundTabAndSubManager.Value.Get();
-	if (!ForegroundTabDrawerManager)
-	{
-		ForegroundTabDrawerManager = this;
-	}
-
-	TSharedPtr<UE::Slate::Private::SPanelDrawerArea> NewPanelDrawer = ForegroundTabDrawerManager->PanelDrawerArea.Pin();
-	if (!NewPanelDrawer)
-	{
-		ForegroundTabDrawerManager->TransferNomadPanelDrawerOnRestore(MoveTemp(PreviousPanelData), CurrentWindow);
-	}
-	else
-	{
-		ForegroundTabDrawerManager->TransferNomadPanelDrawer(MoveTemp(PreviousPanelData));
-	}
 }
 
 void FGlobalTabmanager::OnTabRelocated( const TSharedRef<SDockTab>& RelocatedTab, const TSharedPtr<SWindow>& NewOwnerWindow )
@@ -3555,6 +3612,19 @@ void FGlobalTabmanager::FinishRestore()
 		if (TSharedPtr<FTabManager> Manager = SubManagerInfo.TabManager.Pin())
 		{
 			Manager->UpdateMainMenu(nullptr, false);
+
+			if (TSharedPtr<SDockTab> Tab = SubManagerInfo.MajorTab.Pin())
+			{
+				/**
+				 * Grab the primary area and restore its PanelDrawer.
+				 * When restoring the via global manager they don't restore themselves right away
+				 */
+				TSharedPtr<SDockingArea> PrimaryDockingArea = Manager->GetDockingAreaForPanelDrawer({});
+				if (!PrimaryDockingArea->IsPanelDrawerOpen())
+				{
+					PrimaryDockingArea->RestorePanelDrawerArea();
+				}
+			}
 		}
 	}
 }
diff --git a/Engine/Source/Runtime/Slate/Private/Widgets/Docking/SDockTab.cpp b/Engine/Source/Runtime/Slate/Private/Widgets/Docking/SDockTab.cpp
--- a/Engine/Source/Runtime/Slate/Private/Widgets/Docking/SDockTab.cpp
+++ b/Engine/Source/Runtime/Slate/Private/Widgets/Docking/SDockTab.cpp
@@ -343,9 +343,15 @@ const FTabId SDockTab::GetLayoutIdentifier() const
 void SDockTab::SetParent(TSharedPtr<SDockingTabWell> Parent)
 {
 	ParentPtr = Parent;
+	ParentDockingAreaPtr = nullptr;
 	OnParentSet();
 }
 
+void SDockTab::SetParentDockingArea(TSharedPtr<SDockingArea> Parent)
+{
+	ParentDockingAreaPtr = Parent;
+}
+
 TSharedPtr<SDockingTabWell> SDockTab::GetParent() const
 {
 	return ParentPtr.IsValid() ? ParentPtr.Pin() : TSharedPtr<SDockingTabWell>();
@@ -372,7 +378,17 @@ void SDockTab::RemoveTabFromParent()
 
 TSharedPtr<SDockingArea> SDockTab::GetDockArea() const
 {
-	return ParentPtr.IsValid() ? ParentPtr.Pin()->GetDockArea() : TSharedPtr<SDockingArea>();
+	if (TSharedPtr<SDockingArea> DockingArea = ParentDockingAreaPtr.Pin())
+	{
+		return DockingArea;
+	}
+
+	if (TSharedPtr<SDockingTabWell> Parent = ParentPtr.Pin())
+	{
+		return Parent->GetDockArea();
+	}
+
+	return {};
 }
 
 TSharedPtr<SWindow> SDockTab::GetParentWindow() const
diff --git a/Engine/Source/Runtime/Slate/Public/Framework/Docking/TabManager.h b/Engine/Source/Runtime/Slate/Public/Framework/Docking/TabManager.h
--- a/Engine/Source/Runtime/Slate/Public/Framework/Docking/TabManager.h
+++ b/Engine/Source/Runtime/Slate/Public/Framework/Docking/TabManager.h
@@ -40,11 +40,6 @@ namespace UE::Editor::HomeScreen
 		return bIsHomeScreenEnabled;
 	}
 }
-namespace UE::Slate::Private
-{
-	class SPanelDrawerArea;
-	struct FPanelDrawerData;
-}
 
 DECLARE_MULTICAST_DELEGATE_TwoParams(
 	FOnActiveTabChanged,
@@ -257,12 +252,12 @@ struct FTabSpawnerEntry : public FWorkspaceItem
 			bHiddenTabName = true;
 			bIsLocked = true;
 		}
-		
+
 		// Special treatment for the LevelEditor tab to be locked, this is currently intended to not be a public feature
-        if (InTabType == FName(TEXT("LevelEditor")))
-        {
-        	bIsLocked = true;
-        }
+		if (InTabType == FName(TEXT("LevelEditor")))
+		{
+			bIsLocked = true;
+		}
 	}
 
 	FTabSpawnerEntry& SetReadOnlyBehavior( const ETabReadOnlyBehavior& InReadOnlyBehavior)
@@ -678,6 +673,18 @@ class FTabManager : public TSharedFromThis<FTabManager>
 				TArray< TSharedRef<FLayoutNode> > ChildNodes;
 		};
 
+		struct FPanelDrawerSize
+		{
+			float MainContentCoefficient = 0.75f;
+			float PanelDrawerCoefficient = 0.25f;
+		};
+
+		struct FPanelDrawerTab
+		{
+			FTabId TabId;
+			FPanelDrawerSize Size;
+		};
+
 		class FArea : public FSplitter
 		{
 				friend class FTabManager;
@@ -743,6 +750,26 @@ class FTabManager : public TSharedFromThis<FTabManager>
 				{
 				}
 
+				UE_EXPERIMENTAL(5.7, "The panel drawer is still in development and might be removed. The api, implementation details and its form may change in near releases.")
+				TSharedRef<FArea> SetPanelDrawerActiveTab(FPanelDrawerTab&& InPanelDrawerTab)
+				{
+					ActivePanelDrawerTab = InPanelDrawerTab;
+					return SharedThis(this);
+				}
+
+				UE_EXPERIMENTAL(5.7, "The panel drawer is still in development and might be removed. The api, implementation details and its form may change in near releases.")
+				TSharedRef<FArea> AddPanelDrawerInactiveTab(FPanelDrawerTab&& InPanelDrawerTab)
+				{
+					InactivePanelDrawerTabs.Add(InPanelDrawerTab);
+					return SharedThis(this);
+				}
+
+				UE_EXPERIMENTAL(5.7, "The panel drawer is still in development and might be removed. The api, implementation details and its form may change in near releases.")
+				const TArray<FPanelDrawerTab>& GetPanelDrawerInactiveTabs() const
+				{
+					return InactivePanelDrawerTabs;
+				}
+
 			/**
 			 * Returns true if this FArea defines a positionally specified window, else it returns false
 			*/
@@ -765,6 +792,8 @@ class FTabManager : public TSharedFromThis<FTabManager>
 				FVector2D UnscaledWindowSize;
 				bool bIsMaximized;
 				FName ExtensionId;
+				FPanelDrawerTab ActivePanelDrawerTab;
+				TArray<FPanelDrawerTab> InactivePanelDrawerTabs;
 		};
 
 
@@ -844,6 +873,7 @@ class FTabManager : public TSharedFromThis<FTabManager>
 				void HideWindows();
 
 				void SetCanDoDeferredLayoutSave(bool bInCanDoDeferredLayoutSave);
+				TSharedPtr<SDockingArea> GetDockingAreaForPanelDrawer(const TSharedPtr<SWindow>& InWindow) const;
 
 			private:
 				FTabManager& TabManager;
@@ -973,7 +1003,7 @@ class FTabManager : public TSharedFromThis<FTabManager>
 			const EOutputCanBeNullptr RestoreAreaOutputCanBeNullptr = EOutputCanBeNullptr::Never);
 
 		UE_EXPERIMENTAL(5.7, "The panel drawer is still in development and might be removed. The api, implementation details and its form may change in near releases.")
-		SLATE_API TSharedPtr<SWidget> RestorePanelDrawerFrom(const TSharedRef<FLayout>& Layout, const TSharedRef<SWidget>& InContent, const TSharedPtr<SWindow>& ParentWindow);
+		SLATE_API TSharedPtr<SWidget> RestorePanelDrawer(const TSharedRef<SWidget>& InContent, const TSharedRef<SWindow>& ParentWindow);
 
 		SLATE_API void PopulateLocalTabSpawnerMenu( FMenuBuilder& PopulateMe );
 
@@ -1027,19 +1057,21 @@ class FTabManager : public TSharedFromThis<FTabManager>
 		SLATE_API TSharedPtr<SDockTab> FindExistingLiveTab(const FTabId& TabId) const;
 
 		/**
-		 * Try to open tab in the drawer. If it already exists elsewhere, it will draw attention to the tab.
-		 * If no drawer is available in the manager we will just open a new standard tab
+		 * Try to open tab in the PanelDrawer. If it already exists elsewhere, it will draw attention to the tab.
+		 * If the tab is open in another PanelDrawer it will transfer it to this drawer
+		 * If no PanelDrawer is available in the manager we will just open a new standard tab
 		 */
 		UE_EXPERIMENTAL(5.7, "The panel drawer is still in development and might be removed. The api, implementation details and its form may change in near releases.")
-		SLATE_API TSharedPtr<SDockTab> TryOpenTabInPanelDrawer(const FTabId& TabId);
+		SLATE_API TSharedPtr<SDockTab> TryOpenTabInPanelDrawer(const FTabId& TabId, const TSharedPtr<SWindow>& ParentWindow);
 
 		/**
-		 * Try to open tab in the drawer. If it already exists elsewhere, it will draw attention to the tab.
-		 * If the drawer is already open with the tab it will close the drawer
-		 * If no drawer is available in the manager we will just open a new standard tab
+		 * Try to open tab in the PanelDrawer. If it already exists elsewhere, it will draw attention to the tab.
+		 * If the PanelDrawer is already open with the tab it will close the drawer
+		 * If the tab is open in another PanelDrawer it will transfer it to this drawer
+		 * If no PanelDrawer is available in the manager we will just open a new standard tab
 		 */
 		UE_EXPERIMENTAL(5.7, "The panel drawer is still in development and might be removed. The api, implementation details and its form may change in near releases.")
-		SLATE_API TSharedPtr<SDockTab> TryToggleTabInPanelDrawer(const FTabId& TabId);
+		SLATE_API TSharedPtr<SDockTab> TryToggleTabInPanelDrawer(const FTabId& TabId, const TSharedPtr<SWindow>& ParentWindow);
 
 		SLATE_API virtual ~FTabManager();
 
@@ -1129,13 +1161,13 @@ class FTabManager : public TSharedFromThis<FTabManager>
 		SLATE_API const TSharedPtr<const FTabSpawnerEntry> FindTabSpawnerFor(FName TabId) const;
 
 		UE_EXPERIMENTAL(5.7, "The panel drawer is still in development and might be removed. The api, implementation details and its form may change in near releases.")
-		SLATE_API bool HasPanelDrawer() const;
+		SLATE_API bool HasPanelDrawer(const TSharedPtr<SWindow>& ParentWindow) const;
 
 		UE_EXPERIMENTAL(5.7, "The panel drawer is still in development and might be removed. The api, implementation details and its form may change in near releases.")
-		SLATE_API bool IsPanelDrawerOpen() const;
+		SLATE_API bool IsPanelDrawerOpen(const TSharedPtr<SWindow>& ParentWindow) const;
 
 		UE_EXPERIMENTAL(5.7, "The panel drawer is still in development and might be removed. The api, implementation details and its form may change in near releases.")
-		SLATE_API void ClosePanelDrawer();
+		SLATE_API void ClosePanelDrawer(const TSharedPtr<SWindow>& ParentWindow);
 
 		UE_EXPERIMENTAL(5.7, "The panel drawer is still in development and might be removed. The api, implementation details and its form may change in near releases.")
 		DECLARE_MULTICAST_DELEGATE_OneParam(FOnPanelDrawerStateChanged, const FTabManager&);
@@ -1204,7 +1236,7 @@ PRAGMA_ENABLE_EXPERIMENTAL_WARNINGS
 		 * Use this standalone RestoreSplitterContent when the output of its internal RestoreArea_Helper cannot be a nullptr.
 		 */
 		SLATE_API void RestoreSplitterContent(const TSharedRef<FSplitter>& SplitterNode, const TSharedRef<class SDockingSplitter>& SplitterWidget, const TSharedPtr<SWindow>& ParentWindow, FSidebarTabLists& OutSidebarTabs);
-		
+
 		SLATE_API bool IsValidTabForSpawning( const FTab& SomeTab ) const;
 		SLATE_API bool IsAllowedTab(const FTabId& TabId) const;
 		SLATE_API bool IsAllowedTabType(const FName TabType) const;
@@ -1266,23 +1298,17 @@ protected:
 		/** Called at the end of RestoreFrom for tab managers to complete any work after all tabs have been restored */
 		virtual void FinishRestore() {};
 
-		UE_EXPERIMENTAL(5.7, "The panel drawer is still in development and might be removed. The api, implementation details and its form may change in near releases.")
-		void TransferNomadPanelDrawerOnRestore(UE::Slate::Private::FPanelDrawerData&& HostedTabData, const TSharedPtr<SWindow>& OwningWindow);
-		
-		UE_EXPERIMENTAL(5.7, "The panel drawer is still in development and might be removed. The api, implementation details and its form may change in near releases.")
-		void TransferNomadPanelDrawer(UE::Slate::Private::FPanelDrawerData&& HostedTabData);
-
-		UE_EXPERIMENTAL(5.7, "The panel drawer is still in development and might be removed. The api, implementation details and its form may change in near releases.")
-		TUniquePtr<UE::Slate::Private::FPanelDrawerData> GetQueuedNomadDrawerPanelForWindow(const TSharedPtr<SWindow>& InWindow);
-		
-		UE_EXPERIMENTAL(5.7, "The panel drawer is still in development and might be removed. The api, implementation details and its form may change in near releases.")
-		TUniquePtr<UE::Slate::Private::FPanelDrawerData> TakeOwnershipOfHiddenTabInDrawerPanel(const TSharedPtr<SDockTab>& InTab);
 
 		/*
 		 * Search the docking area to see if we have one that is part of the window.
 		 */
 		TSharedPtr<SDockingArea> GetDockingAreaForWindow(const TSharedRef<SWindow>& InWindow) const;
 
+		/*
+		 * Search the docking area to see if we have one that is part of the window. With some customized Logic for the panel drawer.
+		 */
+		TSharedPtr<SDockingArea> GetDockingAreaForPanelDrawer(const TSharedPtr<SWindow>& InWindow) const;
+
 	private:
 		/** Checks all dock areas and adds up the number of open tabs and unique parent windows in the manager */
 		SLATE_API void GetRecordableStats( int32& OutTabCount, TArray<TSharedPtr<SWindow>>& OutUniqueParentWindows ) const;
@@ -1310,15 +1336,28 @@ protected:
 	/** Based on if the tab can have a hidden display name it will return the correct Label to use for the tab */
 	FText GetTabLabelBasedOnSpawner(const TSharedPtr<FTabSpawnerEntry>& InSpawnerEntry) const;
 
-	/*
-	 * Does the shared logic between InvokeTab_Internal, TryOpenTabInPanelDrawer and TryToggleTabInPanelDrawer
-	 */
-	bool InvokeTab_SharedInternal(const FTabId& TabId, bool bInvokeAsInactive, TSharedPtr<SDockTab>& OutExistingTab);
+	/** Does the shared logic between InvokeTab_Internal and InvokeTabInPanelDrawer */
+	bool InvokeTab_CanInvokeTab(const FTabId& TabId) const;
+	TSharedPtr<SDockTab> InvokeTab_FindOrReuseExistingTab(const FTabId& TabId) const;
+	void InvokeTab_DrawAttentionToTab(const TSharedPtr<SDockTab>& InvokedTab);
+
+	TSharedPtr<SDockTab> InvokeTabInPanelDrawer(const FTabId& TabId, bool bToggleIfActiveInDrawer, const TSharedPtr<SWindow>& ParentWindow);
+
+	/** Check if the tab can be invoked in a PanelDrawer, also return the tab since we need it to do some test */
+	bool CanInvokeInPanelDrawer(const FTabId& TabId, const TSharedPtr<SDockingArea>& DockingArea, TSharedPtr<SDockTab>& OutTab);
 
-	TSharedPtr<SDockTab> InvokeTabInPanelDrawer(const FTabId& TabId, bool bToggleIfActiveInDrawer);
+	/** Find the right area to host the tab and setup its PanelDrawer to host the tab */
+	void SetTabInPanelDrawer(const TSharedPtr<SDockTab>& InTab, const TSharedPtr<SDockingArea>& TargetDockingArea, const TSharedPtr<SWindow>& InWindow);
 
+	/** Broadcasted when a PanelDrawer is opened or closed */
 	void OnPanelDrawerStateChanged() const;
 
+	/** Remove the tab from keep alive list of the PanelDrawers. Return true if found in the PanelDrawer keep alive */
+	bool RemoveFromHiddenPanelDrawerTabs(TSharedPtr<SDockTab> ExistingTab);
+
+	/** Transfer the KeepAlive tab to TabManager and clean the PanelDrawer data */
+	void HandleClosingAreaPanelDrawerData(const TSharedPtr<SDockingArea>& ClosingDockingArea);
+
 	protected:
 		FTabSpawner TabSpawner;
 		TSharedRef<FTabSpawner> NomadTabSpawner;
@@ -1343,12 +1382,6 @@ protected:
 		/** A Major tab that contains this TabManager's widgets. */
 		TWeakPtr<SDockTab> OwnerTabPtr;
 
-		/** A optional Panel Drawer Area for that tab manager */
-		TWeakPtr<UE::Slate::Private::SPanelDrawerArea> PanelDrawerArea;
-
-		/** A queued transfer for a nomad tab in the panel drawer when restored */
-		TMap<TWeakPtr<SWindow>, TUniquePtr<UE::Slate::Private::FPanelDrawerData>> QueuedPanelDrawerTransfers;
-
 		/** The current menu multi-box for the tab, used to construct platform native main menus */
 		TSharedPtr<FMultiBox> MenuMultiBox;
 		TSharedPtr<SWidget> MenuWidget;
@@ -1418,6 +1451,9 @@ PRAGMA_DISABLE_EXPERIMENTAL_WARNINGS
 		/** Delegate that broadcast when the PanelDrawer of the manager open or close */
 		FOnPanelDrawerStateChanged OnPanelDrawerStateChangedDelegate;
 PRAGMA_ENABLE_EXPERIMENTAL_WARNINGS
+
+		/** Keep alive the PanelDrawer tabs that have no area */
+		TMap<FTabId, TSharedRef<SDockTab>> PanelDrawerActiveHiddenTabs;
 };
 
 
diff --git a/Engine/Source/Runtime/Slate/Public/Widgets/Docking/SDockTab.h b/Engine/Source/Runtime/Slate/Public/Widgets/Docking/SDockTab.h
--- a/Engine/Source/Runtime/Slate/Public/Widgets/Docking/SDockTab.h
+++ b/Engine/Source/Runtime/Slate/Public/Widgets/Docking/SDockTab.h
@@ -194,6 +194,9 @@ public:
 	/** Sets the tab's tab well parent, or resets it if nothing is passed in */
 	SLATE_API void SetParent(TSharedPtr<SDockingTabWell> Parent = TSharedPtr<SDockingTabWell>());
 
+	/** Set the tab parent when the tab live outside of a tab well (PanelDrawer for example) */
+	void SetParentDockingArea(TSharedPtr<SDockingArea> Parent);
+
 	/** Gets the tab's tab well parent, or nothing, if it has none */
 	SLATE_API TSharedPtr<SDockingTabWell> GetParent() const;
 
@@ -447,9 +450,12 @@ protected:
 	/** Determines whether the close button for the tab is shown. */
 	bool bCanEverClose;
 
-	/** The tab's parent tab well. Null if it is a floating tab. */
+	/** The tab's parent tab well. Null if it is a floating tab or hosted by a DrawerPanel. */
 	TWeakPtr<SDockingTabWell> ParentPtr;
 
+	/** The tab's parent docking area (Might be null) */
+	TWeakPtr<SDockingArea> ParentDockingAreaPtr;
+
 	/** The label on the tab */
 	TAttribute<FText> TabLabel;
 
