Fix a missing SharedPtr validity check.

diff --git a/Engine/Source/Runtime/Slate/Private/Framework/Docking/TabManager.cpp b/Engine/Source/Runtime/Slate/Private/Framework/Docking/TabManager.cpp
--- a/Engine/Source/Runtime/Slate/Private/Framework/Docking/TabManager.cpp
+++ b/Engine/Source/Runtime/Slate/Private/Framework/Docking/TabManager.cpp
@@ -2085,7 +2085,7 @@ TSharedPtr<SDockingNode> FTabManager::RestoreArea_Helper(const TSharedRef<FLayou
 			}
 		}
 		
-		if (!NodeAsArea->ActivePanelDrawerTab.TabId.TabType.IsNone())
+		if (NewDockAreaWidget && !NodeAsArea->ActivePanelDrawerTab.TabId.TabType.IsNone())
 		{
 			TSharedPtr<SDockTab> InvokedTab;
 			if (CanInvokeInPanelDrawer(NodeAsArea->ActivePanelDrawerTab.TabId, NewDockAreaWidget, InvokedTab))
