Add another missing validity check. There is some edge case I didn't account for.

diff --git a/Engine/Source/Runtime/Slate/Private/Framework/Docking/TabManager.cpp b/Engine/Source/Runtime/Slate/Private/Framework/Docking/TabManager.cpp
--- a/Engine/Source/Runtime/Slate/Private/Framework/Docking/TabManager.cpp
+++ b/Engine/Source/Runtime/Slate/Private/Framework/Docking/TabManager.cpp
@@ -3619,10 +3619,12 @@ void FGlobalTabmanager::FinishRestore()
 				 * Grab the primary area and restore its PanelDrawer.
 				 * When restoring the via global manager they don't restore themselves right away
 				 */
-				TSharedPtr<SDockingArea> PrimaryDockingArea = Manager->GetDockingAreaForPanelDrawer({});
-				if (!PrimaryDockingArea->IsPanelDrawerOpen())
+				if (TSharedPtr<SDockingArea> PrimaryDockingArea = Manager->GetDockingAreaForPanelDrawer({}))
 				{
-					PrimaryDockingArea->RestorePanelDrawerArea();
+					if (!PrimaryDockingArea->IsPanelDrawerOpen())
+					{
+						PrimaryDockingArea->RestorePanelDrawerArea();
+					}
 				}
 			}
 		}
