[HomeScreen] Remove additional CVar check

diff --git a/Engine/Source/Runtime/Slate/Private/Framework/Docking/TabManager.cpp b/Engine/Source/Runtime/Slate/Private/Framework/Docking/TabManager.cpp
--- a/Engine/Source/Runtime/Slate/Private/Framework/Docking/TabManager.cpp
+++ b/Engine/Source/Runtime/Slate/Private/Framework/Docking/TabManager.cpp
@@ -1129,9 +1129,8 @@ TSharedPtr<SWidget> FTabManager::RestoreFrom(const TSharedRef<FLayout>& Layout,
 		SetTabsTo(ThisArea, ETabState::OpenedTab, ETabState::InvalidTab);
 		const bool bIsPrimaryArea = ThisArea->WindowPlacement == FArea::Placement_NoWindow;
 
-		// Force the HomeScreen to be always added when restoring layouts in the primary area if enabled or just force the LevelEditor to be the first tab
-		// For now only fix the layout if the CVar is enabled
-		if (bIsPrimaryArea && UE::Editor::HomeScreen::IsHomeScreenEnabled())
+		// If this is the primary area containing the LevelEditor based on the HomeScreen CVar it will either force the HomeScreen and LevelEditor position (and always spawn the HomeScreen) or just the LevelEditor position
+		if (bIsPrimaryArea)
 		{
 			FixLayoutLoadingPrimaryArea(ThisArea);
 		}
@@ -3128,7 +3127,6 @@ void FTabManager::FixLayoutLoadingPrimaryArea(const TSharedRef<FArea>& InPrimary
 			}
 
 			// Remove all HomeScreen tab from the current stack and insert a new one in the correct position.
-			// While the change was hidden and tested there may have been cases where the HomeScreen was already saved in the PrimaryArea and we don't want 2 entries.
 			StackNode->Tabs.RemoveAll([HomeScreenTabId] (const FTab& InTab)
 				{
 					return InTab.TabId.TabType == HomeScreenTabId;
