Fix a low risk potential issue that was identifed on my last change. Use the full FTabId to match to the type of key in the container.

diff --git a/Engine/Source/Runtime/Slate/Private/Framework/Docking/SDockingArea.cpp b/Engine/Source/Runtime/Slate/Private/Framework/Docking/SDockingArea.cpp
--- a/Engine/Source/Runtime/Slate/Private/Framework/Docking/SDockingArea.cpp
+++ b/Engine/Source/Runtime/Slate/Private/Framework/Docking/SDockingArea.cpp
@@ -791,12 +791,12 @@ TMap<FTabId, TSharedRef<SDockTab>> SDockingArea::GetPanelDrawerKeepAliveTabs() c
 			// Only keep alive the nomad tab if we own the window
 			if (HostedTabRole != ETabRole::NomadTab || GetParentWindow())
 			{
-				Output.Add(PanelDrawerData->HostedTab->GetLayoutIdentifier().TabType, PanelDrawerData->HostedTab.ToSharedRef());
+				Output.Add(PanelDrawerData->HostedTab->GetLayoutIdentifier(), PanelDrawerData->HostedTab.ToSharedRef());
 			}
 		}
 		else if (HiddenPanelDrawerTabToReopenOnRestore && HiddenPanelDrawerTabToReopenOnRestore->HostedTab)
 		{
-			Output.Add(HiddenPanelDrawerTabToReopenOnRestore->HostedTab->GetLayoutIdentifier().TabType, HiddenPanelDrawerTabToReopenOnRestore->HostedTab.ToSharedRef());
+			Output.Add(HiddenPanelDrawerTabToReopenOnRestore->HostedTab->GetLayoutIdentifier(), HiddenPanelDrawerTabToReopenOnRestore->HostedTab.ToSharedRef());
 		}
 
 	}
