Added the TabPanelDrawer. It is the DockTab equivalent for a Tab that is hosted in a PanelDrawer. It should have most of the expected features from a DockTab but with some sightly teak to its presentation and behavior for the panel drawer.

Added the TabPanelDrawer to the PanelDrawerArea content.
The PanelDrawerArea now notifies the Global Tab Manager to change its active tab when the hosted tab in the PanelDrawer has the focus.

Tweaked the visual role of a DockTab when hosted in a PanelDrawer.
Request to close a DockTab now also work when the Tab is Hosted in a PanelDrawer.

DockingDragOperation can now host a drag operation from a tab that is not part of a DockingTabWell.

Still todo:
The Dock In layout is not functioning  at the moment. There is a behavior bug I plan to fix for 5.7 that will make it work properly.
Cosmetic: the tab label does collapse properly at the moment when there is not enough space to render it.

diff --git a/Engine/Source/Runtime/Slate/Private/Framework/Docking/FDockingDragOperation.cpp b/Engine/Source/Runtime/Slate/Private/Framework/Docking/FDockingDragOperation.cpp
--- a/Engine/Source/Runtime/Slate/Private/Framework/Docking/FDockingDragOperation.cpp
+++ b/Engine/Source/Runtime/Slate/Private/Framework/Docking/FDockingDragOperation.cpp
@@ -260,7 +260,7 @@ FDockingDragOperation::FDockingDragOperation( const TSharedRef<SDockTab>& InTabT
 	: TabBeingDragged( InTabToBeDragged )
 	, TabGrabOffsetFraction( InTabGrabOffsetFraction )
 	, TabOwnerAreaOfOrigin( InTabOwnerArea )
-	, TabStackOfOrigin( InTabToBeDragged->GetParent()->GetParentDockTabStack() )
+	, TabStackOfOrigin(InTabToBeDragged->GetParent() ? InTabToBeDragged->GetParent()->GetParentDockTabStack() : nullptr)
 	, HoveredTabPanelPtr( )
 	, HoveredDockTarget()
 	, LastContentSize( OwnerAreaSize )
@@ -297,6 +297,17 @@ FDockingDragOperation::FDockingDragOperation( const TSharedRef<SDockTab>& InTabT
 	{
 		FGlobalTabmanager::Get()->SetActiveTab(NULL);
 	}
+
+	if (!TabStackOfOrigin.IsValid())
+	{
+		// No parent, open the window right away
+		CursorDecoratorWindow->Resize(OwnerAreaSize);
+		CursorDecoratorWindow->ShowWindow();
+		CursorDecoratorWindow->ReshapeWindow(FSlateRect(FVector2D(0.0, 0.0), OwnerAreaSize));
+
+		CursorDecoratorStackNode->OpenTab(TabBeingDragged.ToSharedRef());
+		TabBeingDragged->SetDraggedOverDockArea(NULL);
+	}
 }
 
 /** @return The offset into the tab where the user grabbed in Slate Units. */
diff --git a/Engine/Source/Runtime/Slate/Private/Framework/Docking/SPanelDrawerArea.cpp b/Engine/Source/Runtime/Slate/Private/Framework/Docking/SPanelDrawerArea.cpp
--- a/Engine/Source/Runtime/Slate/Private/Framework/Docking/SPanelDrawerArea.cpp
+++ b/Engine/Source/Runtime/Slate/Private/Framework/Docking/SPanelDrawerArea.cpp
@@ -3,7 +3,9 @@
 #include "Framework/Docking/SPanelDrawerArea.h"
 
 #include "Containers/Array.h"
+#include "Framework/Docking/STabPanelDrawer.h"
 #include "Layout/Geometry.h"
+#include "Layout/WidgetPath.h"
 #include "Math/UnrealMathUtility.h"
 #include "Templates/SharedPointer.h"
 #include "Widgets/Docking/SDockTab.h"
@@ -21,7 +23,6 @@ constexpr float AnimationLength = 0.15f;
 void SPanelDrawerArea::Construct(const FArguments& InArgs, const TSharedRef<SWidget>& InContent)
 {
 	AreaContent = InContent;
-	PanelDrawerData = MakeShared<UE::Slate::Private::FPanelDrawerData>();
 	OpenCloseAnimation = FCurveSequence(0.f, AnimationLength, ECurveEaseFunction::QuadOut);
 
 	// Add Restoring logic here
@@ -36,20 +37,30 @@ SPanelDrawerArea::~SPanelDrawerArea()
 void SPanelDrawerArea::Tick(const FGeometry& AllottedGeometry, const double InCurrentTime, const float InDeltaTime)
 {
 	// Compute the desired final size of the panel (Always cached to avoid issue jumping issue)
+	if (PanelDrawerData)
+	{
+		const FTabManager::FPanelDrawerSize& Size = PanelDrawerData->Size;
 
-	const FTabManager::FPanelDrawerSize& Size = PanelDrawerData->Size;
-
-	DesiredOpenWidth = (AllottedGeometry.GetLocalSize().X - SSplitterHandleWith) * Size.PanelDrawerCoefficient / (Size.MainContentCoefficient + Size.PanelDrawerCoefficient);
+		float TotalCoefficient = Size.MainContentCoefficient + Size.PanelDrawerCoefficient;
+		if (TotalCoefficient == 0)
+		{
+			DesiredOpenWidth = (AllottedGeometry.GetLocalSize().X - SSplitterHandleWith) / 2;
+		}
+		else
+		{
+			DesiredOpenWidth = (AllottedGeometry.GetLocalSize().X - SSplitterHandleWith) * Size.PanelDrawerCoefficient / (Size.MainContentCoefficient + Size.PanelDrawerCoefficient);
+		}
 
-	if (OpenCloseAnimation.IsPlaying())
-	{
-		UpdateAnimatedSlideWidth();
+		if (OpenCloseAnimation.IsPlaying())
+		{
+			UpdateAnimatedSlideWidth();
+		}
 	}
 }
 
 void SPanelDrawerArea::OpenPanel(bool bWithAnimation, const TSharedRef<UE::Slate::Private::FPanelDrawerData>& InData)
 {
-	PanelDrawerData = InData;
+	SetPanelDrawerData(InData);
 
 	if (bWithAnimation)
 	{
@@ -76,10 +87,11 @@ void SPanelDrawerArea::ClosePanel(bool bWithAnimation, bool bIsTabBeingTransfere
 		return;
 	}
 
-	if (bIsTabBeingTransfered && PanelDrawerData->HostedTab)
+	if (bIsTabBeingTransfered && PanelDrawerData)
 	{
-		PanelDrawerData = MakeShared<UE::Slate::Private::FPanelDrawerData>(*PanelDrawerData.Get());
-		PanelDrawerData->HostedTab.Reset();
+		TSharedRef<UE::Slate::Private::FPanelDrawerData> NewPanelDrawerData = MakeShared<UE::Slate::Private::FPanelDrawerData>(*PanelDrawerData.Get());
+		NewPanelDrawerData->HostedTab.Reset();
+		SetPanelDrawerData(MoveTemp(NewPanelDrawerData));
 	}
 
 	if (bWithAnimation)
@@ -117,11 +129,12 @@ void SPanelDrawerArea::SetupAnimationLayout()
 				SNew(SBox)
 				.HAlign(HAlign_Fill)
 				.VAlign(VAlign_Fill)
-				.Clipping(EWidgetClipping::OnDemand)
+				.Clipping(EWidgetClipping::ClipToBounds)
 				[
 					AreaContent.ToSharedRef()
 				]
 			]
+			// Fake the splitter handle during the slide, only shown when that part of the splitter should be visible but can't be yet.
 			+SHorizontalBox::Slot()
 			.AutoWidth()
 			[
@@ -141,6 +154,8 @@ void SPanelDrawerArea::SetupAnimationLayout()
 			.VAlign(VAlign_Fill)
 			.WidthOverride(this, &SPanelDrawerArea::GetAnimatedWidthOverrideForPanelDrawer)
 			.Visibility(this, &SPanelDrawerArea::GetAnimatedDrawerPanelVisibility)
+			.Padding(FMargin(0.0, 0.0, 0.0, 2.0))
+			.Clipping(EWidgetClipping::ClipToBounds)
 			[
 				SNew(SConstraintCanvas)
 				+ SConstraintCanvas::Slot()
@@ -148,7 +163,7 @@ void SPanelDrawerArea::SetupAnimationLayout()
 				.Offset(FMargin(0.f, 0.f, DesiredOpenWidth, 0.f))
 				.Alignment(FVector2D(0.f, 0.f))
 				[
-					PanelDrawerData->HostedTab ? PanelDrawerData->HostedTab->GetContent() : SNullWidget::NullWidget
+					PanelDrawerContent.ToSharedRef()
 				]
 			]
 		]
@@ -162,11 +177,8 @@ void SPanelDrawerArea::SetupClosedLayout()
 		AreaContent.ToSharedRef()
 	];
 
-	if (PanelDrawerData->HostedTab)
-	{
-		PanelDrawerData = MakeShared<UE::Slate::Private::FPanelDrawerData>(*PanelDrawerData.Get());
-		PanelDrawerData->HostedTab.Reset();
-	}
+
+	SetPanelDrawerData({});
 }
 
 void SPanelDrawerArea::SetupOpenedLayout()
@@ -183,7 +195,7 @@ void SPanelDrawerArea::SetupOpenedLayout()
 			SNew(SBox)
 			.HAlign(HAlign_Fill)
 			.VAlign(VAlign_Fill)
-			.Clipping(EWidgetClipping::OnDemand)
+			.Clipping(EWidgetClipping::ClipToBounds)
 			[
 				AreaContent.ToSharedRef()
 			]
@@ -196,9 +208,10 @@ void SPanelDrawerArea::SetupOpenedLayout()
 			SNew(SBox)
 			.HAlign(HAlign_Fill)
 			.VAlign(VAlign_Fill)
-			.Clipping(EWidgetClipping::OnDemand)
+			.Padding(FMargin(0.0, 0.0, 0.0, 2.0))
+			.Clipping(EWidgetClipping::ClipToBounds)
 			[
-				PanelDrawerData->HostedTab ? PanelDrawerData->HostedTab->GetContent() : SNullWidget::NullWidget
+				PanelDrawerContent.ToSharedRef()
 			]
 		]
 	];
@@ -351,4 +364,50 @@ void SPanelDrawerArea::RequestSaveLayout()
 	}
 }
 
+TSharedRef<SWidget> SPanelDrawerArea::MakePanelDrawerContent() const
+{
+	if (PanelDrawerData && PanelDrawerData->HostedTab)
+	{
+		return SNew(SVerticalBox)
+			+ SVerticalBox::Slot()
+			.HAlign(HAlign_Fill)
+			.AutoHeight()
+			.MinHeight(27.f)
+			[
+				SNew(UE::Slate::Private::STabPanelDrawer, PanelDrawerData->HostedTab.ToSharedRef())
+			]
+			+ SVerticalBox::Slot()
+			.FillHeight(1.f)
+			[
+				PanelDrawerData->HostedTab->GetContent()
+			];
+	}
+
+	return SNullWidget::NullWidget;
+}
+
+void SPanelDrawerArea::SetPanelDrawerData(TSharedPtr<UE::Slate::Private::FPanelDrawerData>&& InNewData)
+{
+	PanelDrawerData = MoveTemp(InNewData);
+	PanelDrawerContent = MakePanelDrawerContent();
+}
+
+void SPanelDrawerArea::OnFocusChanging(const FWeakWidgetPath& PreviousFocusPath, const FWidgetPath& NewWidgetPath, const FFocusEvent& InFocusEvent)
+{
+	if (PanelDrawerContent && PanelDrawerData && PanelDrawerData->HostedTab)
+	{
+		if (IsOpen())
+		{
+			const bool bIsDrawerPanelContentFocused = NewWidgetPath.ContainsWidget(PanelDrawerContent.Get());
+
+			if (bIsDrawerPanelContentFocused)
+			{
+				// If a widget inside this tab stack got focused, activate this tab.
+				FGlobalTabmanager::Get()->SetActiveTab(PanelDrawerData->HostedTab);
+				PanelDrawerData->HostedTab->ActivateInParent(ETabActivationCause::SetDirectly);
+			}
+		}
+	}
+}
+
 }
diff --git a/Engine/Source/Runtime/Slate/Private/Framework/Docking/SPanelDrawerArea.h b/Engine/Source/Runtime/Slate/Private/Framework/Docking/SPanelDrawerArea.h
--- a/Engine/Source/Runtime/Slate/Private/Framework/Docking/SPanelDrawerArea.h
+++ b/Engine/Source/Runtime/Slate/Private/Framework/Docking/SPanelDrawerArea.h
@@ -7,8 +7,12 @@
 #include "Framework/Docking/TabManager.h"
 #include "Widgets/SCompoundWidget.h"
 
-class SDockingArea;
+class FWeakWidgetPath;
+class FWidgetPath;
 class SDockTab;
+class SDockingArea;
+
+struct FFocusEvent;
 
 namespace UE::Slate::Private
 {
@@ -66,10 +70,15 @@ private:
 	EVisibility GetAnimatedSpacerVisibility() const;
 	EVisibility GetAnimatedDrawerPanelVisibility() const;
 
+	void SetPanelDrawerData(TSharedPtr<UE::Slate::Private::FPanelDrawerData>&& InNewData);
+	TSharedRef<SWidget> MakePanelDrawerContent() const;
+	void OnFocusChanging(const FWeakWidgetPath& PreviousFocusPath, const FWidgetPath& NewWidgetPath, const FFocusEvent& InFocusEvent);
+
 
 	bool bIsOpen = false;
 
 	TSharedPtr<SWidget> AreaContent;
+	TSharedPtr<SWidget> PanelDrawerContent;
 	TSharedPtr<UE::Slate::Private::FPanelDrawerData> PanelDrawerData;
 
 	// Animation Data
diff --git a/Engine/Source/Runtime/Slate/Private/Framework/Docking/STabPanelDrawer.cpp b/Engine/Source/Runtime/Slate/Private/Framework/Docking/STabPanelDrawer.cpp
new file mode 100644
--- /dev/null
+++ b/Engine/Source/Runtime/Slate/Private/Framework/Docking/STabPanelDrawer.cpp
@@ -0,0 +1,489 @@
+// Copyright Epic Games, Inc. All Rights Reserved.
+
+#include "Framework/Docking/STabPanelDrawer.h"
+
+#include "Framework/Application/SlateApplication.h"
+#include "Framework/Docking/FDockingDragOperation.h"
+#include "Framework/Docking/SDockingArea.h"
+#include "Framework/MultiBox/MultiBoxBuilder.h"
+#include "Framework/MultiBox/MultiBoxExtender.h"
+#include "Layout/WidgetPath.h"
+#include "Styling/SlateColor.h"
+#include "Styling/SlateTypes.h"
+#include "Styling/StyleDefaults.h"
+#include "Widgets/Colors/SComplexGradient.h"
+#include "Widgets/Images/SImage.h"
+#include "Widgets/Input/SButton.h"
+#include "Widgets/Layout/SBorder.h"
+#include "Widgets/Layout/SBox.h"
+#include "Widgets/Layout/SSpacer.h"
+#include "Widgets/Text/STextBlock.h"
+
+#define LOCTEXT_NAMESPACE "TabPanelDrawer"
+
+namespace UE::Slate::Private
+{
+void STabPanelDrawer::Construct(const FArguments& InArgs, const TSharedRef<SDockTab> InTab)
+{
+	DisplayedTab = InTab;
+
+	// Post 5.7 clean a bit the logic here to share more of it with the SDockTab
+	// This makes a gradient that displays whether or not a viewport is active
+	static FLinearColor ActiveBorderColor = FAppStyle::Get().GetSlateColor("Docking.Tab.ActiveTabIndicatorColor").GetSpecifiedColor();
+	static FLinearColor ActiveBorderColorTransparent = FLinearColor(ActiveBorderColor.R, ActiveBorderColor.G, ActiveBorderColor.B, 0.0f);
+	static TArray<FLinearColor> GradientStops{ ActiveBorderColorTransparent, ActiveBorderColor, ActiveBorderColorTransparent };
+	DarkerTabBrush = *DisplayedTab->GetImageBrush();
+	DarkerTabBrush.TintColor = DisplayedTab->GetTabWellBrush()->TintColor.GetSpecifiedColor();
+
+	const FButtonStyle* const CloseButtonStyle = &DisplayedTab->GetCurrentStyle().CloseButtonStyle;
+
+	const FButtonStyle* const OpenInNewWindowButtonStyle = &FAppStyle::Get().GetWidgetStyle<FButtonStyle>("Docking.OpenTabInWindow.Button");
+
+	FMargin GrayLineMargin(2.f, 2.f, 2.f, 0.f);
+	
+	ChildSlot
+	[
+		SNew(SOverlay)
+		+ SOverlay::Slot()
+		[
+			SNew(SImage)
+			.Image(DisplayedTab->GetImageBrush())
+		]
+
+		+ SOverlay::Slot()
+		.Padding(GrayLineMargin)
+		[
+			SNew(SImage)
+			.Image(&DarkerTabBrush)
+		]
+
+		// Overlay for active tab indication.
+		+SOverlay::Slot()
+		.VAlign(VAlign_Top)
+		.HAlign(HAlign_Fill)
+
+		[
+			SNew(SComplexGradient)
+			.Visibility(DisplayedTab.Get(), &SDockTab::GetActiveTabIndicatorVisibility)
+			.DesiredSizeOverride(FVector2D(1.0f, 1.0f))
+			.GradientColors(GradientStops)
+			.Orientation(EOrientation::Orient_Vertical)
+		]
+
+		// Overlay for flashing a tab for attention
+		+SOverlay::Slot()
+		[
+			SNew(SBorder)
+			// Don't allow flasher tab overlay to absorb mouse clicks
+			.Visibility(EVisibility::HitTestInvisible)
+			.Padding(DisplayedTab.Get(), &SDockTab::GetTabPadding)
+			.BorderImage(DisplayedTab.Get(), &SDockTab::GetFlashOverlayImageBrush)
+			.BorderBackgroundColor(DisplayedTab.Get(), &SDockTab::GetFlashColor)
+		]
+
+		+ SOverlay::Slot()
+		.Padding(GrayLineMargin)
+		.VAlign(VAlign_Center)
+		.HAlign(HAlign_Fill)
+		[
+			SNew(SOverlay)
+			+ SOverlay::Slot()
+			.Padding(TAttribute<FMargin>::Create(TAttribute<FMargin>::FGetter::CreateSP(DisplayedTab.Get(), &SDockTab::GetTabPadding)))
+			.VAlign(VAlign_Center)
+			.HAlign(HAlign_Center)
+			[
+				// Label sub HBox
+				SNew(SHorizontalBox)
+				.ToolTipText(DisplayedTab.Get(), &SDockTab::GetTabLabel)
+				// Tab Label
+				+ SHorizontalBox::Slot()
+				.FillWidth(1.0f)
+				.Padding(0.0f, 1.0f)
+				.VAlign(VAlign_Center)
+				[
+					DisplayedTab->LabelWidget.ToSharedRef()
+				]
+
+				// Tab Label Suffix
+				+ SHorizontalBox::Slot()
+				.AutoWidth()
+				.Padding(0.0f, 1.0f)
+				.VAlign(VAlign_Center)
+				[
+					DisplayedTab->LabelSuffix.ToSharedRef()
+				]
+			]
+
+			+ SOverlay::Slot()
+			.Padding(TAttribute<FMargin>::Create(TAttribute<FMargin>::FGetter::CreateSP(DisplayedTab.Get(), &SDockTab::GetTabPadding)))
+			.VAlign(VAlign_Center)
+			.HAlign(HAlign_Fill)
+			[
+				SNew(SHorizontalBox)
+				.Visibility(EVisibility::Visible)
+				.ToolTip(InArgs._ToolTip)
+				.ToolTipText(DisplayedTab.Get(), &SDockTab::GetTabLabel)
+
+				// Tab Icon
+				+ SHorizontalBox::Slot()
+				.AutoWidth()
+				.VAlign(VAlign_Center)
+				.Padding(DisplayedTab->IsTabNameHidden() ? FMargin(0) : FMargin(0, 0, 5, 0))
+				[
+					SNew(SBorder)
+					// Don't allow active tab overlay to absorb mouse clicks
+					.Padding(DisplayedTab.Get(), &SDockTab::GetTabIconBorderPadding)
+					.Visibility(EVisibility::HitTestInvisible)
+					// Overlay for color-coded tab effect
+					.BorderImage(DisplayedTab.Get(), &SDockTab::GetColorOverlayImageBrush)
+					.BorderBackgroundColor(DisplayedTab.Get(), &SDockTab::GetTabColor)
+					[
+						DisplayedTab->IconWidget.ToSharedRef()
+					]
+				]
+
+				+ SHorizontalBox::Slot()
+				[
+					SNullWidget::NullWidget
+				]
+
+				+ SHorizontalBox::Slot()
+				.AutoWidth()
+				.HAlign(HAlign_Right)
+				.VAlign(VAlign_Center)
+				[
+					SNew(SButton)
+						.ButtonStyle(OpenInNewWindowButtonStyle)
+						.OnClicked(this, &STabPanelDrawer::OnButtonOpenTabInNewWindowClicked)
+						.ContentPadding(FMargin(0.0, 1.5, 0.0, 0.0))
+						.ToolTipText(LOCTEXT("OpenTabInNewWindowButtonTooltip", "Pop out into a floating window."))
+						[
+							SNew(SSpacer)
+							.Size(OpenInNewWindowButtonStyle->Normal.ImageSize)
+						]
+				]
+
+				+ SHorizontalBox::Slot()
+				.AutoWidth()
+				.HAlign(HAlign_Right)
+				.VAlign(VAlign_Center)
+				[
+					SNew(SSpacer)
+					.Size(OpenInNewWindowButtonStyle->Normal.ImageSize * 3 / 4)
+				]
+
+				+ SHorizontalBox::Slot()
+				.AutoWidth()
+				.HAlign(HAlign_Right)
+				.VAlign(VAlign_Center)
+				[
+					SNew(SButton)
+					.ButtonStyle(CloseButtonStyle) 
+					.OnClicked(this, &STabPanelDrawer::OnButtonDismissTabClicked)
+					.ContentPadding(FMargin(0.0, 1.5, 0.0, 0.0))
+					.ToolTipText(LOCTEXT("DismissPanelDrawerButtonTooltipText", "Dismiss Tab"))
+					[
+						SNew(SSpacer)
+						.Size(CloseButtonStyle->Normal.ImageSize)
+					]
+				]
+			]
+		]
+	];
+
+}
+
+void STabPanelDrawer::DismissTab() const
+{
+	if (TSharedPtr<SDockingArea> DockArea = DisplayedTab->ParentDockingAreaPtr.Pin())
+	{
+		DockArea->ClosePanelDrawer();
+	}
+	else
+	{
+		DisplayedTab->RequestCloseTab();
+	}
+}
+
+FReply STabPanelDrawer::OnButtonDismissTabClicked() const
+{
+	DismissTab();
+	return FReply::Handled();
+}
+
+FReply STabPanelDrawer::OnMouseButtonDown(const FGeometry& MyGeometry, const FPointerEvent& MouseEvent)
+{
+	if (!HasMouseCapture())
+	{
+		if (MouseEvent.GetEffectingButton() == EKeys::LeftMouseButton)
+		{
+			return FReply::Handled().DetectDrag(SharedThis(this), EKeys::LeftMouseButton);
+		}
+		else if (MouseEvent.GetEffectingButton() == EKeys::MiddleMouseButton)
+		{
+			return FReply::Handled().CaptureMouse(SharedThis(this));
+		}
+		else if (MouseEvent.GetEffectingButton() == EKeys::RightMouseButton)
+		{
+			// This should be on mouse up but for consistency with the other tabs we will put it here for now
+			FWidgetPath WidgetPath = MouseEvent.GetEventPath() != nullptr ? *MouseEvent.GetEventPath() : FWidgetPath();
+			FSlateApplication::Get().PushMenu(AsShared(), WidgetPath, MakeContextMenu(), FSlateApplication::Get().GetCursorPos(), FPopupTransitionEffect::ContextMenu);
+			return FReply::Handled();
+		}
+
+	}
+
+	return FReply::Unhandled();
+}
+
+FReply STabPanelDrawer::OnMouseButtonDoubleClick(const FGeometry& MyGeometry, const FPointerEvent& MouseEvent)
+{
+	if (MouseEvent.GetEffectingButton() == EKeys::MiddleMouseButton)
+	{
+		return FReply::Handled().CaptureMouse(SharedThis(this));
+	}
+	else
+	{
+		return FReply::Unhandled();
+	}
+}
+
+FReply STabPanelDrawer::OnMouseButtonUp(const FGeometry& MyGeometry, const FPointerEvent& MouseEvent)
+{
+	if (HasMouseCapture())
+	{
+		if (MouseEvent.GetEffectingButton() == EKeys::LeftMouseButton)
+		{
+			return FReply::Handled().ReleaseMouseCapture();
+		}
+		else if (MouseEvent.GetEffectingButton() == EKeys::MiddleMouseButton)
+		{
+			if (MyGeometry.IsUnderLocation(MouseEvent.GetScreenSpacePosition()))
+			{
+				DismissTab();
+			}
+
+			return FReply::Handled().ReleaseMouseCapture();
+		}
+	}
+	return FReply::Unhandled();
+}
+
+FReply STabPanelDrawer::OnDragDetected(const FGeometry& MyGeometry, const FPointerEvent& MouseEvent)
+{
+	// Need to remember where within a tab we grabbed
+	const FVector2D TabGrabOffset = MyGeometry.AbsoluteToLocal(MouseEvent.GetScreenSpacePosition());
+	const FVector2D TabPanelDrawerSize = MyGeometry.GetLocalSize();
+
+	FVector2D TabGrabOffsetFraction = FVector2D(
+		TabPanelDrawerSize.X != 0 ? FMath::Clamp(TabGrabOffset.X / TabPanelDrawerSize.X, 0.0, 1.0) : 0.0,
+		TabPanelDrawerSize.Y != 0 ? FMath::Clamp(TabGrabOffset.Y / TabPanelDrawerSize.Y, 0.0, 1.0) : 0.0
+		);
+
+	FVector2D OriginalSize(TabPanelDrawerSize.X, TabPanelDrawerSize.Y + DisplayedTab->GetContent()->GetTickSpaceGeometry().GetLocalSize().Y);
+
+	if (TSharedPtr<SDockingArea> PinnedParent = DisplayedTab->ParentDockingAreaPtr.Pin())
+	{
+		//See if we can drag tabs contain in this manager
+		TSharedPtr<FTabManager> TabManager = DisplayedTab->GetTabManagerPtr();
+
+
+		if (TabManager.IsValid() && TabManager->GetCanDoDragOperation())
+		{
+			PinnedParent->ClosePanelDrawerForTransfer();
+			TSharedRef<FDockingDragOperation> DragDropOperation =
+				FDockingDragOperation::New(
+					DisplayedTab.ToSharedRef(),
+					TabGrabOffsetFraction,
+					PinnedParent.ToSharedRef(),
+					OriginalSize
+				);
+
+			return FReply::Handled().BeginDragDrop(DragDropOperation);
+		}
+		else
+		{
+			return FReply::Handled();
+		}
+	}
+	
+	return FReply::Unhandled();
+}
+
+FReply STabPanelDrawer::OnTouchStarted(const FGeometry& MyGeometry, const FPointerEvent& InTouchEvent)
+{
+	if (!HasMouseCapture())
+	{
+		return FReply::Handled().CaptureMouse(SharedThis(this));
+	}
+
+	return FReply::Unhandled();
+}
+
+FReply STabPanelDrawer::OnTouchEnded(const FGeometry& MyGeometry, const FPointerEvent& InTouchEvent)
+{
+	if (HasMouseCapture())
+	{
+		return FReply::Handled().ReleaseMouseCapture();
+	}
+	return FReply::Unhandled();
+}
+
+TSharedRef<SWidget> STabPanelDrawer::MakeContextMenu() const
+{
+	constexpr bool bCloseAfterSelection = true;
+	constexpr bool bCloseSelfOnly = false;
+
+	// Legacy menu system. We should revisit it to use the UToolsMenu system
+	FMenuBuilder MenuBuilder(bCloseAfterSelection, nullptr, TSharedPtr<FExtender>(), bCloseSelfOnly, &FCoreStyle::Get());
+	{
+		MenuBuilder.BeginSection("TabPanelDrawerCloseTab");
+		{
+			MenuBuilder.AddMenuEntry(
+				LOCTEXT("DimissTab", "Dismiss Tab"),
+				LOCTEXT("DimissTabTooltip", "Close the Panel Drawer, but keep the tab alive if reopened later."),
+				FSlateIcon(),
+				FUIAction(
+					FExecuteAction::CreateSP(this, &STabPanelDrawer::DismissTab)
+					)
+				);
+
+			MenuBuilder.AddMenuEntry(
+				LOCTEXT("CloseTab", "Close Tab"),
+				LOCTEXT("CloseTabTooltil", "Close this tab."),
+				FSlateIcon(),
+				FUIAction(
+					FExecuteAction::CreateSP(DisplayedTab.Get(), &SDockTab::RemoveTabFromParent),
+					FCanExecuteAction::CreateSP(DisplayedTab.Get(), &SDockTab::CanCloseTab)
+					)
+				);
+		}
+		MenuBuilder.EndSection();
+
+		MenuBuilder.BeginSection("TabPanelDrawer", LOCTEXT("LayoutMenuSection", "Layout"));
+		{
+			MenuBuilder.AddMenuEntry(
+				LOCTEXT("OpenTabInNewWindow", "Open Tab in a new window"),
+				LOCTEXT("OpenTabInNewWindowTooltip", "Pop out into a floating window."),
+				FSlateIcon(),
+				FUIAction(
+					FExecuteAction::CreateSP(this, &STabPanelDrawer::OpenTabInNewWindow)
+				)
+			);
+
+			MenuBuilder.AddMenuEntry(
+				LOCTEXT("DockInLayout", "Dock in layout"),
+				LOCTEXT("DockInLayoutTooltip", "Open Tab outside of the Panel Drawer in its current docked layout."),
+				FSlateIcon(),
+				FUIAction(
+					FExecuteAction::CreateSP(this, &STabPanelDrawer::InvokeTabOutsidePanelDrawer)
+				)
+			);
+		}
+		MenuBuilder.EndSection();
+
+		DisplayedTab->ExtendContextMenu(MenuBuilder);
+	}
+
+	return MenuBuilder.MakeWidget();
+}
+
+FReply STabPanelDrawer::OnButtonOpenTabInNewWindowClicked() const
+{
+	OpenTabInNewWindow();
+	return FReply::Handled();
+}
+
+void STabPanelDrawer::OpenTabInNewWindow() const
+{
+	// This should be refactored into a utility function post 5.7
+	const TSharedPtr<FTabManager> MyTabManager = DisplayedTab->GetTabManagerPtr();
+	if (!MyTabManager.IsValid())
+	{
+		return;
+	}
+
+	DisplayedTab->RemoveTabFromParent();
+
+	TSharedPtr<SWindow> NewWindowParent = MyTabManager->GetPrivateApi().GetParentWindow();
+
+	TSharedRef<SWindow> NewWindow = SNew(SWindow)
+		.Title(FGlobalTabmanager::Get()->GetApplicationTitle())
+		.AutoCenter(EAutoCenter::None)
+		// Divide out scale, it is already factored into position
+		.ScreenPosition(GetCachedGeometry().LocalToAbsolute(FVector2D(0, 0)))
+		// Make room for the title bar; otherwise windows will get progressive smaller whenver you float them.
+		.ClientSize(SWindow::ComputeWindowSizeForContent(DisplayedTab->GetContent()->GetTickSpaceGeometry().GetLocalSize()))
+		.CreateTitleBar(false);
+
+	TSharedPtr<SDockingTabStack> NewDockNode;
+	TSharedPtr<FTabManager> TabManagerToUse;
+	if (DisplayedTab->GetTabRole() == ETabRole::NomadTab)
+	{
+		TabManagerToUse = FGlobalTabmanager::Get();
+		DisplayedTab->SetTabManager(FGlobalTabmanager::Get());
+	}
+	else
+	{
+		TabManagerToUse = MyTabManager;
+	}
+
+	// Create a new dockarea
+	TSharedRef<SDockingArea> NewDockArea =
+		SNew(SDockingArea, TabManagerToUse.ToSharedRef(), FTabManager::NewPrimaryArea())
+		.ParentWindow(NewWindow)
+		.InitialContent
+		(
+			SAssignNew(NewDockNode, SDockingTabStack, FTabManager::NewStack())
+		);
+
+	if (DisplayedTab->GetTabRole() == ETabRole::MajorTab || DisplayedTab->GetTabRole() == ETabRole::NomadTab)
+	{
+		TSharedPtr<SWindow> RootWindow = FGlobalTabmanager::Get()->GetRootWindow();
+		if (RootWindow.IsValid())
+		{
+			// We have a root window, so all MajorTabs are nested under it.
+			FSlateApplication::Get().AddWindowAsNativeChild(NewWindow, RootWindow.ToSharedRef())->SetContent(NewDockArea);
+		}
+		else
+		{
+			// App tabs get put in top-level windows. They show up on the taskbar.
+			FSlateApplication::Get().AddWindow(NewWindow)->SetContent(NewDockArea);
+		}
+	}
+	else
+	{
+		// Other tab types are placed in child windows. Their life is controlled by the top-level windows.
+		// They do not show up on the taskbar.
+		if (NewWindowParent.IsValid())
+		{
+			FSlateApplication::Get().AddWindowAsNativeChild(NewWindow, NewWindowParent.ToSharedRef())->SetContent(NewDockArea);
+		}
+		else
+		{
+			FSlateApplication::Get().AddWindow(NewWindow)->SetContent(NewDockArea);
+		}
+	}
+
+	// Do this after the window parenting so that the window title is set correctly
+	NewDockNode->OpenTab(DisplayedTab.ToSharedRef());
+
+	MyTabManager->GetPrivateApi().SetCanDoDeferredLayoutSave(true);
+
+	// Let every widget under this tab manager know that this tab has found a new home.
+	MyTabManager->GetPrivateApi().OnTabRelocated(DisplayedTab.ToSharedRef(), NewWindow);
+}
+
+void STabPanelDrawer::InvokeTabOutsidePanelDrawer() const
+{
+	// Try invoke tab will force the tab to open in the docking layout
+	if (TSharedPtr<FTabManager> TabManager = DisplayedTab->GetTabManagerPtr())
+	{
+		TabManager->TryInvokeTab(DisplayedTab->GetLayoutIdentifier());
+	}
+}
+
+}
+
+#undef LOCTEXT_NAMESPACE
+
diff --git a/Engine/Source/Runtime/Slate/Private/Framework/Docking/STabPanelDrawer.h b/Engine/Source/Runtime/Slate/Private/Framework/Docking/STabPanelDrawer.h
new file mode 100644
--- /dev/null
+++ b/Engine/Source/Runtime/Slate/Private/Framework/Docking/STabPanelDrawer.h
@@ -0,0 +1,44 @@
+// Copyright Epic Games, Inc. All Rights Reserved.
+
+#pragma once
+
+#include "Widgets/Docking/SDockTab.h"
+#include "Widgets/SCompoundWidget.h"
+
+#include "Templates/SharedPointer.h"
+
+namespace UE::Slate::Private
+{
+
+/**
+ * Does the display of the tab header in when hosted in a panel drawer
+ */
+class STabPanelDrawer : public SCompoundWidget
+{
+public:
+	SLATE_BEGIN_ARGS(STabPanelDrawer) {}
+	SLATE_END_ARGS()
+
+	void Construct(const FArguments& InArgs, const TSharedRef<SDockTab> InTab);
+
+private:
+
+	void DismissTab() const;
+	TSharedRef<SWidget> MakeContextMenu() const;
+	void OpenTabInNewWindow() const;
+	void InvokeTabOutsidePanelDrawer() const;
+
+	FReply OnButtonDismissTabClicked() const;
+	FReply OnMouseButtonUp(const FGeometry& MyGeometry, const FPointerEvent& MouseEvent);
+	FReply OnMouseButtonDoubleClick(const FGeometry& MyGeometry, const FPointerEvent& MouseEvent);
+	FReply OnMouseButtonDown(const FGeometry& MyGeometry, const FPointerEvent& MouseEvent);
+	FReply OnDragDetected(const FGeometry& MyGeometry, const FPointerEvent& MouseEvent);
+	FReply OnTouchStarted(const FGeometry& MyGeometry, const FPointerEvent& InTouchEvent);
+	FReply OnTouchEnded(const FGeometry& MyGeometry, const FPointerEvent& InTouchEvent);
+	FReply OnButtonOpenTabInNewWindowClicked() const;
+
+	TSharedPtr<SDockTab> DisplayedTab;
+	FSlateBrush DarkerTabBrush;
+};
+
+}
\ No newline at end of file
diff --git a/Engine/Source/Runtime/Slate/Private/Widgets/Docking/SDockTab.cpp b/Engine/Source/Runtime/Slate/Private/Widgets/Docking/SDockTab.cpp
--- a/Engine/Source/Runtime/Slate/Private/Widgets/Docking/SDockTab.cpp
+++ b/Engine/Source/Runtime/Slate/Private/Widgets/Docking/SDockTab.cpp
@@ -280,9 +280,10 @@ ETabRole SDockTab::GetVisualTabRole() const
 		{
 			bNomadMajorStyle = GetParent()->GetDockArea()->GetTabManager() == FGlobalTabmanager::Get();
 		}
-		else
+		// If we have a parent docking area but no docking tab well. We are in drawer. In this case we will not use the major style.
+		else if (!GetParent().IsValid() && !ParentDockingAreaPtr.IsValid())
 		{
-			// We are dragging or have no parent, but we are not dragging over anything, assume major
+			// We are dragging, but we are not dragging over anything, assume major
 			bNomadMajorStyle = true;
 		}
 
@@ -374,6 +375,12 @@ void SDockTab::RemoveTabFromParent()
 	{
 		ParentPtr.Pin()->RemoveAndDestroyTab(SharedThis(this), SDockingNode::TabRemoval_Closed);
 	}
+	else if (TSharedPtr<SDockingArea> DockingArea = ParentDockingAreaPtr.Pin())
+	{
+		// Tab is hosted in a panel drawer
+		DockingArea->ClosePanelDrawerForTransfer();
+		DockingArea->RemoveHiddenInactivePanelDrawerTab(SharedThis(this));
+	}
 }
 
 TSharedPtr<SDockingArea> SDockTab::GetDockArea() const
diff --git a/Engine/Source/Runtime/Slate/Public/Widgets/Docking/SDockTab.h b/Engine/Source/Runtime/Slate/Public/Widgets/Docking/SDockTab.h
--- a/Engine/Source/Runtime/Slate/Public/Widgets/Docking/SDockTab.h
+++ b/Engine/Source/Runtime/Slate/Public/Widgets/Docking/SDockTab.h
@@ -42,6 +42,11 @@ enum ETabActivationCause : uint8
 
 class FMenuBuilder;
 
+namespace UE::Slate::Private
+{
+	class STabPanelDrawer;
+}
+
 /**
  * A tab widget that also holds on to some content that should be shown when this tab is selected.
  * Intended to be used in conjunction with SDockingTabStack.
@@ -50,6 +55,8 @@ class SDockTab : public SBorder
 {
 	friend class FTabManager;
 	friend class STabSidebar;
+	friend UE::Slate::Private::STabPanelDrawer;
+
 public:
 
 	/** Invoked when a tab is closing */
diff --git a/Engine/Source/Runtime/SlateCore/Private/Styling/StarshipCoreStyle.cpp b/Engine/Source/Runtime/SlateCore/Private/Styling/StarshipCoreStyle.cpp
--- a/Engine/Source/Runtime/SlateCore/Private/Styling/StarshipCoreStyle.cpp
+++ b/Engine/Source/Runtime/SlateCore/Private/Styling/StarshipCoreStyle.cpp
@@ -1974,6 +1974,14 @@ void FStarshipCoreStyle::SetupDockingStyles(TSharedRef<FStyle>& Style)
 	Style->Set("Docking.Cross.PreviewWindowTint", FLinearColor(1.0f, 0.75f, 0.5f));
 	Style->Set("Docking.Cross.Tint", FLinearColor::White);
 	Style->Set("Docking.Cross.HoveredTint", FLinearColor::White);
+
+
+	const FButtonStyle OpenTabInWindowButton = FButtonStyle()
+		.SetNormal(IMAGE_BRUSH_SVG("Starship/Docking/OpenTabInWindow", Icon12x12, FStyleColors::Foreground))
+		.SetPressed(IMAGE_BRUSH_SVG("Starship/Docking/OpenTabInWindow", Icon12x12, FStyleColors::Foreground))
+		.SetHovered(IMAGE_BRUSH_SVG("Starship/Docking/OpenTabInWindow", Icon12x12, FStyleColors::ForegroundHover));
+
+	Style->Set("Docking.OpenTabInWindow.Button", OpenTabInWindowButton);
 }
 
 void FStarshipCoreStyle::SetupColorPickerStyles(TSharedRef<FStyle>& Style)
