Fixed a crash in the AI plugin when using the pop out into a floating window button and then pressing F1 to ask the AI about a widget.
The fix consist of splitting the RemoveTabFromParent function of a SDockTab so that we can use it as a remove from parent internally without closing the Tab inside of the docking system.

diff --git a/Engine/Source/Runtime/Slate/Private/Framework/Docking/STabPanelDrawer.cpp b/Engine/Source/Runtime/Slate/Private/Framework/Docking/STabPanelDrawer.cpp
--- a/Engine/Source/Runtime/Slate/Private/Framework/Docking/STabPanelDrawer.cpp
+++ b/Engine/Source/Runtime/Slate/Private/Framework/Docking/STabPanelDrawer.cpp
@@ -403,7 +403,7 @@ void STabPanelDrawer::OpenTabInNewWindow() const
 		return;
 	}
 
-	DisplayedTab->RemoveTabFromParent();
+	DisplayedTab->RemoveTabFromParent_Internal();
 
 	TSharedPtr<SWindow> NewWindowParent = MyTabManager->GetPrivateApi().GetParentWindow();
 
diff --git a/Engine/Source/Runtime/Slate/Private/Widgets/Docking/SDockTab.cpp b/Engine/Source/Runtime/Slate/Private/Widgets/Docking/SDockTab.cpp
--- a/Engine/Source/Runtime/Slate/Private/Widgets/Docking/SDockTab.cpp
+++ b/Engine/Source/Runtime/Slate/Private/Widgets/Docking/SDockTab.cpp
@@ -371,16 +371,7 @@ TSharedPtr<SDockingTabStack> SDockTab::GetParentDockTabStack() const
 void SDockTab::RemoveTabFromParent()
 {
 	OnTabClosed.ExecuteIfBound(SharedThis(this));
-	if (ParentPtr.IsValid())
-	{
-		ParentPtr.Pin()->RemoveAndDestroyTab(SharedThis(this), SDockingNode::TabRemoval_Closed);
-	}
-	else if (TSharedPtr<SDockingArea> DockingArea = ParentDockingAreaPtr.Pin())
-	{
-		// Tab is hosted in a panel drawer
-		DockingArea->ClosePanelDrawerForTransfer();
-		DockingArea->RemoveHiddenInactivePanelDrawerTab(SharedThis(this));
-	}
+	RemoveTabFromParent_Internal();
 }
 
 TSharedPtr<SDockingArea> SDockTab::GetDockArea() const
@@ -817,6 +808,23 @@ void SDockTab::CheckTabAllowed()
 	}
 }
 
+void SDockTab::RemoveTabFromParent_Internal()
+{
+	if (ParentPtr.IsValid())
+	{
+		ParentPtr.Pin()->RemoveAndDestroyTab(SharedThis(this), SDockingNode::TabRemoval_Closed);
+	}
+	else if (TSharedPtr<SDockingArea> DockingArea = ParentDockingAreaPtr.Pin())
+	{
+		// Tab is hosted in a panel drawer
+		DockingArea->ClosePanelDrawerForTransfer();
+		DockingArea->RemoveHiddenInactivePanelDrawerTab(SharedThis(this));
+	}
+
+	ParentPtr.Reset();
+	ParentDockingAreaPtr.Reset();
+}
+
 const FDockTabStyle& SDockTab::GetCurrentStyle() const
 {
 	if ( GetVisualTabRole() == ETabRole::MajorTab )
diff --git a/Engine/Source/Runtime/Slate/Public/Widgets/Docking/SDockTab.h b/Engine/Source/Runtime/Slate/Public/Widgets/Docking/SDockTab.h
--- a/Engine/Source/Runtime/Slate/Public/Widgets/Docking/SDockTab.h
+++ b/Engine/Source/Runtime/Slate/Public/Widgets/Docking/SDockTab.h
@@ -431,6 +431,9 @@ private:
 	/* Closes tab if permissions filter no longer allows this tab to be open. */
 	SLATE_API void CheckTabAllowed();
 
+	/** Do the removal of the tab but doesn't emit tab close events */
+	void RemoveTabFromParent_Internal();
+
 	/** The handle to the active tab activation tick */
 	TWeakPtr<FActiveTimerHandle> DragDropTimerHandle;
 	TWeakPtr<FActiveTimerHandle> UpdateStyleTimerHandle;
