Make the Behavior of the PanelDrawer more consistent and less affected by the state of the system.
Pressing the button to open the ai assistant into the drawer (The ask AI button) will always open the AI into the PanelDrawer of that window.
Opening the assitant via the window menu will always open the AI assistant in a docked tab.
This change remove the the dock in layout button on the panel drawer tab header contextual menu as it require a bit more work that will be done past this release to be usefull.

diff --git a/Engine/Source/Editor/StatusBar/Private/SStatusBar.cpp b/Engine/Source/Editor/StatusBar/Private/SStatusBar.cpp
--- a/Engine/Source/Editor/StatusBar/Private/SStatusBar.cpp
+++ b/Engine/Source/Editor/StatusBar/Private/SStatusBar.cpp
@@ -1109,7 +1109,8 @@ FReply SStatusBar::ToggleTabInPanelDrawer(FName TabId)
 	TSharedPtr<SDockTab> ParentTabPtr = ParentTab.Pin();
 	if (TSharedPtr<FTabManager> TabManager = FGlobalTabmanager::Get()->GetTabManagerForMajorTab(ParentTabPtr))
 	{
-		TabManager->TryToggleTabInPanelDrawer(TabId, ParentTabPtr->GetParentWindow());
+		constexpr bool bForceInPanelDrawer = true;
+		TabManager->TryToggleTabInPanelDrawer(TabId, ParentTabPtr->GetParentWindow(), bForceInPanelDrawer);
 	}
 
 	return FReply::Handled();
diff --git a/Engine/Source/Runtime/Slate/Private/Framework/Docking/STabPanelDrawer.cpp b/Engine/Source/Runtime/Slate/Private/Framework/Docking/STabPanelDrawer.cpp
--- a/Engine/Source/Runtime/Slate/Private/Framework/Docking/STabPanelDrawer.cpp
+++ b/Engine/Source/Runtime/Slate/Private/Framework/Docking/STabPanelDrawer.cpp
@@ -371,6 +371,7 @@ TSharedRef<SWidget> STabPanelDrawer::MakeContextMenu() const
 				)
 			);
 
+			/* Remove for the 5.7 we need to do a bit more work here for it to be usefull/
 			MenuBuilder.AddMenuEntry(
 				LOCTEXT("DockInLayout", "Dock in layout"),
 				LOCTEXT("DockInLayoutTooltip", "Open Tab outside of the Panel Drawer in its current docked layout."),
@@ -378,7 +379,7 @@ TSharedRef<SWidget> STabPanelDrawer::MakeContextMenu() const
 				FUIAction(
 					FExecuteAction::CreateSP(this, &STabPanelDrawer::InvokeTabOutsidePanelDrawer)
 				)
-			);
+			);*/
 		}
 		MenuBuilder.EndSection();
 
diff --git a/Engine/Source/Runtime/Slate/Private/Framework/Docking/TabManager.cpp b/Engine/Source/Runtime/Slate/Private/Framework/Docking/TabManager.cpp
--- a/Engine/Source/Runtime/Slate/Private/Framework/Docking/TabManager.cpp
+++ b/Engine/Source/Runtime/Slate/Private/Framework/Docking/TabManager.cpp
@@ -1641,13 +1641,21 @@ TSharedPtr<SDockTab> FTabManager::InvokeTab_Internal(const FTabId& TabId, bool b
 	TSharedPtr<SDockTab> Tab = InvokeTab_FindOrReuseExistingTab(TabId);
 	if (Tab.IsValid())
 	{
-		if (!RemoveFromHiddenPanelDrawerTabs(Tab))
+		bool bCanReuseTab = false;
+		if (!Tab->ParentPtr.IsValid())
 		{
-			if (!bInvokeAsInactive)
+			if (TSharedPtr<SDockingArea> TabDockingArea = Tab->ParentDockingAreaPtr.Pin())
 			{
-				InvokeTab_DrawAttentionToTab(Tab);
+				if (TabDockingArea->GetPanelDrawerHostedTab() == Tab || TabDockingArea->RemoveHiddenInactivePanelDrawerTab(Tab))
+				{
+					bCanReuseTab = true;
+				}
 			}
+		}
 
+		if (!bCanReuseTab)
+		{
+			InvokeTab_DrawAttentionToTab(Tab);
 			return Tab;
 		}
 	}
@@ -1660,6 +1668,11 @@ TSharedPtr<SDockTab> FTabManager::InvokeTab_Internal(const FTabId& TabId, bool b
 		{
 			Tab = SpawnTab(TabId, GetPrivateApi().GetParentWindow());
 		}
+		else
+		{
+			// Remove the tab from the panel drawer (post 5.7 make this code less spread more owned by the tab)
+			Tab->RemoveTabFromParent_Internal();
+		}
 
 		if (Tab.IsValid())
 		{
@@ -1675,11 +1688,82 @@ TSharedPtr<SDockTab> FTabManager::InvokeTab_Internal(const FTabId& TabId, bool b
 		// This tab could have been spawned in the global tab manager since it has a nomad tab spawner
 		return FGlobalTabmanager::Get()->InvokeTab_Internal(TabId, bInvokeAsInactive, bForceOpenWindowIfNeeded);
 	}
-	else
+	else if (Tab)
 	{
-		// Not yet able to reuse inactive but hidden tab in this scenario
-		Tab.Reset();
+		// Remove the tab from the panel drawer (post 5.7 make this code less spread more owned by the tab)
+		Tab->RemoveTabFromParent_Internal();
+
+		// Todo post 5.7 rework this code to remove logic duplication in the docking system (Open a window for the tab)
+		TSharedPtr<SWindow> NewWindowParent = GetPrivateApi().GetParentWindow();
+
+		TSharedRef<SWindow> NewWindow = SNew(SWindow)
+			.Title(FGlobalTabmanager::Get()->GetApplicationTitle())
+			.AutoCenter(EAutoCenter::None)
+			// Divide out scale, it is already factored into position
+			.ScreenPosition(Tab->GetContent()->GetTickSpaceGeometry().LocalToAbsolute(FVector2D(0, 0)))
+			// Make room for the title bar; otherwise windows will get progressive smaller whenver you float them.
+			.ClientSize(SWindow::ComputeWindowSizeForContent(Tab->GetContent()->GetTickSpaceGeometry().GetLocalSize()))
+			.CreateTitleBar(false);
+
+		TSharedPtr<SDockingTabStack> NewDockNode;
+		TSharedPtr<FTabManager> TabManagerToUse;
+		if (Tab->GetTabRole() == ETabRole::NomadTab)
+		{
+			TabManagerToUse = FGlobalTabmanager::Get();
+			Tab->SetTabManager(FGlobalTabmanager::Get());
+		}
+		else
+		{
+			TabManagerToUse = AsShared();
+		}
+
+		// Create a new dockarea
+		TSharedRef<SDockingArea> NewDockArea =
+			SNew(SDockingArea, TabManagerToUse.ToSharedRef(), FTabManager::NewPrimaryArea())
+			.ParentWindow(NewWindow)
+			.InitialContent
+			(
+				SAssignNew(NewDockNode, SDockingTabStack, FTabManager::NewStack())
+			);
+
+		if (Tab->GetTabRole() == ETabRole::MajorTab || Tab->GetTabRole() == ETabRole::NomadTab)
+		{
+			TSharedPtr<SWindow> RootWindow = FGlobalTabmanager::Get()->GetRootWindow();
+			if (RootWindow.IsValid())
+			{
+				// We have a root window, so all MajorTabs are nested under it.
+				FSlateApplication::Get().AddWindowAsNativeChild(NewWindow, RootWindow.ToSharedRef())->SetContent(NewDockArea);
+			}
+			else
+			{
+				// App tabs get put in top-level windows. They show up on the taskbar.
+				FSlateApplication::Get().AddWindow(NewWindow)->SetContent(NewDockArea);
+			}
+		}
+		else
+		{
+			// Other tab types are placed in child windows. Their life is controlled by the top-level windows.
+			// They do not show up on the taskbar.
+			if (NewWindowParent.IsValid())
+			{
+				FSlateApplication::Get().AddWindowAsNativeChild(NewWindow, NewWindowParent.ToSharedRef())->SetContent(NewDockArea);
+			}
+			else
+			{
+				FSlateApplication::Get().AddWindow(NewWindow)->SetContent(NewDockArea);
+			}
+		}
+
+		// Do this after the window parenting so that the window title is set correctly
+		NewDockNode->OpenTab(Tab.ToSharedRef());
 
+		// Let every widget under this tab manager know that this tab has found a new home.
+		OnTabRelocated(Tab.ToSharedRef(), NewWindow);
+
+		return Tab;
+	}
+	else
+	{
 		const TSharedRef<FArea> NewAreaForTab = GetAreaForTabId(TabId);
 		NewAreaForTab
 			->Split
@@ -2087,7 +2171,7 @@ TSharedPtr<SDockingNode> FTabManager::RestoreArea_Helper(const TSharedRef<FLayou
 		if (NewDockAreaWidget && !NodeAsArea->ActivePanelDrawerTab.TabId.TabType.IsNone())
 		{
 			TSharedPtr<SDockTab> InvokedTab;
-			if (CanInvokeInPanelDrawer(NodeAsArea->ActivePanelDrawerTab.TabId, NewDockAreaWidget, InvokedTab))
+			if (CanInvokeInPanelDrawer(NodeAsArea->ActivePanelDrawerTab.TabId, NewDockAreaWidget, InvokedTab, false))
 			{
 				TSharedPtr<SDockingArea> OtherDockingArea = InvokedTab->GetDockArea();
 
@@ -2433,23 +2517,23 @@ TSharedPtr<SDockTab> FTabManager::FindExistingLiveTab( const FTabId& TabId ) con
 	return TSharedPtr<SDockTab>();
 }
 
-TSharedPtr<SDockTab> FTabManager::TryOpenTabInPanelDrawer(const FTabId& TabId, const TSharedPtr<SWindow>& ParentWindow)
+TSharedPtr<SDockTab> FTabManager::TryOpenTabInPanelDrawer(const FTabId& TabId, const TSharedPtr<SWindow>& ParentWindow, bool bForceInDrawerPanel)
 {
-	return InvokeTabInPanelDrawer(TabId, false, ParentWindow);
+	return InvokeTabInPanelDrawer(TabId, false, ParentWindow, bForceInDrawerPanel);
 }
 
-TSharedPtr<SDockTab> FTabManager::TryToggleTabInPanelDrawer(const FTabId& TabId, const TSharedPtr<SWindow>& ParentWindow)
+TSharedPtr<SDockTab> FTabManager::TryToggleTabInPanelDrawer(const FTabId& TabId, const TSharedPtr<SWindow>& ParentWindow, bool bForceInDrawerPanel)
 {
-	return InvokeTabInPanelDrawer(TabId, true, ParentWindow);
+	return InvokeTabInPanelDrawer(TabId, true, ParentWindow, bForceInDrawerPanel);
 }
 
-TSharedPtr<SDockTab> FTabManager::InvokeTabInPanelDrawer(const FTabId& TabId, bool bToggleIfActiveInDrawer, const TSharedPtr<SWindow>& ParentWindow)
+TSharedPtr<SDockTab> FTabManager::InvokeTabInPanelDrawer(const FTabId& TabId, bool bToggleIfActiveInDrawer, const TSharedPtr<SWindow>& ParentWindow, bool bForceInDrawerPanel)
 {
 	TSharedPtr<SDockingArea> DockingArea = GetDockingAreaForPanelDrawer(ParentWindow);
 	if ( DockingArea && DockingArea->HasPanelDrawer())
 	{
 		TSharedPtr<SDockTab> InvokedTab;
-		if (CanInvokeInPanelDrawer(TabId, DockingArea, InvokedTab))
+		if (CanInvokeInPanelDrawer(TabId, DockingArea, InvokedTab, bForceInDrawerPanel))
 		{
 			if (DockingArea->IsPanelDrawerOpen() && DockingArea->GetPanelDrawerHostedTab() == InvokedTab)
 			{
@@ -2482,6 +2566,11 @@ TSharedPtr<SDockTab> FTabManager::InvokeTabInPanelDrawer(const FTabId& TabId, bo
 				}
 			}
 
+			if (bForceInDrawerPanel)
+			{
+				InvokedTab->RemoveTabFromParent_Internal();
+			}
+
 			SetTabInPanelDrawer(InvokedTab, DockingArea, ParentWindow);
 			return InvokedTab;
 		}
@@ -2491,7 +2580,7 @@ TSharedPtr<SDockTab> FTabManager::InvokeTabInPanelDrawer(const FTabId& TabId, bo
 	return TryInvokeTab(TabId);
 }
 
-bool FTabManager::CanInvokeInPanelDrawer(const FTabId& TabId, const TSharedPtr<SDockingArea>& DockingArea, TSharedPtr<SDockTab>& OutTab)
+bool FTabManager::CanInvokeInPanelDrawer(const FTabId& TabId, const TSharedPtr<SDockingArea>& DockingArea, TSharedPtr<SDockTab>& OutTab, bool bForceInPanelDrawer)
 {
 	if (!CVarPanelDrawerToggle.GetValueOnGameThread())
 	{
@@ -2522,8 +2611,14 @@ bool FTabManager::CanInvokeInPanelDrawer(const FTabId& TabId, const TSharedPtr<S
 				}
 			}
 
+			bShouldOpenDrawer |= bForceInPanelDrawer;
+
 			return bShouldOpenDrawer;
 		}
+		else
+		{
+			return false;
+		}
 	}
 
 
diff --git a/Engine/Source/Runtime/Slate/Public/Framework/Docking/TabManager.h b/Engine/Source/Runtime/Slate/Public/Framework/Docking/TabManager.h
--- a/Engine/Source/Runtime/Slate/Public/Framework/Docking/TabManager.h
+++ b/Engine/Source/Runtime/Slate/Public/Framework/Docking/TabManager.h
@@ -1062,7 +1062,7 @@ class FTabManager : public TSharedFromThis<FTabManager>
 		 * If no PanelDrawer is available in the manager we will just open a new standard tab
 		 */
 		UE_EXPERIMENTAL(5.7, "The panel drawer is still in development and might be removed. The api, implementation details and its form may change in near releases.")
-		SLATE_API TSharedPtr<SDockTab> TryOpenTabInPanelDrawer(const FTabId& TabId, const TSharedPtr<SWindow>& ParentWindow);
+		SLATE_API TSharedPtr<SDockTab> TryOpenTabInPanelDrawer(const FTabId& TabId, const TSharedPtr<SWindow>& ParentWindow, bool bForceInDrawerPanel = false);
 
 		/**
 		 * Try to open tab in the PanelDrawer. If it already exists elsewhere, it will draw attention to the tab.
@@ -1071,7 +1071,7 @@ class FTabManager : public TSharedFromThis<FTabManager>
 		 * If no PanelDrawer is available in the manager we will just open a new standard tab
 		 */
 		UE_EXPERIMENTAL(5.7, "The panel drawer is still in development and might be removed. The api, implementation details and its form may change in near releases.")
-		SLATE_API TSharedPtr<SDockTab> TryToggleTabInPanelDrawer(const FTabId& TabId, const TSharedPtr<SWindow>& ParentWindow);
+		SLATE_API TSharedPtr<SDockTab> TryToggleTabInPanelDrawer(const FTabId& TabId, const TSharedPtr<SWindow>& ParentWindow, bool bForceInDrawerPanel = false);
 
 		SLATE_API virtual ~FTabManager();
 
@@ -1341,10 +1341,10 @@ protected:
 	TSharedPtr<SDockTab> InvokeTab_FindOrReuseExistingTab(const FTabId& TabId) const;
 	void InvokeTab_DrawAttentionToTab(const TSharedPtr<SDockTab>& InvokedTab);
 
-	TSharedPtr<SDockTab> InvokeTabInPanelDrawer(const FTabId& TabId, bool bToggleIfActiveInDrawer, const TSharedPtr<SWindow>& ParentWindow);
+	TSharedPtr<SDockTab> InvokeTabInPanelDrawer(const FTabId& TabId, bool bToggleIfActiveInDrawer, const TSharedPtr<SWindow>& ParentWindow, bool bForceInDrawerPanel);
 
 	/** Check if the tab can be invoked in a PanelDrawer, also return the tab since we need it to do some test */
-	bool CanInvokeInPanelDrawer(const FTabId& TabId, const TSharedPtr<SDockingArea>& DockingArea, TSharedPtr<SDockTab>& OutTab);
+	bool CanInvokeInPanelDrawer(const FTabId& TabId, const TSharedPtr<SDockingArea>& DockingArea, TSharedPtr<SDockTab>& OutTab, bool bForceInDrawerPanel);
 
 	/** Find the right area to host the tab and setup its PanelDrawer to host the tab */
 	void SetTabInPanelDrawer(const TSharedPtr<SDockTab>& InTab, const TSharedPtr<SDockingArea>& TargetDockingArea, const TSharedPtr<SWindow>& InWindow);
