Fix incorrect per character wrapping in multiline texts.

#tests Some manual tests to validate that the bug with rich text was fixed, that there were no regression + auto test

diff --git a/Engine/Source/Runtime/Slate/Private/Framework/Text/TextLayout.cpp b/Engine/Source/Runtime/Slate/Private/Framework/Text/TextLayout.cpp
--- a/Engine/Source/Runtime/Slate/Private/Framework/Text/TextLayout.cpp
+++ b/Engine/Source/Runtime/Slate/Private/Framework/Text/TextLayout.cpp
@@ -639,7 +639,7 @@ void FTextLayout::FlowLineLayout(const int32 LineModelIndex, const float Wrappin
 					float BreakWidth = 0.0f;
 					for (int32 BreakRunIndex = CurrentRunIndex; BreakRunIndex < LineModel.Runs.Num(); ++BreakRunIndex)
 					{
-						const FRunModel& RunModel = LineModel.Runs[CurrentRunIndex];
+						const FRunModel& RunModel = LineModel.Runs[BreakRunIndex];
 						const FTextRange RunRange = RunModel.GetTextRange();
 
 						const int32 StartIndex = FMath::Max(RunRange.BeginIndex - NonBreakingStringIndexOffset, PreviousBreak);
