Quick workaround for memory corruption in build caused by resetting pointers in deleted object.

#rnx

diff --git a/Engine/Source/Runtime/Slate/Private/Widgets/Docking/SDockTab.cpp b/Engine/Source/Runtime/Slate/Private/Widgets/Docking/SDockTab.cpp
--- a/Engine/Source/Runtime/Slate/Private/Widgets/Docking/SDockTab.cpp
+++ b/Engine/Source/Runtime/Slate/Private/Widgets/Docking/SDockTab.cpp
@@ -810,15 +810,17 @@ void SDockTab::CheckTabAllowed()
 
 void SDockTab::RemoveTabFromParent_Internal()
 {
+	// Keep this object alive until the end of this method since it may die by the time RemoveAndDestroyTab returns
+	TSharedRef<SDockTab> This = SharedThis(this);
 	if (ParentPtr.IsValid())
 	{
-		ParentPtr.Pin()->RemoveAndDestroyTab(SharedThis(this), SDockingNode::TabRemoval_Closed);
+		ParentPtr.Pin()->RemoveAndDestroyTab(This, SDockingNode::TabRemoval_Closed);
 	}
 	else if (TSharedPtr<SDockingArea> DockingArea = ParentDockingAreaPtr.Pin())
 	{
 		// Tab is hosted in a panel drawer
 		DockingArea->ClosePanelDrawerForTransfer();
-		DockingArea->RemoveHiddenInactivePanelDrawerTab(SharedThis(this));
+		DockingArea->RemoveHiddenInactivePanelDrawerTab(This);
 	}
 
 	ParentPtr.Reset();
