Fix for ScrollBox not stopping in EndInertia.

Currently, if the user tried to scroll with the mouse wheel and left clicks, the ScrollBox would continue scrolling even though we expect it not to. To avoid this, we set the DesiredScrollOffset to the current offset to stop immediately

diff --git a/Engine/Source/Runtime/Slate/Private/Widgets/Layout/SScrollBox.cpp b/Engine/Source/Runtime/Slate/Private/Widgets/Layout/SScrollBox.cpp
--- a/Engine/Source/Runtime/Slate/Private/Widgets/Layout/SScrollBox.cpp
+++ b/Engine/Source/Runtime/Slate/Private/Widgets/Layout/SScrollBox.cpp
@@ -1384,4 +1384,7 @@ void SScrollBox::EndInertialScrolling()
 
 	// Zero the scroll velocity so the panel stops immediately on mouse down, even if the user does not drag
 	InertialScrollManager.ClearScrollVelocity();
+
+	// Sets the desired scroll offset to the current one so the panel does not continue moving
+	DesiredScrollOffset = ScrollPanel->PhysicalOffset;
 }
