Fix PanelDrawer issue with the Mac WebBrowser by emiting an event when the drawer is opened or closed and add the tab workaround for the web browser of the ai assitant web browser.
Disable the sidebar feature for the ai assistant, the webbrowser is still broken for those on mac and the panel drawer can crash with those at the moment.
Also force the Mac webbrowser to recompute its size and position when the visibilty of the web window is changed.

diff --git a/Engine/Config/Mac/BaseMacEngine.ini b/Engine/Config/Mac/BaseMacEngine.ini
--- a/Engine/Config/Mac/BaseMacEngine.ini
+++ b/Engine/Config/Mac/BaseMacEngine.ini
@@ -36,6 +36,8 @@ QualityLevelMapping="high"
 r.HairStrands.Visibility.MSAA.SamplePerPixel=1
 r.SceneCulling.ExplicitCellBounds=1
 r.GPUStatsEnabled=0
+; Workaround for the webkit WebBrowser issues with slate
+Slate.NoAnimationOnTabForegroundedEvent=True
 
 [/Script/Engine.RendererSettings]
 ; Substrate settings for more advanced shading on PC and editor by default.
diff --git a/Engine/Plugins/Experimental/AIAssistant/Source/AIAssistant/Private/AIAssistant.cpp b/Engine/Plugins/Experimental/AIAssistant/Source/AIAssistant/Private/AIAssistant.cpp
--- a/Engine/Plugins/Experimental/AIAssistant/Source/AIAssistant/Private/AIAssistant.cpp
+++ b/Engine/Plugins/Experimental/AIAssistant/Source/AIAssistant/Private/AIAssistant.cpp
@@ -75,7 +75,9 @@ static const FName AIAssistantTabName("AIAssistant");
 	FGlobalTabmanager::Get()->RegisterNomadTabSpawner(AIAssistantTabName, FOnSpawnTab::CreateRaw(this, &FAIAssistantModule::OnSpawnPluginTab))
 		.SetDisplayName(LOCTEXT("FAIAssistantTabTitle", "AI Assistant"))
 		.SetMenuType(ETabSpawnerMenuType::Hidden)
-		.SetIcon(FSlateIcon("AIAssistantStyle","AIAssistant.OpenPluginWindow"));
+		.SetIcon(FSlateIcon("AIAssistantStyle","AIAssistant.OpenPluginWindow"))
+		// Sidebar tabs stil cause some issue with the panel drawer
+		.SetCanSidebarTab(false);
 
 
 	InputProcessor = MakeShared<FAIAssistantInputProcessor>(PluginCommands);
@@ -302,22 +304,19 @@ void FAIAssistantModule::GenerateStatusBarPanelDrawerSummon(TArray<UStatusBarSub
 
 TSharedRef<SDockTab> FAIAssistantModule::OnSpawnPluginTab(const FSpawnTabArgs& SpawnTabArgs)
 {
-	// Work around for the web browser drawing out of bounds while doing the on demand clipping
-	AIAssistantWebBrowserWidget = SNew(SAIAssistantWebBrowser)
-		.Clipping(EWidgetClipping::ClipToBounds);
-
-
 	TSharedRef<SDockTab> DockTabWidget = SNew(SDockTab)
 		.TabRole(ETabRole::NomadTab)
 		.OnTabClosed_Lambda(
 			[this](TSharedRef<SDockTab>) -> void
 			{
 				AIAssistantWebBrowserWidget.Reset();
-			})
-		[
-			AIAssistantWebBrowserWidget.ToSharedRef()
-		];
+			});
+
+	// Work around for the web browser drawing out of bounds while doing the on demand clipping
+	AIAssistantWebBrowserWidget = SNew(SAIAssistantWebBrowser, DockTabWidget)
+		.Clipping(EWidgetClipping::ClipToBounds);
 
+	DockTabWidget->SetContent(AIAssistantWebBrowserWidget.ToSharedRef());
 
 	return DockTabWidget;
 }
diff --git a/Engine/Plugins/Experimental/AIAssistant/Source/AIAssistant/Private/AIAssistantWebBrowser.cpp b/Engine/Plugins/Experimental/AIAssistant/Source/AIAssistant/Private/AIAssistantWebBrowser.cpp
--- a/Engine/Plugins/Experimental/AIAssistant/Source/AIAssistant/Private/AIAssistantWebBrowser.cpp
+++ b/Engine/Plugins/Experimental/AIAssistant/Source/AIAssistant/Private/AIAssistantWebBrowser.cpp
@@ -9,6 +9,7 @@
 #include "Misc/AssertionMacros.h"
 #include "Framework/Application/SlateApplication.h"
 #include "Async/TaskGraphInterfaces.h"
+#include "Widgets/Docking/SDockTab.h"
 
 #include "AIAssistantConfig.h"
 #include "AIAssistantLog.h"
@@ -24,7 +25,7 @@ using namespace UE::AIAssistant;
 #define UE_AIA_SET_INITIAL_URL_VIA_WEB_BROWSER_SLATE_ARG 0
 
 
-void SAIAssistantWebBrowser::Construct(const FArguments& InArgs)
+void SAIAssistantWebBrowser::Construct(const FArguments& InArgs, const TSharedRef<SDockTab>& InParentTab)
 {
 	Config = FAIAssistantConfig::Load();
 	
@@ -59,6 +60,7 @@ void SAIAssistantWebBrowser::Construct(const FArguments& InArgs)
 		WindowSettings.BrowserFrameRate = 60;
 
 		WebBrowserWindow = WebBrowserSingleton->CreateBrowserWindow(WindowSettings);
+		WebBrowserWindow->SetParentDockTab(InParentTab);
 	}
 	else
 	{
diff --git a/Engine/Plugins/Experimental/AIAssistant/Source/AIAssistant/Private/AIAssistantWebBrowser.h b/Engine/Plugins/Experimental/AIAssistant/Source/AIAssistant/Private/AIAssistantWebBrowser.h
--- a/Engine/Plugins/Experimental/AIAssistant/Source/AIAssistant/Private/AIAssistantWebBrowser.h
+++ b/Engine/Plugins/Experimental/AIAssistant/Source/AIAssistant/Private/AIAssistantWebBrowser.h
@@ -49,7 +49,7 @@ public:
 	SLATE_BEGIN_ARGS(SAIAssistantWebBrowser) {}
 	SLATE_END_ARGS()
 
-	void Construct(const FArguments& InArgs);
+	void Construct(const FArguments& InArgs, const TSharedRef<class SDockTab>& InParentTab);
 
 	virtual ~SAIAssistantWebBrowser();
 	
diff --git a/Engine/Source/Runtime/Slate/Private/Framework/Docking/DockingUtilsPrivate.h b/Engine/Source/Runtime/Slate/Private/Framework/Docking/DockingUtilsPrivate.h
new file mode 100644
--- /dev/null
+++ b/Engine/Source/Runtime/Slate/Private/Framework/Docking/DockingUtilsPrivate.h
@@ -0,0 +1,15 @@
+// Copyright Epic Games, Inc. All Rights Reserved.
+
+#pragma once
+
+#include "HAL/IConsoleManager.h"
+
+namespace UE::Slate::Private
+{
+	// Workaround for the web browser issues on some platform
+	TAutoConsoleVariable<bool> CVarNoAnimationOnTabForegroundedEvent(
+		TEXT("Slate.NoAnimationOnTabForegroundedEvent"),
+		false,
+		TEXT("Tell the docking system to broadcast the event in a way that doesn't account for the animations")
+	);
+}
\ No newline at end of file
diff --git a/Engine/Source/Runtime/Slate/Private/Framework/Docking/SPanelDrawerArea.cpp b/Engine/Source/Runtime/Slate/Private/Framework/Docking/SPanelDrawerArea.cpp
--- a/Engine/Source/Runtime/Slate/Private/Framework/Docking/SPanelDrawerArea.cpp
+++ b/Engine/Source/Runtime/Slate/Private/Framework/Docking/SPanelDrawerArea.cpp
@@ -3,6 +3,7 @@
 #include "Framework/Docking/SPanelDrawerArea.h"
 
 #include "Containers/Array.h"
+#include "Framework/Docking/DockingUtilsPrivate.h"
 #include "Framework/Docking/STabPanelDrawer.h"
 #include "Layout/Geometry.h"
 #include "Layout/WidgetPath.h"
@@ -77,6 +78,22 @@ void SPanelDrawerArea::OpenPanel(bool bWithAnimation, const TSharedRef<UE::Slate
 		bIsOpen = true;
 		RequestSaveLayout();
 		OnExternalStateChange.ExecuteIfBound();
+
+		if (!CVarNoAnimationOnTabForegroundedEvent.GetValueOnGameThread())
+		{
+			if (PanelDrawerData && PanelDrawerData->HostedTab)
+			{
+				TSharedRef<FGlobalTabmanager> GlobalTabManager = FGlobalTabmanager::Get();
+				TSharedPtr<FTabManager> LocalTabManager = PanelDrawerData->HostedTab->GetTabManagerPtr();
+
+				if (GlobalTabManager != LocalTabManager)
+				{
+					LocalTabManager->GetPrivateApi().OnTabForegrounded(PanelDrawerData->HostedTab, {});
+				}
+
+				GlobalTabManager->GetPrivateApi().OnTabForegrounded(PanelDrawerData->HostedTab, {});
+			}
+		}
 	}
 }
 
@@ -87,6 +104,12 @@ void SPanelDrawerArea::ClosePanel(bool bWithAnimation, bool bIsTabBeingTransfere
 		return;
 	}
 
+	TSharedPtr<SDockTab> OldTabDisplayed;
+	if (PanelDrawerData)
+	{
+		OldTabDisplayed = PanelDrawerData->HostedTab;
+	}
+	
 	if (bIsTabBeingTransfered && PanelDrawerData)
 	{
 		TSharedRef<UE::Slate::Private::FPanelDrawerData> NewPanelDrawerData = MakeShared<UE::Slate::Private::FPanelDrawerData>(*PanelDrawerData.Get());
@@ -109,6 +132,22 @@ void SPanelDrawerArea::ClosePanel(bool bWithAnimation, bool bIsTabBeingTransfere
 		RequestSaveLayout();
 		bIsOpen = false;
 		OnExternalStateChange.ExecuteIfBound();
+
+		if (CVarNoAnimationOnTabForegroundedEvent.GetValueOnGameThread() || bIsTabBeingTransfered)
+		{
+			if (OldTabDisplayed)
+			{
+				TSharedRef<FGlobalTabmanager> GlobalTabManager = FGlobalTabmanager::Get();
+				TSharedPtr<FTabManager> LocalTabManager = OldTabDisplayed->GetTabManagerPtr();
+
+				if (GlobalTabManager != LocalTabManager)
+				{
+					LocalTabManager->GetPrivateApi().OnTabForegrounded({}, OldTabDisplayed);
+				}
+
+				GlobalTabManager->GetPrivateApi().OnTabForegrounded({}, OldTabDisplayed);
+			}
+		}
 	}
 }
 
@@ -177,6 +216,21 @@ void SPanelDrawerArea::SetupClosedLayout()
 		AreaContent.ToSharedRef()
 	];
 
+	if (!CVarNoAnimationOnTabForegroundedEvent.GetValueOnGameThread())
+	{
+		if (PanelDrawerData && PanelDrawerData->HostedTab)
+		{
+			TSharedRef<FGlobalTabmanager> GlobalTabManager = FGlobalTabmanager::Get();
+			TSharedPtr<FTabManager> LocalTabManager = PanelDrawerData->HostedTab->GetTabManagerPtr();
+
+			if (GlobalTabManager != LocalTabManager)
+			{
+				LocalTabManager->GetPrivateApi().OnTabForegrounded({}, PanelDrawerData->HostedTab);
+			}
+
+			GlobalTabManager->GetPrivateApi().OnTabForegrounded({}, PanelDrawerData->HostedTab);
+		}
+	}
 
 	SetPanelDrawerData({});
 }
@@ -215,6 +269,22 @@ void SPanelDrawerArea::SetupOpenedLayout()
 			]
 		]
 	];
+
+	if (CVarNoAnimationOnTabForegroundedEvent.GetValueOnGameThread())
+	{
+		if (PanelDrawerData && PanelDrawerData->HostedTab)
+		{
+			TSharedRef<FGlobalTabmanager> GlobalTabManager = FGlobalTabmanager::Get();
+			TSharedPtr<FTabManager> LocalTabManager = PanelDrawerData->HostedTab->GetTabManagerPtr();
+
+			if (GlobalTabManager != LocalTabManager)
+			{
+				LocalTabManager->GetPrivateApi().OnTabForegrounded(PanelDrawerData->HostedTab, {});
+			}
+
+			GlobalTabManager->GetPrivateApi().OnTabForegrounded(PanelDrawerData->HostedTab, {});
+		}
+	}
 }
 
 void SPanelDrawerArea::OnMainContentSlotResized(float InMainContentCoefficient)
diff --git a/Engine/Source/Runtime/WebBrowser/Private/Apple/ApplePlatformWebBrowser.cpp b/Engine/Source/Runtime/WebBrowser/Private/Apple/ApplePlatformWebBrowser.cpp
--- a/Engine/Source/Runtime/WebBrowser/Private/Apple/ApplePlatformWebBrowser.cpp
+++ b/Engine/Source/Runtime/WebBrowser/Private/Apple/ApplePlatformWebBrowser.cpp
@@ -182,10 +182,11 @@ class SAppleWebBrowserWidget : public SViewport
 		}
 #endif
 
-		if(!Is3DBrowser && (NewFrame.size.width != WindowFrame.size.width || NewFrame.size.height != WindowFrame.size.height)) 
+		if(!Is3DBrowser && (bForceRecomputeWindowOnPaint || NewFrame.size.width != WindowFrame.size.width || NewFrame.size.height != WindowFrame.size.height)) 
 		{
 			// WindowFrame is only used when we are a 2D web browser
 			WindowFrame = NewFrame;
+			bForceRecomputeWindowOnPaint = false;
 		}
 		
 #if PLATFORM_IOS
@@ -669,7 +670,8 @@ class SAppleWebBrowserWidget : public SViewport
 						{
 							if (CurrentTab->IsForeground())
 							{
-								[WebViewWrapper setVisibility : YES] ;
+								[WebViewWrapper setVisibility : YES];
+								bForceRecomputeWindowOnPaint = true;
 							}
 						}
 					}
@@ -694,6 +696,7 @@ class SAppleWebBrowserWidget : public SViewport
 			if(ActiveTab == CurrentTab)
 			{
 				[WebViewWrapper setVisibility: YES ];
+				bForceRecomputeWindowOnPaint = true;
 			}
 			else if(PreviouslyActiveTab == CurrentTab)
 			{
@@ -717,6 +720,7 @@ private:
 	bool Is3DBrowser;
 	bool bAcquiredInitialFocus;
 	bool bFocused;
+	bool bForceRecomputeWindowOnPaint = false;
 
 #if !PLATFORM_TVOS
 	/** The external texture to render the webbrowser output. */
