Fixed the ai assistant tab not spawning when closing a docking area that was hosting it in a pannel drawer.

diff --git a/Engine/Source/Runtime/Slate/Private/Framework/Docking/TabManager.cpp b/Engine/Source/Runtime/Slate/Private/Framework/Docking/TabManager.cpp
--- a/Engine/Source/Runtime/Slate/Private/Framework/Docking/TabManager.cpp
+++ b/Engine/Source/Runtime/Slate/Private/Framework/Docking/TabManager.cpp
@@ -1651,6 +1651,10 @@ TSharedPtr<SDockTab> FTabManager::InvokeTab_Internal(const FTabId& TabId, bool b
 					bCanReuseTab = true;
 				}
 			}
+			else if (TSharedPtr<FTabManager> TabManager = Tab->GetTabManagerPtr())
+			{
+				bCanReuseTab = TabManager->RemoveFromHiddenPanelDrawerTabs(Tab);
+			}
 		}
 
 		if (!bCanReuseTab)
