GitHub 13769 : Only restrict BP breakpoints that are in the current UI window.
#13769

diff --git a/Engine/Source/Runtime/Slate/Private/Framework/Application/SlateApplication.cpp b/Engine/Source/Runtime/Slate/Private/Framework/Application/SlateApplication.cpp
--- a/Engine/Source/Runtime/Slate/Private/Framework/Application/SlateApplication.cpp
+++ b/Engine/Source/Runtime/Slate/Private/Framework/Application/SlateApplication.cpp
@@ -42,6 +42,7 @@
 #include "Framework/Application/IInputProcessor.h"
 #include "GenericPlatform/ITextInputMethodSystem.h"
 #include "Framework/Docking/TabCommands.h"
+#include "Framework/Docking/TabManager.h"
 #include "Math/UnitConversion.h"
 #include "ProfilingDebugging/CsvProfiler.h"
 #include "ProfilingDebugging/StallDetector.h"
@@ -1602,10 +1603,6 @@ void FSlateApplication::Tick(ESlateTickType TickType)
 
 	TGuardValue<bool> IsTickingGuard(bIsTicking, true);
 
-#if WITH_EDITOR
-	FScopedPreventDebuggingMode SlatePreventDebugginModeWhileTicking(NSLOCTEXT("EnterDebuggingMode", "WindowTicking", "The window is ticking."));
-#endif
-
 	SCOPED_NAMED_EVENT_F(TEXT("Slate::Tick (%s)"), FColor::Magenta, LexToString(TickType));
 
 	{
@@ -3779,6 +3776,7 @@ void FSlateApplication::EnterDebuggingMode()
 		{
 			FNotificationInfo Info(NSLOCTEXT("EnterDebuggingMode", "FailTitle", "Debugging Mode Fail"));
 			Info.SubText = SubText;
+			Info.ExpireDuration = 15.f;
 			Self->DebuggingModeNotificationMessage = FSlateNotificationManager::Get().AddNotification(Info);
 		}
 		UE_LOG(LogSlate, Warning, TEXT("Enter Debugging Mode failed."));
@@ -3815,6 +3813,22 @@ void FSlateApplication::EnterDebuggingMode()
 	}
 
 #if WITH_EDITOR
+	// If we enter debugging mode to debug a widget in the main editor window,
+	// the editor will become mostly unresponsive as it stops ticking the main editor window.
+	// Otherwise allow debugging within tick.
+	if (bIsTicking)
+	{
+		TSharedPtr<SWindow> RootWindow = FGlobalTabmanager::Get()->GetRootWindow();
+		TSharedPtr<SWidget> ContextWidgetPinned = CurrentDebugContextWidget.Pin();
+		TSharedPtr<SWindow> DebugWindow = ContextWidgetPinned.IsValid() ? FindWidgetWindow(ContextWidgetPinned.ToSharedRef()) : TSharedPtr<SWindow>();
+		if (DebugWindow.IsValid() && DebugWindow == RootWindow)
+		{
+			AddNotification(INVTEXT("Tried to enter debugging mode inside Tick in the main editor window.\n\n"
+				"Restart PIE using New Editor Window instead of Selected Viewport (or undock the viewport tab) and try again."));
+			return;
+		}
+	}
+
 	if (PreventDebuggingModeStack.Num() > 0)
 	{
 		AddNotification(PreventDebuggingModeStack.Last().Key);
diff --git a/Engine/Source/Runtime/SlateRHIRenderer/Private/SlateRHIRenderer.cpp b/Engine/Source/Runtime/SlateRHIRenderer/Private/SlateRHIRenderer.cpp
--- a/Engine/Source/Runtime/SlateRHIRenderer/Private/SlateRHIRenderer.cpp
+++ b/Engine/Source/Runtime/SlateRHIRenderer/Private/SlateRHIRenderer.cpp
@@ -300,7 +300,6 @@ FSlateDrawBuffer& FSlateRHIRenderer::AcquireDrawBuffer()
 		else
 		{
 			FlushCommands();
-			UE_LOG(LogSlate, Warning, TEXT("Slate: Had to block on waiting for a draw buffer"));
 			FreeBufferIndex = (FreeBufferIndex + 1) % NumDrawBuffers;
 		}
 
