Slate Dynamic Invalidation - Layout Widgets

diff --git a/Engine/Source/Runtime/Slate/Private/Widgets/Layout/SScrollBar.cpp b/Engine/Source/Runtime/Slate/Private/Widgets/Layout/SScrollBar.cpp
--- a/Engine/Source/Runtime/Slate/Private/Widgets/Layout/SScrollBar.cpp
+++ b/Engine/Source/Runtime/Slate/Private/Widgets/Layout/SScrollBar.cpp
@@ -8,24 +8,34 @@
 #include "Widgets/Layout/SBox.h"
 #include "Framework/Application/SlateApplication.h"
 
+SLATE_IMPLEMENT_WIDGET(SScrollBar)
+void SScrollBar::PrivateRegisterAttributes(FSlateAttributeInitializer& AttributeInitializer)
+{
+	SLATE_ADD_MEMBER_ATTRIBUTE_DEFINITION(AttributeInitializer, UserVisibilityAttribute, EInvalidateWidgetReason::Visibility)
+		.AffectVisibility();
+}
+
 void SScrollBar::Construct(const FArguments& InArgs)
 {
+	SetCanTick(false);
+	
 	OnUserScrolled = InArgs._OnUserScrolled;
 	OnScrollBarVisibilityChanged = InArgs._OnScrollBarVisibilityChanged;
 	Orientation = InArgs._Orientation;
 	DragFocusCause = InArgs._DragFocusCause;
-	UserVisibility = InArgs._Visibility;
+	SetUserVisibility(InArgs._Visibility);
 
 	check(InArgs._Style);
-	SetStyle(InArgs._Style);
 
-	const TAttribute<FVector2D> Thickness = InArgs._Thickness.IsSet() ? InArgs._Thickness : FVector2D(InArgs._Style->Thickness, InArgs._Style->Thickness);
+	bIsThicknessAttrSet = InArgs._Thickness.IsSet();
+	const TAttribute<FVector2D> Thickness = InArgs._Thickness.IsSet() ? InArgs._Thickness : FVector2D(InArgs._Style->Thickness);
 
 	EHorizontalAlignment HorizontalAlignment = Orientation == Orient_Vertical ? HAlign_Center : HAlign_Fill;
 	EVerticalAlignment VerticalAlignment = Orientation == Orient_Vertical ? VAlign_Fill : VAlign_Center;
 
 	bHideWhenNotInUse = InArgs._HideWhenNotInUse;
 	bPreventThrottling = InArgs._PreventThrottling;
+	ScrollbarDisabledVisibility = InArgs._ScrollbarDisabledVisibility;
 	bIsScrolling = false;
 	LastInteractionTime = 0;
 
@@ -38,8 +48,7 @@ void SScrollBar::Construct(const FArguments& InArgs)
 			+ SVerticalBox::Slot()
 			.FillHeight(1)
 			[
-				SNew(SBorder)
-				.BorderImage(BackgroundBrush)
+				SAssignNew(BackgroundBorder, SBorder)
 				.HAlign(HorizontalAlignment)
 				.VAlign(VerticalAlignment)
 				.Padding(0.f)
@@ -53,7 +62,7 @@ void SScrollBar::Construct(const FArguments& InArgs)
 						.VAlign(VerticalAlignment)
 						[
 							SAssignNew(TopImage, SImage)
-							.Image(TopBrush)
+							.ColorAndOpacity(TAttribute<FSlateColor>(SharedThis(this), &SScrollBar::GetTrackOpacity))
 						]
 					]
 					.ThumbSlot()
@@ -61,6 +70,15 @@ void SScrollBar::Construct(const FArguments& InArgs)
 						SAssignNew(DragThumb, SBorder)
 						.HAlign(HAlign_Center)
 						.VAlign(VAlign_Center)
+						.BorderImage(TAttribute<const FSlateBrush*>(SharedThis(this), &SScrollBar::GetDragThumbImage))
+						.BorderBackgroundColor_Lambda([WeakThis = AsWeak()]()
+						{
+							if (const TSharedPtr<SWidget> This = WeakThis.Pin())
+							{
+								return FSlateColor(StaticCastSharedPtr<SScrollBar>(This)->GetThumbOpacity());
+							}
+							return FSlateColor();
+						})
 						.Padding(0.f)
 						[
 							SAssignNew(ThicknessSpacer, SSpacer)
@@ -74,7 +92,7 @@ void SScrollBar::Construct(const FArguments& InArgs)
 						.VAlign(VerticalAlignment)
 						[
 							SAssignNew(BottomImage, SImage)
-							.Image(BottomBrush)
+							.ColorAndOpacity(TAttribute<FSlateColor>(SharedThis(this), &SScrollBar::GetTrackOpacity))
 						]
 					]
 				]
@@ -82,7 +100,10 @@ void SScrollBar::Construct(const FArguments& InArgs)
 		]
 	);
 
+	SetStyle(InArgs._Style);
+
 	SetEnabled(TAttribute<bool>( Track.ToSharedRef(), &SScrollBarTrack::IsNeeded ));
+	SetVisibility(TAttribute<EVisibility>(SharedThis(this), &SScrollBar::ShouldBeVisible));
 	SetScrollBarAlwaysVisible(InArgs._AlwaysShowScrollbar);
 	bAlwaysShowScrollbarTrack = InArgs._AlwaysShowScrollbarTrack;
 }
@@ -104,7 +125,6 @@ void SScrollBar::SetState( float InOffsetFraction, float InThumbSizeFraction, bo
 		// Note that the maximum offset depends on how many items fit per screen
 		// It is 1.0f-InThumbSizeFraction.
 		Track->SetSizes(InOffsetFraction, InThumbSizeFraction);
-		GetVisibilityAttribute().UpdateValue();
 
 		LastInteractionTime = FSlateApplication::Get().GetCurrentTime();
 		if (bCallOnUserScrolled)
@@ -114,17 +134,6 @@ void SScrollBar::SetState( float InOffsetFraction, float InThumbSizeFraction, bo
 	}
 }
 
-void SScrollBar::Tick(const FGeometry& AllottedGeometry, const double InCurrentTime, const float InDeltaTime)
-{
-	SBorder::Tick(AllottedGeometry, InCurrentTime, InDeltaTime);
-
-	DragThumb->SetBorderImage(GetDragThumbImage());
-	DragThumb->SetColorAndOpacity(GetThumbOpacity());
-
-	TopImage->SetColorAndOpacity(GetTrackOpacity());
-	BottomImage->SetColorAndOpacity(GetTrackOpacity());
-}
-
 FReply SScrollBar::OnMouseButtonDown( const FGeometry& MyGeometry, const FPointerEvent& MouseEvent )
 {
 	if (MouseEvent.GetEffectingButton() == EKeys::LeftMouseButton && IsNeeded())
@@ -248,6 +257,7 @@ float SScrollBar::ThumbSizeFraction() const
 SScrollBar::SScrollBar()
 	: bDraggingThumb(false)
 	, DragGrabOffset( 0.0f )
+	, UserVisibilityAttribute(*this)
 {
 }
 
@@ -327,13 +337,13 @@ EVisibility SScrollBar::ShouldBeVisible() const
 	const EVisibility CurrentVisibility = GetVisibility();
     EVisibility NewVisibility = ScrollbarDisabledVisibility;
     
-	if ( this->HasMouseCapture() )
+	if ( bAlwaysShowScrollbar || this->HasMouseCapture() )
 	{
 		NewVisibility = EVisibility::Visible;
 	}
 	else if( Track->IsNeeded() )
 	{
-		NewVisibility = UserVisibility.Get();
+		NewVisibility = UserVisibilityAttribute.Get();
 	}
 
 	if (NewVisibility != CurrentVisibility)
@@ -356,7 +366,7 @@ EOrientation SScrollBar::GetOrientation() const
 
 void SScrollBar::SetStyle(const FScrollBarStyle* InStyle)
 {
-	const FScrollBarStyle* Style = InStyle;
+	Style = InStyle;
 
 	if (Style == nullptr)
 	{
@@ -382,6 +392,14 @@ void SScrollBar::SetStyle(const FScrollBarStyle* InStyle)
 		TopBrush = &Style->HorizontalTopSlotImage;
 		BottomBrush = &Style->HorizontalBottomSlotImage;
 	}
+	BackgroundBorder->SetBorderImage(BackgroundBrush);
+	TopImage->SetImage(TopBrush);
+	BottomImage->SetImage(BottomBrush);
+
+	if (!bIsThicknessAttrSet)
+	{
+		ThicknessSpacer->SetSize(FVector2D(Style->Thickness));
+	}
 
 	InvalidateStyle();
 }
@@ -391,39 +409,43 @@ void SScrollBar::InvalidateStyle()
 	Invalidate(EInvalidateWidgetReason::Layout);
 }
 
-void SScrollBar::SetDragFocusCause(EFocusCause InDragFocusCause)
+void SScrollBar::SetUserVisibility(TAttribute<EVisibility> InUserVisibility)
 {
-	DragFocusCause = InDragFocusCause;
+	UserVisibilityAttribute.Assign(*this, MoveTemp(InUserVisibility));
 }
 
-void SScrollBar::SetThickness(TAttribute<FVector2D> InThickness)
+void SScrollBar::SetDragFocusCause(EFocusCause InDragFocusCause)
 {
-	ThicknessSpacer->SetSize(InThickness);
+	DragFocusCause = InDragFocusCause;
 }
 
-void SScrollBar::SetScrollBarAlwaysVisible(bool InAlwaysVisible)
+void SScrollBar::SetThickness(TAttribute<FVector2D> InThickness)
 {
-	bAlwaysShowScrollbar = InAlwaysVisible;
-
-	if ( InAlwaysVisible )
+	if (InThickness.IsSet())
 	{
-		SetVisibility(EVisibility::Visible);
+		bIsThicknessAttrSet = true;
+		ThicknessSpacer->SetSize( MoveTemp(InThickness) );
 	}
 	else
 	{
-		SetVisibility(TAttribute<EVisibility>(SharedThis(this), &SScrollBar::ShouldBeVisible));
+		bIsThicknessAttrSet = false;
+		ThicknessSpacer->SetSize(FVector2D(Style->Thickness));
 	}
+}
+
+void SScrollBar::SetScrollBarAlwaysVisible(bool InAlwaysVisible)
+{
+	bAlwaysShowScrollbar = InAlwaysVisible;
 
 	Track->SetIsAlwaysVisible(InAlwaysVisible);
 }
 
 void SScrollBar::SetScrollBarTrackAlwaysVisible(bool InAlwaysVisible)
 {
-	// Doesn't need to be invalidated here, tick updates these values.
 	bAlwaysShowScrollbarTrack = InAlwaysVisible;
 }
 
-SLATE_API void SScrollBar::SetScrollbarDisabledVisibility(EVisibility InVisibility)
+void SScrollBar::SetScrollbarDisabledVisibility(EVisibility InVisibility)
 {
 	ScrollbarDisabledVisibility = InVisibility;
 }
diff --git a/Engine/Source/Runtime/Slate/Private/Widgets/Layout/SScrollBarTrack.cpp b/Engine/Source/Runtime/Slate/Private/Widgets/Layout/SScrollBarTrack.cpp
--- a/Engine/Source/Runtime/Slate/Private/Widgets/Layout/SScrollBarTrack.cpp
+++ b/Engine/Source/Runtime/Slate/Private/Widgets/Layout/SScrollBarTrack.cpp
@@ -2,6 +2,7 @@
 
 #include "Widgets/Layout/SScrollBarTrack.h"
 #include "Layout/ArrangedChildren.h"
+#include "Widgets/SOverlay.h"
 
 static const int32 NUM_SCROLLBAR_SLOTS = 3;
 
@@ -14,12 +15,32 @@ void SScrollBarTrack::Construct(const FArguments& InArgs)
 
 	
 	static_assert(NUM_SCROLLBAR_SLOTS == 3, "SScrollBarTrack::FSlot numbers changed");
+
 	static_assert(TOP_SLOT_INDEX == 0, "SScrollBarTrack::FSlot order has changed");
 	Children.AddSlot(MoveTemp(FSlot::FSlotArguments(MakeUnique<FSlot>())[InArgs._TopSlot.Widget]));
+
 	static_assert(BOTTOM_SLOT_INDEX == 1, "SScrollBarTrack::FSlot order has changed");
-	Children.AddSlot(MoveTemp(FSlot::FSlotArguments(MakeUnique<FSlot>())[InArgs._BottomSlot.Widget]));
+	Children.AddSlot(MoveTemp(FSlot::FSlotArguments(MakeUnique<FSlot>())
+	[
+		SNew(SOverlay) // We need a parent widget to control the visibility
+		+ SOverlay::Slot()
+		[
+			InArgs._BottomSlot.Widget
+		]
+	]));
+
 	static_assert(THUMB_SLOT_INDEX == 2, "SScrollBarTrack::FSlot order has changed");
-	Children.AddSlot(MoveTemp(FSlot::FSlotArguments(MakeUnique<FSlot>())[InArgs._ThumbSlot.Widget]));
+	Children.AddSlot(MoveTemp(FSlot::FSlotArguments(MakeUnique<FSlot>())
+	[
+		SNew(SOverlay) // We need a parent widget to control the visibility
+		+ SOverlay::Slot()
+		[
+			InArgs._ThumbSlot.Widget
+		]
+	]));
+
+	// Ensure the visibilities are set as expected
+	InvalidateIsNeeded(!IsNeeded());
 }
 
 SScrollBarTrack::FTrackSizeInfo SScrollBarTrack::GetTrackSizeInfo(const FGeometry& InTrackGeometry) const
@@ -104,20 +125,24 @@ FChildren* SScrollBarTrack::GetChildren()
 
 void SScrollBarTrack::SetSizes(float InThumbOffsetFraction, float InThumbSizeFraction)
 {
-	OffsetFraction = InThumbOffsetFraction;
-	ThumbSizeFraction = InThumbSizeFraction;
-
 	// If you have no thumb, then it's effectively the size of the whole track
-	if (ThumbSizeFraction == 0.0f && !bIsAlwaysVisible)
+	if (InThumbSizeFraction == 0.0f && !bIsAlwaysVisible)
 	{
-		ThumbSizeFraction = 1.0f;
+		InThumbSizeFraction = 1.0f;
 	}
-	else if (ThumbSizeFraction > 1.0f && bIsAlwaysVisible)
+	else if (InThumbSizeFraction > 1.0f && bIsAlwaysVisible)
 	{
-		ThumbSizeFraction = 0.0f;
+		InThumbSizeFraction = 0.0f;
 	}
 
-	Invalidate(EInvalidateWidget::Layout);
+	if (OffsetFraction != InThumbOffsetFraction || ThumbSizeFraction != InThumbSizeFraction)
+	{
+		const bool bWasNeeded = IsNeeded();
+		OffsetFraction = InThumbOffsetFraction;
+		ThumbSizeFraction = InThumbSizeFraction;
+
+		InvalidateIsNeeded(bWasNeeded);
+	}
 }
 
 bool SScrollBarTrack::IsNeeded() const
@@ -151,7 +176,23 @@ void SScrollBarTrack::SetIsAlwaysVisible(bool InIsAlwaysVisible)
 {
 	if (bIsAlwaysVisible != InIsAlwaysVisible)
 	{
+		const bool bWasNeeded = IsNeeded();
 		bIsAlwaysVisible = InIsAlwaysVisible;
-		Invalidate(EInvalidateWidget::Layout);
+
+		// SetSizes is called because it can adjust the ThumbSizeFraction value
+		SetSizes(OffsetFraction, ThumbSizeFraction);
+		
+		InvalidateIsNeeded(bWasNeeded);
+	}
+}
+
+void SScrollBarTrack::InvalidateIsNeeded(bool bWasNeeded)
+{
+	Invalidate(EInvalidateWidget::Layout);
+	if (IsNeeded() != bWasNeeded)
+	{
+		const EVisibility NewVisibility = IsNeeded() ? EVisibility::SelfHitTestInvisible : EVisibility::Hidden;
+		Children[BOTTOM_SLOT_INDEX].GetWidget()->SetVisibility(NewVisibility);
+		Children[THUMB_SLOT_INDEX].GetWidget()->SetVisibility(NewVisibility);
 	}
 }
diff --git a/Engine/Source/Runtime/Slate/Private/Widgets/Layout/SScrollBorder.cpp b/Engine/Source/Runtime/Slate/Private/Widgets/Layout/SScrollBorder.cpp
--- a/Engine/Source/Runtime/Slate/Private/Widgets/Layout/SScrollBorder.cpp
+++ b/Engine/Source/Runtime/Slate/Private/Widgets/Layout/SScrollBorder.cpp
@@ -6,6 +6,18 @@
 #include "Framework/Layout/IScrollableWidget.h"
 
 
+SLATE_IMPLEMENT_WIDGET(SScrollBorder)
+void SScrollBorder::PrivateRegisterAttributes(FSlateAttributeInitializer& AttributeInitializer)
+{
+	SLATE_ADD_MEMBER_ATTRIBUTE_DEFINITION(AttributeInitializer, BorderFadeDistanceAttribute, EInvalidateWidgetReason::Paint);
+}
+
+
+SScrollBorder::SScrollBorder()
+	: BorderFadeDistanceAttribute(*this, FVector2D(0.01f, 0.01f))
+{
+}
+
 SScrollBorder::~SScrollBorder()
 { }
 
@@ -14,7 +26,7 @@ void SScrollBorder::Construct(const FArguments& InArgs, TSharedRef<IScrollableWi
 {
 	check( InArgs._Style );
 
-	BorderFadeDistance = InArgs._BorderFadeDistance;
+	SetBorderFadeDistance(InArgs._BorderFadeDistance);
 	ScrollableWidget = InScrollableWidget;
 
 	TSharedRef<SOverlay> Overlay = SNew(SOverlay)
@@ -52,6 +64,11 @@ void SScrollBorder::Construct(const FArguments& InArgs, TSharedRef<IScrollableWi
 	];
 }
 
+void SScrollBorder::SetBorderFadeDistance(TAttribute<FVector2D> InBorderFadeDistance)
+{
+	BorderFadeDistanceAttribute.Assign(*this, MoveTemp(InBorderFadeDistance));
+}
+
 FSlateColor SScrollBorder::GetTopBorderOpacity() const
 {
 	float ShadowOpacity = 0;
@@ -60,7 +77,8 @@ FSlateColor SScrollBorder::GetTopBorderOpacity() const
 	if ( Scrollable.IsValid() )
 	{
 		// The shadow should only be visible when the user needs a hint that they can scroll up.
-		ShadowOpacity = FMath::Clamp( Scrollable->GetScrollDistance().Y / BorderFadeDistance.Get().Y, 0.0f, 1.0f );
+		const float FadeDistance = FMath::Max(0.01f, BorderFadeDistanceAttribute.Get().Y);
+		ShadowOpacity = FMath::Clamp( Scrollable->GetScrollDistance().Y / FadeDistance, 0.0f, 1.0f );
 	}
 
 	return FLinearColor( 1.0f, 1.0f, 1.0f, ShadowOpacity );
@@ -74,7 +92,8 @@ FSlateColor SScrollBorder::GetBottomBorderOpacity() const
 	if (Scrollable.IsValid())
 	{
 		// The shadow should only be visible when the user needs a hint that they can scroll down.
-		ShadowOpacity = ( Scrollable->GetScrollDistanceRemaining().Y ) / BorderFadeDistance.Get().Y;
+		const float FadeDistance = FMath::Max(0.01f, BorderFadeDistanceAttribute.Get().Y);
+		ShadowOpacity = FMath::Clamp( Scrollable->GetScrollDistanceRemaining().Y / FadeDistance, 0.0f, 1.0f );
 	}
 
 	return FLinearColor( 1.0f, 1.0f, 1.0f, ShadowOpacity );
diff --git a/Engine/Source/Runtime/Slate/Private/Widgets/Layout/SScrollBox.cpp b/Engine/Source/Runtime/Slate/Private/Widgets/Layout/SScrollBox.cpp
--- a/Engine/Source/Runtime/Slate/Private/Widgets/Layout/SScrollBox.cpp
+++ b/Engine/Source/Runtime/Slate/Private/Widgets/Layout/SScrollBox.cpp
@@ -28,6 +28,13 @@ SScrollBox::FSlot::FSlotArguments SScrollBox::Slot()
 	return FSlot::FSlotArguments(MakeUnique<FSlot>());
 }
 
+SLATE_IMPLEMENT_WIDGET(SScrollPanel)
+void SScrollPanel::PrivateRegisterAttributes(FSlateAttributeInitializer& AttributeInitializer)
+{
+	FSlateWidgetSlotAttributeInitializer Initializer = SLATE_ADD_PANELCHILDREN_DEFINITION(AttributeInitializer, Children);
+	SScrollBox::FSlot::RegisterAttributes(Initializer);
+}
+
 void SScrollPanel::Construct(const FArguments& InArgs, TArray<SScrollBox::FSlot::FSlotArguments> InSlots)
 {
 	PhysicalOffset = 0;
@@ -112,7 +119,7 @@ void SScrollBox::Construct( const FArguments& InArgs )
 	Orientation = InArgs._Orientation;
 	bScrollToEnd = false;
 	bIsScrollingActiveTimerRegistered = false;
-	bAllowsRightClickDragScrolling = false;
+	bAllowsRightClickDragScrolling = InArgs._ScrollBarRightClickDragAllowed;
 	ConsumeMouseWheel = InArgs._ConsumeMouseWheel;
 	TickScrollDelta = 0;
 	AllowOverscroll = InArgs._AllowOverscroll;
@@ -210,7 +217,7 @@ void SScrollBox::ConstructVerticalLayout()
 			.VAlign(VAlign_Top)
 			[
 				// Shadow: Hint to scroll up
-				SNew(SImage)
+				SAssignNew(TopLeftShadowHint, SImage)
 				.Visibility(EVisibility::HitTestInvisible)
 				.ColorAndOpacity(this, &SScrollBox::GetStartShadowOpacity)
 				.Image(&Style->TopShadowBrush)
@@ -221,7 +228,7 @@ void SScrollBox::ConstructVerticalLayout()
 			.VAlign(VAlign_Bottom)
 			[
 				// Shadow: a hint to scroll down
-				SNew(SImage)
+				SAssignNew(BottomRightShadowHint, SImage)
 				.Visibility(EVisibility::HitTestInvisible)
 				.ColorAndOpacity(this, &SScrollBox::GetEndShadowOpacity)
 				.Image(&Style->BottomShadowBrush)
@@ -266,7 +273,7 @@ void SScrollBox::ConstructHorizontalLayout()
 			.VAlign(VAlign_Fill)
 			[
 				// Shadow: Hint to left
-				SNew(SImage)
+				SAssignNew(TopLeftShadowHint, SImage)
 				.Visibility(EVisibility::HitTestInvisible)
 				.ColorAndOpacity(this, &SScrollBox::GetStartShadowOpacity)
 				.Image(&Style->LeftShadowBrush)
@@ -277,7 +284,7 @@ void SScrollBox::ConstructHorizontalLayout()
 			.VAlign(VAlign_Fill)
 			[
 				// Shadow: a hint to scroll right
-				SNew(SImage)
+				SAssignNew(BottomRightShadowHint, SImage)
 				.Visibility(EVisibility::HitTestInvisible)
 				.ColorAndOpacity(this, &SScrollBox::GetEndShadowOpacity)
 				.Image(&Style->RightShadowBrush)
@@ -550,9 +557,12 @@ bool SScrollBox::InternalScrollDescendantIntoView(const FGeometry& MyGeometry, c
 
 void SScrollBox::SetStyle(const FScrollBoxStyle* InStyle)
 {
+	check(InStyle);
 	if (Style != InStyle)
 	{
 		Style = InStyle;
+		TopLeftShadowHint->SetImage(Orientation == Orient_Vertical ? &Style->TopShadowBrush : &Style->LeftShadowBrush);
+		BottomRightShadowHint->SetImage(Orientation == Orient_Vertical ? &Style->BottomShadowBrush : &Style->RightShadowBrush);
 		InvalidateStyle();
 	}
 }
@@ -772,6 +782,10 @@ void SScrollBox::Tick( const FGeometry& AllottedGeometry, const double InCurrent
 	{
 		Invalidate(EInvalidateWidget::Layout);
 	}
+	if (bWasScrolling != bIsScrolling)
+	{
+		Invalidate(EInvalidateWidget::Volatility);
+	}
 	
 	ScrollBar->SetState(CurrentViewOffset, ViewFraction);
 	if (!ScrollBar->IsNeeded())
diff --git a/Engine/Source/Runtime/Slate/Private/Widgets/Layout/SSplitter.cpp b/Engine/Source/Runtime/Slate/Private/Widgets/Layout/SSplitter.cpp
--- a/Engine/Source/Runtime/Slate/Private/Widgets/Layout/SSplitter.cpp
+++ b/Engine/Source/Runtime/Slate/Private/Widgets/Layout/SSplitter.cpp
@@ -5,28 +5,63 @@
 #include "Rendering/DrawElements.h"
 #include "Framework/Application/SlateApplication.h"
 
+
+SLATE_IMPLEMENT_WIDGET(SSplitter)
+void SSplitter::PrivateRegisterAttributes(FSlateAttributeInitializer& AttributeInitializer)
+{
+	FSlateWidgetSlotAttributeInitializer Initializer = SLATE_ADD_PANELCHILDREN_DEFINITION(AttributeInitializer, Children);
+	FSlot::RegisterAttributes(Initializer);
+}
+
+SSplitter::FSlot::FSlot()
+ : TWidgetSlotWithAttributeSupport<FSlot>()
+ , SizingRule(*this, FractionOfParent)
+ , SizeValue(*this, 1.f)
+ , MinSizeValue(*this, 0.f)
+{
+}
+
 void SSplitter::FSlot::Construct(const FChildren& SlotOwner, FSlotArguments&& InArgs)
 {
-	TSlotBase<FSlot>::Construct(SlotOwner, MoveTemp(InArgs));
-	if (InArgs._Value.IsSet())
-	{
-		SizeValue = MoveTemp(InArgs._Value);
-	}
+	TWidgetSlotWithAttributeSupport<FSlot>::Construct(SlotOwner, MoveTemp(InArgs));
+
+	SetSizingRule( MoveTemp(InArgs._SizeRule) );
+	SetSizeValue( MoveTemp(InArgs._Value) );
+	SetMinSize( MoveTemp(InArgs._MinSize) );
+
 	if (InArgs._Resizable.IsSet())
 	{
 		bIsResizable = InArgs._Resizable.GetValue();
 	}
-	
-	MinSizeValue = InArgs._MinSize;
 
 	if (InArgs._OnSlotResized.IsBound())
 	{
 		OnSlotResized_Handler = MoveTemp(InArgs._OnSlotResized);
 	}
-	if (InArgs._SizeRule.IsSet())
-	{
-		SizingRule = MoveTemp(InArgs._SizeRule);
-	}
+}
+
+void SSplitter::FSlot::RegisterAttributes(FSlateWidgetSlotAttributeInitializer& AttributeInitializer)
+{
+	TWidgetSlotWithAttributeSupport<FSlot>::RegisterAttributes(AttributeInitializer);
+
+	SLATE_ADD_SLOT_ATTRIBUTE_DEFINITION(FSlot, AttributeInitializer, SizingRule, EInvalidateWidgetReason::Layout);
+	SLATE_ADD_SLOT_ATTRIBUTE_DEFINITION(FSlot, AttributeInitializer, SizeValue, EInvalidateWidgetReason::Layout);
+	SLATE_ADD_SLOT_ATTRIBUTE_DEFINITION(FSlot, AttributeInitializer, MinSizeValue, EInvalidateWidgetReason::Layout);
+}
+
+void SSplitter::FSlot::SetSizeValue( TAttribute<float> InValue )
+{
+	SizeValue.Assign(*this, MoveTemp(InValue));
+}
+
+void SSplitter::FSlot::SetSizingRule( TAttribute<ESizeRule> InSizeRule ) 
+{
+	SizingRule.Assign(*this, MoveTemp(InSizeRule));
+}
+
+void SSplitter::FSlot::SetMinSize(TAttribute<float> InMinSize)
+{
+	MinSizeValue.Assign(*this, MoveTemp(InMinSize));
 }
 
 bool SSplitter::FSlot::CanBeResized() const
@@ -574,7 +609,11 @@ FCursorReply SSplitter::OnCursorQuery( const FGeometry& MyGeometry, const FPoint
 */
 void SSplitter::SetOrientation( EOrientation NewOrientation )
 {
-	Orientation = NewOrientation;
+	if (Orientation != NewOrientation)
+	{
+		Orientation = NewOrientation;
+		Invalidate(EInvalidateWidgetReason::Layout);
+	}
 }
 
 /**
diff --git a/Engine/Source/Runtime/Slate/Public/Widgets/Layout/SScrollBar.h b/Engine/Source/Runtime/Slate/Public/Widgets/Layout/SScrollBar.h
--- a/Engine/Source/Runtime/Slate/Public/Widgets/Layout/SScrollBar.h
+++ b/Engine/Source/Runtime/Slate/Public/Widgets/Layout/SScrollBar.h
@@ -25,8 +25,15 @@ DECLARE_DELEGATE_OneParam(FOnScrollBarVisibilityChanged, EVisibility);	/** chang
 
 class SScrollBarTrack;
 
+template <>
+struct TWidgetTypeTraits<class SScrollBar>
+{
+	static constexpr bool SupportsInvalidation() { return true; }
+};
+
 class SScrollBar : public SBorder
 {
+	SLATE_DECLARE_WIDGET_API(SScrollBar, SBorder, SLATE_API)
 public:
 
 	inline static const float DefaultUniformPadding = 2.0f;
@@ -36,6 +43,7 @@ public:
 		, _OnUserScrolled()
 		, _AlwaysShowScrollbar(false)
 		, _AlwaysShowScrollbarTrack(true)
+		, _ScrollbarDisabledVisibility(EVisibility::Collapsed)
 #if PLATFORM_UI_HAS_MOBILE_SCROLLBARS
 		, _HideWhenNotInUse(true)
 #else
@@ -54,6 +62,7 @@ public:
 		SLATE_EVENT( FOnScrollBarVisibilityChanged, OnScrollBarVisibilityChanged )
 		SLATE_ARGUMENT( bool, AlwaysShowScrollbar )
 		SLATE_ARGUMENT( bool, AlwaysShowScrollbarTrack )
+		SLATE_ARGUMENT( EVisibility, ScrollbarDisabledVisibility )
 		SLATE_ARGUMENT( bool, HideWhenNotInUse )
 		SLATE_ARGUMENT( bool, PreventThrottling )
 		SLATE_ARGUMENT( EOrientation, Orientation )
@@ -97,7 +106,6 @@ public:
 	SLATE_API virtual void SetState( float InOffsetFraction, float InThumbSizeFraction, bool bCallOnUserScrolled = false);
 
 	// SWidget
-	SLATE_API virtual void Tick(const FGeometry& AllottedGeometry, const double InCurrentTime, const float InDeltaTime) override;
 	SLATE_API virtual FReply OnMouseButtonDown( const FGeometry& MyGeometry, const FPointerEvent& MouseEvent ) override;
 	SLATE_API virtual FReply OnMouseButtonUp( const FGeometry& MyGeometry, const FPointerEvent& MouseEvent ) override;
 	SLATE_API virtual FReply OnMouseMove( const FGeometry& MyGeometry, const FPointerEvent& MouseEvent ) override;
@@ -133,7 +141,7 @@ public:
 	SLATE_API void InvalidateStyle();
 
 	/** Set UserVisibility attribute */
-	void SetUserVisibility(TAttribute<EVisibility> InUserVisibility) { UserVisibility = InUserVisibility; }
+	SLATE_API void SetUserVisibility(TAttribute<EVisibility> InUserVisibility);
 
 	/** Set DragFocusCause attribute */
 	SLATE_API void SetDragFocusCause(EFocusCause InDragFocusCause);
@@ -144,7 +152,11 @@ public:
 	/** Set ScrollBarAlwaysVisible attribute */
 	SLATE_API void SetScrollBarAlwaysVisible(bool InAlwaysVisible);
 
-	/** Set ScrollBarTrackAlwaysVisible attribute */
+	/**
+	 * Set ScrollBarTrackAlwaysVisible attribute.
+	 * Even if true, tracks will only be visible if the brush from the style has an ImageSize.
+	 * See FScrollBarStyle::VerticalBottomSlotImage and equivalents
+	 */
 	SLATE_API void SetScrollBarTrackAlwaysVisible(bool InAlwaysVisible);
 
 	/** Set the visibility of the ScrollBar when it is not needed. The default value is EVisibility::Collapsed. */
@@ -176,9 +188,19 @@ protected:
 	/** @return the name of an image for the scrollbar thumb based on whether the user is dragging or hovering it. */
 	SLATE_API const FSlateBrush* GetDragThumbImage() const;
 
+	TSlateAttributeRef<EVisibility> GetUserVisibilityAttribute() const
+	{
+		return TSlateAttributeRef<EVisibility>{SharedThis(this), UserVisibilityAttribute};
+	}
+
+#if WITH_EDITORONLY_DATA
 	/** The scrollbar's visibility as specified by the user. Will be compounded with internal visibility rules. */
-	TAttribute<EVisibility> UserVisibility;
+	UE_DEPRECATED(5.7, "Use SetUserVisibility / GetUserVisibilityAttribute")
+	TSlateDeprecatedTAttribute<EVisibility> UserVisibility;
+#endif
 
+	const FScrollBarStyle* Style = nullptr;
+	TSharedPtr<SBorder> BackgroundBorder;
 	TSharedPtr<SImage> TopImage;
 	TSharedPtr<SImage> BottomImage;
 	TSharedPtr<SBorder> DragThumb;
@@ -214,4 +236,10 @@ protected:
 	const FSlateBrush* TopBrush;
 	/** Bottom brush */
 	const FSlateBrush* BottomBrush;
+
+private:
+	bool bIsThicknessAttrSet = false;
+
+	/** The scrollbar's visibility as specified by the user. Will be compounded with internal visibility rules. */
+	TSlateAttribute<EVisibility> UserVisibilityAttribute;
 };
diff --git a/Engine/Source/Runtime/Slate/Public/Widgets/Layout/SScrollBarTrack.h b/Engine/Source/Runtime/Slate/Public/Widgets/Layout/SScrollBarTrack.h
--- a/Engine/Source/Runtime/Slate/Public/Widgets/Layout/SScrollBarTrack.h
+++ b/Engine/Source/Runtime/Slate/Public/Widgets/Layout/SScrollBarTrack.h
@@ -39,7 +39,7 @@ public:
 		_Visibility = EVisibility::SelfHitTestInvisible;
 	}
 
-	SLATE_NAMED_SLOT(FArguments, TopSlot)
+		SLATE_NAMED_SLOT(FArguments, TopSlot)
 		SLATE_NAMED_SLOT(FArguments, ThumbSlot)
 		SLATE_NAMED_SLOT(FArguments, BottomSlot)
 		SLATE_ARGUMENT(EOrientation, Orientation)
@@ -114,6 +114,9 @@ public:
 
 	SLATE_API void SetIsAlwaysVisible(bool InIsAlwaysVisible);
 
+private:
+	void InvalidateIsNeeded(bool bWasNeeded);
+
 protected:
 
 	static const int32 TOP_SLOT_INDEX = 0;
diff --git a/Engine/Source/Runtime/Slate/Public/Widgets/Layout/SScrollBorder.h b/Engine/Source/Runtime/Slate/Public/Widgets/Layout/SScrollBorder.h
--- a/Engine/Source/Runtime/Slate/Public/Widgets/Layout/SScrollBorder.h
+++ b/Engine/Source/Runtime/Slate/Public/Widgets/Layout/SScrollBorder.h
@@ -13,11 +13,18 @@
 
 class IScrollableWidget;
 
+template <>
+struct TWidgetTypeTraits<class SScrollBorder>
+{
+	static constexpr bool SupportsInvalidation() { return true; }
+};
+
 /**
  * Shows a border above and below a scrollable area.
  */
 class SScrollBorder : public SCompoundWidget
 {
+	SLATE_DECLARE_WIDGET_API(SScrollBorder, SCompoundWidget, SLATE_API)
 public:
 	SLATE_BEGIN_ARGS(SScrollBorder)
 		: _Style( &FCoreStyle::Get().GetWidgetStyle<FScrollBorderStyle>("ScrollBorder") )
@@ -34,6 +41,7 @@ public:
 		SLATE_DEFAULT_SLOT( FArguments, Content )
 	SLATE_END_ARGS()
 
+	SLATE_API SScrollBorder();
 	SLATE_API virtual ~SScrollBorder();
 
 	/**
@@ -42,6 +50,11 @@ public:
 	SLATE_API void Construct(const FArguments& InArgs, TSharedRef<IScrollableWidget> InScrollableWidget);
 	
 protected:
+	SLATE_API void SetBorderFadeDistance(TAttribute<FVector2D> InBorderFadeDistance);
+	TSlateAttributeRef<FVector2D> GetBorderFadeDistanceAttribute() const
+	{
+		return TSlateAttributeRef<FVector2D>{SharedThis(this), BorderFadeDistanceAttribute};
+	}
 
 	/** Gets the top border opacity. */
 	SLATE_API FSlateColor GetTopBorderOpacity() const;
@@ -49,9 +62,16 @@ protected:
 	/** Gets the bottom border opacity. */
 	SLATE_API FSlateColor GetBottomBorderOpacity() const;
 
+#if WITH_EDITORONLY_DATA
 	/** The border fade distance in normalized coordinates. */
-	TAttribute<FVector2D> BorderFadeDistance;
+	UE_DEPRECATED(5.7, "Use SetBorderFadeDistance / GetBorderFadeDistanceAttribute")
+	TSlateDeprecatedTAttribute<FVector2D> BorderFadeDistance;
+#endif
 
 	/** The widget that to watch for when the scrolling changes to show the border. */
 	TAttribute< TWeakPtr< IScrollableWidget > > ScrollableWidget;
+
+private:
+	/** The border fade distance in normalized coordinates. */
+	TSlateAttribute<FVector2D> BorderFadeDistanceAttribute;
 };
diff --git a/Engine/Source/Runtime/Slate/Public/Widgets/Layout/SScrollBox.h b/Engine/Source/Runtime/Slate/Public/Widgets/Layout/SScrollBox.h
--- a/Engine/Source/Runtime/Slate/Public/Widgets/Layout/SScrollBox.h
+++ b/Engine/Source/Runtime/Slate/Public/Widgets/Layout/SScrollBox.h
@@ -71,6 +71,12 @@ enum class EScrollWhenFocusChanges : uint8
 	AnimatedScroll,
 };
 
+template <>
+struct TWidgetTypeTraits<class SScrollBox>
+{
+	static constexpr bool SupportsInvalidation() { return true; }
+};
+
 /** SScrollBox can scroll through an arbitrary number of widgets. */
 class SScrollBox : public SCompoundWidget
 {
@@ -166,6 +172,7 @@ public:
 		, _NavigationDestination(EDescendantScrollDestination::IntoView)
 		, _NavigationScrollPadding(0.0f)
 		, _ScrollWhenFocusChanges(EScrollWhenFocusChanges::NoScroll)
+		, _ScrollBarRightClickDragAllowed(false)
 		, _OnUserScrolled()
 	    , _OnScrollBarVisibilityChanged()
 		, _OnFocusReceived()
@@ -227,6 +234,8 @@ public:
 
 		SLATE_ARGUMENT(EScrollWhenFocusChanges, ScrollWhenFocusChanges);
 
+		SLATE_ARGUMENT(bool, ScrollBarRightClickDragAllowed);
+
 		/** Called when the button is clicked */
 		SLATE_EVENT(FOnUserScrolled, OnUserScrolled)
 
@@ -476,6 +485,11 @@ protected:
 	/** The scrollbar which controls scrolling for the scrollbox. */
 	TSharedPtr<SScrollBar> ScrollBar;
 
+	/** The Top or Left Image of the shadow hint depending on orientation. */
+	TSharedPtr<SImage> TopLeftShadowHint;
+	/** The Bottom or Right Image of the shadow hint depending on orientation. */
+	TSharedPtr<SImage> BottomRightShadowHint;
+
 	/** The amount we have scrolled this tick cycle */
 	float TickScrollDelta;
 
@@ -598,8 +612,16 @@ protected:
 	bool bAllowContentToShrink = false;
 };
 
+
+template <>
+struct TWidgetTypeTraits<class SScrollPanel>
+{
+	static constexpr bool SupportsInvalidation() { return true; }
+};
+
 class SScrollPanel : public SPanel
 {
+	SLATE_DECLARE_WIDGET(SScrollPanel, SPanel)
 public:
 
 	SLATE_BEGIN_ARGS(SScrollPanel)
@@ -631,6 +653,7 @@ public:
 	void SetOrientation(EOrientation InOrientation)
 	{
 		Orientation = InOrientation;
+		Invalidate(EInvalidateWidgetReason::Layout);
 	}
 
 	virtual void OnArrangeChildren(const FGeometry& AllottedGeometry, FArrangedChildren& ArrangedChildren) const override;
diff --git a/Engine/Source/Runtime/Slate/Public/Widgets/Layout/SSplitter.h b/Engine/Source/Runtime/Slate/Public/Widgets/Layout/SSplitter.h
--- a/Engine/Source/Runtime/Slate/Public/Widgets/Layout/SSplitter.h
+++ b/Engine/Source/Runtime/Slate/Public/Widgets/Layout/SSplitter.h
@@ -37,13 +37,21 @@ namespace ESplitterResizeMode
 }
 
 class FLayoutGeometry;
+
+
+template <>
+struct TWidgetTypeTraits<class SSplitter>
+{
+	static constexpr bool SupportsInvalidation() { return true; }
+};
+
 /**
  * SSplitter divides its allotted area into N segments, where N is the number of children it has.
  * It allows the users to resize the children along the splitters axis: that is, horizontally or vertically.
  */
 class SSplitter : public SPanel
 {
-
+	SLATE_DECLARE_WIDGET_API(SSplitter, SPanel, SLATE_API)
 public:
 	/** How should a child's size be determined */
 	enum ESizeRule
@@ -62,44 +70,36 @@ public:
 	DECLARE_DELEGATE_RetVal_OneParam(FVector2D, FOnGetMaxSlotSize, int32);
 
 public:
-	class FSlot : public TSlotBase<FSlot>
+	class FSlot : public TWidgetSlotWithAttributeSupport<FSlot>
 	{
 	public:		
-		FSlot()
-			: TSlotBase<FSlot>()
-			, SizingRule( FractionOfParent )
-			, SizeValue( 1 )
-			, MinSizeValue( 0 )
-		{
-		}
+		FSlot();
 
-		SLATE_SLOT_BEGIN_ARGS(FSlot, TSlotBase<FSlot>)
+		SLATE_SLOT_BEGIN_ARGS(FSlot, TWidgetSlotWithAttributeSupport<FSlot>)
 			/** The size rule used by the slot. */
 			SLATE_ATTRIBUTE(ESizeRule, SizeRule)
-			/** When the RuleSize is set to FractionOfParent, the size of the slot is the Value percentage of its parent size. */
+			/** When the SizeRule is set to FractionOfParent, the size of the slot is the Value percentage of its parent size. */
 			SLATE_ATTRIBUTE(float, Value)
 			/** Minimum slot size when resizing. */
 			SLATE_ATTRIBUTE(float, MinSize)
-			/** Can the slot be resize by the user. */
+			/** Can the slot be resized by the user. */
 			SLATE_ARGUMENT(TOptional<bool>, Resizable)
 			/** Callback when the slot is resized. */
 			SLATE_EVENT(FOnSlotResized, OnSlotResized)
 		SLATE_SLOT_END_ARGS()
 
 		UE_API void Construct(const FChildren& SlotOwner, FSlotArguments&& InArgs);
+		static UE_API void RegisterAttributes(FSlateWidgetSlotAttributeInitializer& AttributeInitializer);
 
-		/** When the RuleSize is set to FractionOfParent, the size of the slot is the Value percentage of its parent size. */
-		void SetSizeValue( TAttribute<float> InValue )
-		{
-			SizeValue = MoveTemp(InValue);
-		}
+		/** When the SizeRule is set to FractionOfParent, the size of the slot is the Value percentage of its parent size. */
+		UE_API void SetSizeValue( TAttribute<float> InValue );
 		float GetSizeValue() const
 		{
 			return SizeValue.Get();
 		}
 
 		/**
-		 * Can the slot be resize by the user.
+		 * Can the slot be resized by the user.
 		 * @see CanBeResized()
 		 */
 		void SetResizable(bool bInIsResizable)
@@ -112,13 +112,10 @@ public:
 		}
 
 		/** Minimum slot size when resizing. */
-		void SetMinSize(float InMinSize)
-		{
-			MinSizeValue = InMinSize;
-		}
+		UE_API void SetMinSize( TAttribute<float> InMinSize );
 		float GetMinSize() const
 		{
-			return MinSizeValue.Get(0.f);
+			return MinSizeValue.Get();
 		}
 		
 		/**
@@ -135,10 +132,7 @@ public:
 		}
 
 		/** The size rule used by the slot. */
-		void SetSizingRule( TAttribute<ESizeRule> InSizeRule ) 
-		{
-			SizingRule = MoveTemp(InSizeRule);
-		}
+		UE_API void SetSizingRule( TAttribute<ESizeRule> InSizeRule ); 
 		ESizeRule GetSizingRule() const
 		{
 			return SizingRule.Get();
@@ -149,9 +143,10 @@ public:
 		UE_API bool CanBeResized() const;
 
 	private:
-		TAttribute<ESizeRule> SizingRule;
-		TAttribute<float> SizeValue;
-		TAttribute<float> MinSizeValue;
+		TSlateSlotAttribute<ESizeRule> SizingRule;
+		TSlateSlotAttribute<float> SizeValue;
+		TSlateSlotAttribute<float> MinSizeValue;
+
 		FOnSlotResized OnSlotResized_Handler;
 		TOptional<bool> bIsResizable;
 	};
