Add ability to set tab to be foregrounded when using FLayoutExtender::ExtendStack

diff --git a/Engine/Source/Runtime/Slate/Private/Framework/Docking/LayoutExtender.cpp b/Engine/Source/Runtime/Slate/Private/Framework/Docking/LayoutExtender.cpp
--- a/Engine/Source/Runtime/Slate/Private/Framework/Docking/LayoutExtender.cpp
+++ b/Engine/Source/Runtime/Slate/Private/Framework/Docking/LayoutExtender.cpp
@@ -12,9 +12,9 @@ void FLayoutExtender::ExtendArea(FName ExtensionId, const FAreaExtension& AreaEx
 	AreaExtensions.Add(ExtensionId, FExtendedArea(AreaExtension));
 }
 
-SLATE_API void FLayoutExtender::ExtendStack(FName ExtensionId, ELayoutExtensionPosition Position, FTabManager::FTab TabToAdd)
+SLATE_API void FLayoutExtender::ExtendStack(FName ExtensionId, ELayoutExtensionPosition Position, FTabManager::FTab TabToAdd, bool bShouldBeForegrounded)
 {
-	StackExtensions.Add(ExtensionId, FExtendedTab(Position, TabToAdd));
+	StackExtensions.Add(ExtensionId, FExtendedTab(Position, TabToAdd, bShouldBeForegrounded));
 }
 
 void FLayoutExtender::ExtendAreaRecursive(const TSharedRef<FTabManager::FArea>& Area) const
@@ -38,4 +38,4 @@ void FLayoutExtender::ExtendAreaRecursive(const TSharedRef<FTabManager::FArea>&
 			ExtendAreaRecursive(ChildArea.ToSharedRef());
 		}
 	}
-}
\ No newline at end of file
+}
diff --git a/Engine/Source/Runtime/Slate/Private/Framework/Docking/TabManager.cpp b/Engine/Source/Runtime/Slate/Private/Framework/Docking/TabManager.cpp
--- a/Engine/Source/Runtime/Slate/Private/Framework/Docking/TabManager.cpp
+++ b/Engine/Source/Runtime/Slate/Private/Framework/Docking/TabManager.cpp
@@ -654,6 +654,13 @@ void FTabManager::FLayout::ProcessExtensions(const FLayoutExtender& Extender)
 				Stack->Tabs.Insert(NewTab, InsertedTabIndex++);
 			}
 		}
+		
+		// Check if any of the extended tabs were marked as needing to be foregrounded
+		const FTabId ForegroundTab = Extender.FindFirstForegroundTab(Stack->GetExtensionId());
+		if (ForegroundTab.TabType != NAME_None)
+		{
+			Stack->SetForegroundTab(ForegroundTab);
+		}
 	}
 }
 
diff --git a/Engine/Source/Runtime/Slate/Public/Framework/Docking/LayoutExtender.h b/Engine/Source/Runtime/Slate/Public/Framework/Docking/LayoutExtender.h
--- a/Engine/Source/Runtime/Slate/Public/Framework/Docking/LayoutExtender.h
+++ b/Engine/Source/Runtime/Slate/Public/Framework/Docking/LayoutExtender.h
@@ -53,8 +53,9 @@ public:
 	 * @param AreaExtensions		A callback to call with the default layout for the desired area
 	 * @param Position				Whether to place the new tab at the start or end of this stack. Above and below are not used for this case.
 	 * @param TabToAdd				The new tab definition
+	 * @param bShouldBeForegrounded	Whether the tab should be foregrounded, first foreground tab found when applying extension will be used. So competing foregrounded tabs are not handled.
 	 */
-	SLATE_API void ExtendStack(FName ExtensionId, ELayoutExtensionPosition Position, FTabManager::FTab TabToAdd);
+	SLATE_API void ExtendStack(FName ExtensionId, ELayoutExtensionPosition Position, FTabManager::FTab TabToAdd, bool bShouldBeForegrounded=false);
 
 
 	/**
@@ -104,6 +105,22 @@ public:
 			}
 		}
 	}
+	
+	FTabId FindFirstForegroundTab(FName ExtensionID) const
+	{
+		TArray<FExtendedTab, TInlineAllocator<2>> Extensions;
+		StackExtensions.MultiFind(ExtensionID, Extensions, true);
+
+		for (FExtendedTab& Extension : Extensions)
+		{
+			if (Extension.bShouldBeForegrounded)
+			{
+				return Extension.TabToAdd.TabId;
+			}
+		}
+		
+		return FTabId();
+	}
 
 	/**
 	 * Recursively extend the specified area
@@ -115,14 +132,16 @@ private:
 	/** Extended tab information */
 	struct FExtendedTab
 	{
-		FExtendedTab(ELayoutExtensionPosition InPosition, const FTabManager::FTab& InTabToAdd)
-			: Position(InPosition), TabToAdd(InTabToAdd)
+		FExtendedTab(ELayoutExtensionPosition InPosition, const FTabManager::FTab& InTabToAdd, bool bInShouldBeForegrounded=false)
+			: Position(InPosition), TabToAdd(InTabToAdd), bShouldBeForegrounded(bInShouldBeForegrounded)
 		{}
 
 		/** The tab position */
 		ELayoutExtensionPosition Position;
 		/** The tab definition */
 		FTabManager::FTab TabToAdd;
+		/** Whether the tab should be foregrounded, competing foregrounded tabs are not handled */ 
+		bool bShouldBeForegrounded;
 	};
 
 	/** Extended area information */
@@ -144,4 +163,4 @@ private:
 
 	/** Map of extensions for areas */
 	TMultiMap<FName, FExtendedArea> AreaExtensions;
-};
\ No newline at end of file
+};
diff --git a/Engine/Source/Runtime/Slate/Public/Framework/Docking/TabManager.h b/Engine/Source/Runtime/Slate/Public/Framework/Docking/TabManager.h
--- a/Engine/Source/Runtime/Slate/Public/Framework/Docking/TabManager.h
+++ b/Engine/Source/Runtime/Slate/Public/Framework/Docking/TabManager.h
@@ -61,7 +61,8 @@ ENUM_CLASS_FLAGS(ETabIdFlags);
 struct FTabId
 {
 	FTabId()
-		: InstanceId(INDEX_NONE)
+		: TabType(NAME_None) 
+		, InstanceId(INDEX_NONE)
 		, Flags(ETabIdFlags::SaveLayout)
 	{ }
 
