Slate Dynamic Invalidation - ProgressBar

diff --git a/Engine/Plugins/Experimental/UAF/UAF/Source/UAFEditor/Private/UAF/AbstractSkeleton/Labels/SLabelBinding.cpp b/Engine/Plugins/Experimental/UAF/UAF/Source/UAFEditor/Private/UAF/AbstractSkeleton/Labels/SLabelBinding.cpp
--- a/Engine/Plugins/Experimental/UAF/UAF/Source/UAFEditor/Private/UAF/AbstractSkeleton/Labels/SLabelBinding.cpp
+++ b/Engine/Plugins/Experimental/UAF/UAF/Source/UAFEditor/Private/UAF/AbstractSkeleton/Labels/SLabelBinding.cpp
@@ -329,7 +329,6 @@ namespace UE::UAF::Labels
 										.FillColorAndOpacity(NumUnbound == 0 ? FAppStyle::Get().GetSlateColor("Colors.AccentGreen") : FAppStyle::Get().GetSlateColor("Colors.AccentBlue"))
 										.FillImage(FAppStyle::Get().GetBrush("WhiteBrush"))
 										.Percent(static_cast<float>(NumBound) / LabelBindings.Num())
-										.RefreshRate(0)
 								]
 						]
 						+ SHorizontalBox::Slot()
diff --git a/Engine/Plugins/IoStoreInsights/Source/IoStoreInsights/Private/Widgets/SIoStoreAnalysisReadSizeHistogramView.cpp b/Engine/Plugins/IoStoreInsights/Source/IoStoreInsights/Private/Widgets/SIoStoreAnalysisReadSizeHistogramView.cpp
--- a/Engine/Plugins/IoStoreInsights/Source/IoStoreInsights/Private/Widgets/SIoStoreAnalysisReadSizeHistogramView.cpp
+++ b/Engine/Plugins/IoStoreInsights/Source/IoStoreInsights/Private/Widgets/SIoStoreAnalysisReadSizeHistogramView.cpp
@@ -51,7 +51,6 @@ namespace UE::IoStoreInsights
 						[
 							SNew(SProgressBar)
 								.Percent(Item->CountAsPctNormalized)
-								.RefreshRate(0)
 								.BorderPadding(FVector2D(4, 4))
 								.BarFillStyle(EProgressBarFillStyle::Scale)
 								.BackgroundImage(FAppStyle::Get().GetBrush("NoBrush"))
diff --git a/Engine/Source/Runtime/Slate/Private/Widgets/Notifications/SProgressBar.cpp b/Engine/Source/Runtime/Slate/Private/Widgets/Notifications/SProgressBar.cpp
--- a/Engine/Source/Runtime/Slate/Private/Widgets/Notifications/SProgressBar.cpp
+++ b/Engine/Source/Runtime/Slate/Private/Widgets/Notifications/SProgressBar.cpp
@@ -6,10 +6,40 @@
 #include UE_INLINE_GENERATED_CPP_BY_NAME(SProgressBar)
 
 
+SLATE_IMPLEMENT_WIDGET(SProgressBar)
+void SProgressBar::PrivateRegisterAttributes(FSlateAttributeInitializer& AttributeInitializer)
+{
+	struct FInvalidation
+	{
+		static void UpdateMarquee(SWidget& OwningWidget)
+		{
+			static_cast<SProgressBar&>(OwningWidget).bNeedsMarqueeTimerUpdate = true;
+		}
+	};
+	
+	SLATE_ADD_MEMBER_ATTRIBUTE_DEFINITION(AttributeInitializer, Percent, EInvalidateWidgetReason::Paint)
+		.OnValueChanged(FSlateAttributeDescriptor::FAttributeValueChangedDelegate::CreateStatic(FInvalidation::UpdateMarquee));
+	
+	SLATE_ADD_MEMBER_ATTRIBUTE_DEFINITION(AttributeInitializer, FillColorAndOpacity, EInvalidateWidgetReason::Paint);
+	SLATE_ADD_MEMBER_ATTRIBUTE_DEFINITION(AttributeInitializer, BorderPadding, EInvalidateWidgetReason::Layout);
+}
+
+SProgressBar::SProgressBar()
+	: Percent(*this, TOptional<float>())
+	, FillColorAndOpacity(*this, FLinearColor::White)
+	, BorderPadding(*this, FVector2D(0,0))
+{
+}
+
+SProgressBar::~SProgressBar()
+{
+}
+
 void SProgressBar::Construct( const FArguments& InArgs )
 {
 	check(InArgs._Style);
 
+	bNeedsMarqueeTimerUpdate = true;
 	MarqueeOffset = 0.0f;
 
 	Style = InArgs._Style;
@@ -22,12 +52,9 @@ void SProgressBar::Construct( const FArguments& InArgs )
 	FillImage = InArgs._FillImage;
 	MarqueeImage = InArgs._MarqueeImage;
 	
-	FillColorAndOpacity = InArgs._FillColorAndOpacity;
-	BorderPadding = InArgs._BorderPadding;
+	FillColorAndOpacity.Assign(*this, InArgs._FillColorAndOpacity);
+	BorderPadding.Assign(*this, InArgs._BorderPadding);
 
-	CurrentTickRate = 0.0f;
-	MinimumTickRate = InArgs._RefreshRate;
-	
 	SetCanTick(false);
 
 	UpdateMarqueeActiveTimer();
@@ -35,12 +62,7 @@ void SProgressBar::Construct( const FArguments& InArgs )
 
 void SProgressBar::SetPercent(TAttribute< TOptional<float> > InPercent)
 {
-	if ( !Percent.IdenticalTo(InPercent) )
-	{
-		Percent = InPercent;
-		UpdateMarqueeActiveTimer();
-		Invalidate(EInvalidateWidget::LayoutAndVolatility);
-	}
+	Percent.Assign(*this, MoveTemp(InPercent));
 }
 
 void SProgressBar::SetStyle(const FProgressBarStyle* InStyle)
@@ -79,20 +101,12 @@ void SProgressBar::SetBarFillStyle(EProgressBarFillStyle::Type InBarFillStyle)
 
 void SProgressBar::SetFillColorAndOpacity(TAttribute< FSlateColor > InFillColorAndOpacity)
 {
-	if(!FillColorAndOpacity.IdenticalTo(InFillColorAndOpacity))
-	{
-		FillColorAndOpacity = InFillColorAndOpacity;
-		Invalidate(EInvalidateWidget::Paint);
-	}
+	FillColorAndOpacity.Assign(*this, MoveTemp(InFillColorAndOpacity));
 }
 
 void SProgressBar::SetBorderPadding(TAttribute< FVector2D > InBorderPadding)
 {
-	if(!BorderPadding.IdenticalTo(InBorderPadding))
-	{
-		BorderPadding = InBorderPadding;
-		Invalidate(EInvalidateWidget::Layout);
-	}
+	BorderPadding.Assign(*this, MoveTemp(InBorderPadding));
 }
 
 void SProgressBar::SetBackgroundImage(const FSlateBrush* InBackgroundImage)
@@ -160,6 +174,11 @@ bool PushTransformedClip(FSlateWindowElementList& OutDrawElements, const FGeomet
 
 int32 SProgressBar::OnPaint( const FPaintArgs& Args, const FGeometry& AllottedGeometry, const FSlateRect& MyCullingRect, FSlateWindowElementList& OutDrawElements, int32 LayerId, const FWidgetStyle& InWidgetStyle, bool bParentEnabled ) const
 {
+	if (bNeedsMarqueeTimerUpdate)
+	{
+		const_cast<SProgressBar*>(this)->UpdateMarqueeActiveTimer();
+	}
+
 	// Used to track the layer ID we will return.
 	int32 RetLayerId = LayerId;
 
@@ -501,54 +520,27 @@ FVector2D SProgressBar::ComputeDesiredSize( float ) const
 	return FVector2D(GetMarqueeImage()->ImageSize);
 }
 
-bool SProgressBar::ComputeVolatility() const
-{
-	return SLeafWidget::ComputeVolatility() || Percent.IsBound() || FillColorAndOpacity.IsBound() || BorderPadding.IsBound();
-}
-
-void SProgressBar::SetActiveTimerTickRate(float TickRate)
+void SProgressBar::UpdateMarqueeActiveTimer()
 {
-	if (CurrentTickRate != TickRate || !ActiveTimerHandle.IsValid())
+	if (!Percent.Get().IsSet() || Style->EnableFillAnimation)
 	{
-		CurrentTickRate = TickRate;
-
-		TSharedPtr<FActiveTimerHandle> SharedActiveTimerHandle = ActiveTimerHandle.Pin();
-		if (SharedActiveTimerHandle.IsValid())
+		// If percent is not set or fill animation is enabled then set the timer
+		if (!ActiveTimerHandle.IsValid())
 		{
-			UnRegisterActiveTimer(SharedActiveTimerHandle.ToSharedRef());
+			// If we already had a handle, there is no need to register a new one.
+			ActiveTimerHandle = RegisterActiveTimer(0.f, FWidgetActiveTimerDelegate::CreateSP(this, &SProgressBar::ActiveTick));
 		}
-
-		UpdateMarqueeActiveTimer();
 	}
-}
-
-void SProgressBar::UpdateMarqueeActiveTimer()
-{
-	if (ActiveTimerHandle.IsValid())
+	else if (ActiveTimerHandle.IsValid())
 	{
 		UnRegisterActiveTimer(ActiveTimerHandle.Pin().ToSharedRef());
 	}
-
-	if ((!Percent.IsBound() && !Percent.Get().IsSet()) || Style->EnableFillAnimation)
-	{
-		// If percent is not bound or set then its marquee. Set the timer
-		ActiveTimerHandle = RegisterActiveTimer(CurrentTickRate, FWidgetActiveTimerDelegate::CreateSP(this, &SProgressBar::ActiveTick));
-	}
+	bNeedsMarqueeTimerUpdate = false;
 }
 
 EActiveTimerReturnType SProgressBar::ActiveTick(double InCurrentTime, float InDeltaTime)
 {
 	MarqueeOffset = (float)(InCurrentTime - FMath::FloorToDouble(InCurrentTime));
-	
-	TOptional<float> PercentFraction = Percent.Get();
-	if (PercentFraction.IsSet() && !Style->EnableFillAnimation)
-	{
-		SetActiveTimerTickRate(MinimumTickRate);
-	}
-	else
-	{
-		SetActiveTimerTickRate(0.0f);
-	}
 
 	Invalidate(EInvalidateWidget::Paint);
 
diff --git a/Engine/Source/Runtime/Slate/Public/Widgets/Notifications/SProgressBar.h b/Engine/Source/Runtime/Slate/Public/Widgets/Notifications/SProgressBar.h
--- a/Engine/Source/Runtime/Slate/Public/Widgets/Notifications/SProgressBar.h
+++ b/Engine/Source/Runtime/Slate/Public/Widgets/Notifications/SProgressBar.h
@@ -57,10 +57,17 @@ namespace EProgressBarFillStyle
 	};
 }
 
+
+template <>
+struct TWidgetTypeTraits<class SProgressBar>
+{
+	static constexpr bool SupportsInvalidation() { return true; }
+};
+
 /** A progress bar widget.*/
 class SProgressBar : public SLeafWidget
 {
-
+	SLATE_DECLARE_WIDGET_API(SProgressBar, SLeafWidget, SLATE_API)
 public:
 	SLATE_BEGIN_ARGS(SProgressBar)
 		: _Style( &FAppStyle::Get().GetWidgetStyle<FProgressBarStyle>("ProgressBar") )
@@ -105,10 +112,13 @@ public:
 		SLATE_ARGUMENT(const FSlateBrush*, MarqueeImage)
 
 		/** Rate at which this widget is ticked when sleeping in seconds */
-		SLATE_ARGUMENT(float, RefreshRate)
+		SLATE_ARGUMENT_DEPRECATED(float, RefreshRate, 5.8, "RefreshRate will be ignored as the ProgressBar does not need to tick when sleeping anymore.")
 
 	SLATE_END_ARGS()
 
+	SLATE_API SProgressBar();
+	SLATE_API ~SProgressBar();
+
 	/**
 	 * Construct the widget
 	 * 
@@ -118,7 +128,6 @@ public:
 
 	SLATE_API virtual int32 OnPaint( const FPaintArgs& Args, const FGeometry& AllottedGeometry, const FSlateRect& MyCullingRect, FSlateWindowElementList& OutDrawElements, int32 LayerId, const FWidgetStyle& InWidgetStyle, bool bParentEnabled ) const override;
 	SLATE_API virtual FVector2D ComputeDesiredSize(float) const override;
-	SLATE_API virtual bool ComputeVolatility() const override;
 
 	/** See attribute Percent */
 	SLATE_API void SetPercent(TAttribute< TOptional<float> > InPercent);
@@ -149,9 +158,6 @@ public:
 	
 private:
 
-	/** Controls the speed at which the widget is ticked when in slate sleep mode */
-	void SetActiveTimerTickRate(float TickRate);
-
 	void UpdateMarqueeActiveTimer();
 
 	/** Widgets active tick */
@@ -170,7 +176,7 @@ private:
 	const FProgressBarStyle* Style;
 
 	/** The text displayed over the progress bar */
-	TAttribute< TOptional<float> > Percent;
+	TSlateAttribute< TOptional<float> > Percent;
 
 	EProgressBarFillType::Type BarFillType;
 
@@ -186,10 +192,10 @@ private:
 	const FSlateBrush* MarqueeImage;
 
 	/** Fill Color and Opacity */
-	TAttribute<FSlateColor> FillColorAndOpacity;
+	TSlateAttribute<FSlateColor> FillColorAndOpacity;
 
 	/** Border Padding */
-	TAttribute<FVector2D> BorderPadding;
+	TSlateAttribute<FVector2D> BorderPadding;
 
 	/** Value to drive progress bar animation */
 	float MarqueeOffset;
@@ -197,10 +203,6 @@ private:
 	/** Reference to the widgets current active timer */
 	TWeakPtr<FActiveTimerHandle> ActiveTimerHandle;
 
-	/** Rate at which the widget is currently ticked when slate sleep mode is active */
-	float CurrentTickRate;
-
-	/** The slowest that this widget can tick when in slate sleep mode */
-	float MinimumTickRate;
+	bool bNeedsMarqueeTimerUpdate;
 };
 
