SlateApplication: Add some basic program boot trace cpu scopes.

#rnx

diff --git a/Engine/Source/Runtime/Slate/Private/Framework/Application/SlateApplication.cpp b/Engine/Source/Runtime/Slate/Private/Framework/Application/SlateApplication.cpp
--- a/Engine/Source/Runtime/Slate/Private/Framework/Application/SlateApplication.cpp
+++ b/Engine/Source/Runtime/Slate/Private/Framework/Application/SlateApplication.cpp
@@ -1033,6 +1033,7 @@ const FStyleNode* FSlateApplication::GetRootStyle() const
 
 bool FSlateApplication::InitializeRenderer( TSharedRef<FSlateRenderer> InRenderer, bool bQuietMode )
 {
+	TRACE_CPUPROFILER_EVENT_SCOPE(FSlateApplication::InitializeRenderer);
 	Renderer = InRenderer;
 	bool bResult = Renderer->Initialize();
 	if (!bResult && !bQuietMode)
@@ -1918,6 +1919,7 @@ FWidgetPath FSlateApplication::LocateWidgetInWindow(UE::Slate::FDeprecateVector2
 
 TSharedRef<SWindow> FSlateApplication::AddWindow( TSharedRef<SWindow> InSlateWindow, const bool bShowImmediately )
 {
+	TRACE_CPUPROFILER_EVENT_SCOPE(FSlateApplication::AddWindow);
 	// Add the Slate window to the Slate application's top-level window array.  Note that neither the Slate window
 	// or the native window are ready to be used yet, however we need to make sure they're in the Slate window
 	// array so that we can properly respond to OS window messages as soon as they're sent.  For example, a window
@@ -7514,6 +7516,8 @@ TSharedRef<FSlateApplication> FSlateApplication::InitializeAsStandaloneApplicati
 
 TSharedRef<FSlateApplication> FSlateApplication::InitializeAsStandaloneApplication(const TSharedRef< class FSlateRenderer >& PlatformRenderer, const TSharedRef<class GenericApplication>& InPlatformApplication)
 {
+	TRACE_CPUPROFILER_EVENT_SCOPE(FSlateApplication::InitializeAsStandaloneApplication);
+
 	// Initialise High DPI mode. This must be called before any window (including the splash screen is created).
 	// We don't need to use the force flag which was added for PIE.
 	const bool bForceEnable = false;
