Slate Dynamic Invalidation - Input Widgets

diff --git a/Engine/Source/Runtime/Slate/Private/Widgets/Input/SMenuAnchor.cpp b/Engine/Source/Runtime/Slate/Private/Widgets/Input/SMenuAnchor.cpp
--- a/Engine/Source/Runtime/Slate/Private/Widgets/Input/SMenuAnchor.cpp
+++ b/Engine/Source/Runtime/Slate/Private/Widgets/Input/SMenuAnchor.cpp
@@ -15,6 +15,13 @@
 #include "Framework/MultiBox/ToolMenuBase.h"
 #endif
 
+
+SLATE_IMPLEMENT_WIDGET(SMenuAnchor)
+void SMenuAnchor::PrivateRegisterAttributes(FSlateAttributeInitializer& AttributeInitializer)
+{
+	SLATE_ADD_MEMBER_ATTRIBUTE_DEFINITION(AttributeInitializer, PlacementAttribute, EInvalidateWidgetReason::Layout);
+}
+
 static FVector2D GetMenuOffsetForPlacement(const FGeometry& AllottedGeometry, EMenuPlacement PlacementMode, FVector2D PopupSizeLocalSpace)
 {
 	switch (PlacementMode)
@@ -81,7 +88,6 @@ void SMenuAnchor::Construct( const FArguments& InArgs )
 	WrappedContent                         = InArgs._MenuContent;
 	OnGetMenuContent                       = InArgs._OnGetMenuContent;
 	OnMenuOpenChanged                      = InArgs._OnMenuOpenChanged;
-	Placement                              = InArgs._Placement;
 	bFitInWindow						   = InArgs._FitInWindow;
 	Method                                 = InArgs._Method;
 	bShouldDeferPaintingAfterWindowContent = InArgs._ShouldDeferPaintingAfterWindowContent;
@@ -89,6 +95,8 @@ void SMenuAnchor::Construct( const FArguments& InArgs )
 	bShowMenuBackground					   = InArgs._ShowMenuBackground;
 	bIsCollapsedByParent                   = InArgs._IsCollapsedByParent;
 	bApplyWidgetStyleToMenu                = InArgs._ApplyWidgetStyleToMenu;
+
+	SetMenuPlacement(InArgs._Placement);
 }
 
 FGeometry SMenuAnchor::ComputeNewWindowMenuPlacement(const FGeometry& AllottedGeometry, const FVector2D& PopupDesiredSize, EMenuPlacement PlacementMode) const
@@ -120,7 +128,7 @@ void SMenuAnchor::Tick( const FGeometry& AllottedGeometry, const double InCurren
 	{
 		// Figure out where our attached pop-up window should be placed.
 		const FVector2D PopupContentDesiredSize = PopupWindow->GetContent()->GetDesiredSize();
-		FGeometry PopupGeometry = ComputeNewWindowMenuPlacement( AllottedGeometry, PopupContentDesiredSize, Placement.Get() );
+		FGeometry PopupGeometry = ComputeNewWindowMenuPlacement( AllottedGeometry, PopupContentDesiredSize, PlacementAttribute.Get() );
 		const FVector2D NewPosition = PopupGeometry.LocalToAbsolute(FVector2D::ZeroVector);
 		// NOTE: In order to get the right size of the window, we need to take whatever the incoming scale of the menu anchor,
 		// then divide out the popup window's DPI scale.  Finally, we need to divide that remainder to the draw size.
@@ -134,7 +142,7 @@ void SMenuAnchor::Tick( const FGeometry& AllottedGeometry, const double InCurren
 		const FVector2D NewSize = PopupGeometry.GetDrawSize() / ( AllottedGeometry.GetAccumulatedLayoutTransform().GetScale() / PopupWindow->GetLocalToWindowTransform().GetScale() );
 
 		// When in a ComboBox, we want the submenu to be at minimum the size of the combo box, so we update the Popupwindow minimum size.
-		const EMenuPlacement PlacementMode = Placement.Get();
+		const EMenuPlacement PlacementMode = PlacementAttribute.Get();
 		if (PlacementMode == MenuPlacement_ComboBox)
 		{
 			// Set the new size limits
@@ -153,7 +161,7 @@ void SMenuAnchor::Tick( const FGeometry& AllottedGeometry, const double InCurren
 		// can be called in DesktopSpace or WindowSpace, and we will not know which version of the Window
 		// geometry to use. Tick() is always in DesktopSpace, so cache the position here and just use
 		// it in OnArrangeChildren().
-		const FPopupPlacement LocalPlacement(AllottedGeometry, Children[1].GetWidget()->GetDesiredSize(), Placement.Get());
+		const FPopupPlacement LocalPlacement(AllottedGeometry, Children[1].GetWidget()->GetDesiredSize(), PlacementAttribute.Get());
 
 		FVector2D FittedPlacement;
 
@@ -189,7 +197,7 @@ void SMenuAnchor::OnArrangeChildren( const FGeometry& AllottedGeometry, FArrange
 	const TSharedPtr<SWindow> PresentingWindow = PopupWindowPtr.Pin();
 	if (IsOpenAndReusingWindow() && PresentingWindow.IsValid())
 	{
-		const FPopupPlacement LocalPlacement(AllottedGeometry, Children[1].GetWidget()->GetDesiredSize(), Placement.Get());
+		const FPopupPlacement LocalPlacement(AllottedGeometry, Children[1].GetWidget()->GetDesiredSize(), PlacementAttribute.Get());
 		ArrangedChildren.AddWidget(AllottedGeometry.MakeChild(Children[1].GetWidget(), LocalPlacement.LocalPopupSize, FSlateLayoutTransform(LocalPopupPosition)));
 	}
 }
@@ -341,7 +349,7 @@ void SMenuAnchor::SetIsOpen( bool InIsOpen, const bool bFocusMenu, const int32 F
 					// Combo-boxes never size down smaller than the widget that spawned them, but all
 					// other pop-up menus are currently auto-sized
 					const FVector2D DesiredContentSize = MenuContentRef->GetDesiredSize();  // @todo slate: This is ignoring any window border size!
-					const EMenuPlacement PlacementMode = Placement.Get();
+					const EMenuPlacement PlacementMode = PlacementAttribute.Get();
 
 					const FVector2D NewPosition = FVector2D(MyGeometry.AbsolutePosition);
 					FVector2D NewWindowSize = DesiredContentSize;
@@ -629,11 +637,7 @@ FVector2D SMenuAnchor::GetMenuPosition() const
 
 void SMenuAnchor::SetMenuPlacement(TAttribute<EMenuPlacement> InMenuPlacement)
 {
-	if (!Placement.IsSet() || !Placement.IdenticalTo(InMenuPlacement))
-	{
-		Placement = InMenuPlacement;
-		Invalidate(EInvalidateWidget::Layout);
-	}
+	PlacementAttribute.Assign(*this, MoveTemp(InMenuPlacement));
 }
 
 void SMenuAnchor::SetFitInWindow(bool bFit)
@@ -695,6 +699,7 @@ SMenuAnchor::SMenuAnchor()
 	, MethodInUse()
 	, LocalPopupPosition( FVector2D::ZeroVector )
 	, Children(this)
+	, PlacementAttribute(*this)
 {
 }
 
diff --git a/Engine/Source/Runtime/Slate/Private/Widgets/Input/SMultiLineEditableTextBox.cpp b/Engine/Source/Runtime/Slate/Private/Widgets/Input/SMultiLineEditableTextBox.cpp
--- a/Engine/Source/Runtime/Slate/Private/Widgets/Input/SMultiLineEditableTextBox.cpp
+++ b/Engine/Source/Runtime/Slate/Private/Widgets/Input/SMultiLineEditableTextBox.cpp
@@ -450,6 +450,11 @@ void SMultiLineEditableTextBox::SetError( const FString& InError )
 	ErrorReporting->SetError( InError );
 }
 
+bool SMultiLineEditableTextBox::HasError() const
+{
+	return ErrorReporting && ErrorReporting->HasError();
+}
+
 /** @return Border image for the text box based on the hovered and focused state */
 const FSlateBrush* SMultiLineEditableTextBox::GetBorderImage() const
 {
diff --git a/Engine/Source/Runtime/Slate/Private/Widgets/Text/SInlineEditableTextBlock.cpp b/Engine/Source/Runtime/Slate/Private/Widgets/Text/SInlineEditableTextBlock.cpp
--- a/Engine/Source/Runtime/Slate/Private/Widgets/Text/SInlineEditableTextBlock.cpp
+++ b/Engine/Source/Runtime/Slate/Private/Widgets/Text/SInlineEditableTextBlock.cpp
@@ -8,10 +8,51 @@
 #include "Widgets/Input/SMultiLineEditableTextBox.h"
 #include "Widgets/Input/SEditableTextBox.h"
 
+SLATE_IMPLEMENT_WIDGET(SInlineEditableTextBlock)
+void SInlineEditableTextBlock::PrivateRegisterAttributes(FSlateAttributeInitializer& AttributeInitializer)
+{
+	SLATE_ADD_MEMBER_ATTRIBUTE_DEFINITION(AttributeInitializer, MaximumLengthAttribute, [](const SWidget& OwningWidget)
+	{
+		const SInlineEditableTextBlock& InlineTextBlock = static_cast<const SInlineEditableTextBlock&>(OwningWidget);
+		if (InlineTextBlock.IsInEditMode())
+		{
+			const bool bIsProcessingAttributeUpdate = InlineTextBlock.GetProxyHandle().IsValid(InlineTextBlock) &&
+				InlineTextBlock.GetProxyHandle().GetInvalidationRootHandle().GetInvalidationRoot() &&
+				InlineTextBlock.GetProxyHandle().GetInvalidationRootHandle().GetInvalidationRoot()->IsProcessingAttributeUpdate();
+
+			if (bIsProcessingAttributeUpdate)
+			{
+				// We want to retrigger an OnTextChanged event with the text of the edited textbox to revalidate the text.
+				// It cannot be retriggered right away while attributes are updated if OnTextChanged ends up creating or destroying the error
+				// message as this would lead to a crash. We also cannot register a timer at this point, so we delay until next tick.
+				bool bHadError = InlineTextBlock.HasTextBoxError();
+				FText ErrorMessage;
+				bool bHasError = !InlineTextBlock.IsTextAllowed(InlineTextBlock.GetEditableText(), ErrorMessage);
+
+				if (bHadError != bHasError)
+				{
+					// a popup widget would end up being created or destroyed, so we need to delay its creation/destruction
+					InlineTextBlock.bNeedsMaximumLengthChecked = true;
+					return EInvalidateWidgetReason::Paint;
+				}
+			}
+			
+			const_cast<SInlineEditableTextBlock&>(InlineTextBlock).OnTextChanged(InlineTextBlock.GetEditableText());
+		}
+		return EInvalidateWidgetReason::None;
+	});
+}
+
+SInlineEditableTextBlock::SInlineEditableTextBlock()
+	: MaximumLengthAttribute(*this)
+{
+}
+
 
 void SInlineEditableTextBlock::Construct( const FArguments& InArgs )
 {
 	check(InArgs._Style);
+	SetCanTick(false);
 
 	OnTextCommittedDelegate = InArgs._OnTextCommitted;
 	IsSelected = InArgs._IsSelected;
@@ -19,7 +60,7 @@ void SInlineEditableTextBlock::Construct( const FArguments& InArgs )
 	Text = InArgs._Text;
 	HintText = InArgs._HintText;
 	bIsReadOnly = InArgs._IsReadOnly;
-	MaximumLength = InArgs._MaximumLength;
+	bNeedsMaximumLengthChecked = false;
 
 	bIsMultiLine = InArgs._MultiLine;
 	bDelayedLeftClickEntersEditMode = InArgs._DelayedLeftClickEntersEditMode;
@@ -89,6 +130,17 @@ void SInlineEditableTextBlock::Construct( const FArguments& InArgs )
 			.ClearKeyboardFocusOnCommit(false)
 			.OverflowPolicy(InArgs._OverflowPolicy);
 	}
+
+	SetMaximumLength(InArgs._MaximumLength);
+}
+
+int32 SInlineEditableTextBlock::OnPaint(const FPaintArgs& Args, const FGeometry& AllottedGeometry, const FSlateRect& MyCullingRect,	FSlateWindowElementList& OutDrawElements, int32 LayerId, const FWidgetStyle& InWidgetStyle, bool bParentEnabled) const
+{
+	if (bNeedsMaximumLengthChecked)
+	{
+		const_cast<SInlineEditableTextBlock*>(this)->OnTextChanged(GetEditableText());
+	}
+	return SCompoundWidget::OnPaint(Args, AllottedGeometry, MyCullingRect, OutDrawElements, LayerId, InWidgetStyle, bParentEnabled);
 }
 
 SInlineEditableTextBlock::~SInlineEditableTextBlock()
@@ -135,6 +187,9 @@ void SInlineEditableTextBlock::EnterEditingMode()
 			FSlateApplication::Get().SetKeyboardFocus(ActiveTextBox, EFocusCause::SetDirectly);
 
 			TextBlock->SetVisibility(EVisibility::Collapsed);
+
+			// We retrigger an OnTextChanged event to revalidate the text
+			OnTextChanged(CurrentText);
 		}
 	}
 }
@@ -147,6 +202,7 @@ void SInlineEditableTextBlock::ExitEditingMode()
 	TextBlock->SetVisibility(EVisibility::Visible);
 	// Clear the error so it will vanish.
 	SetTextBoxError( FText::GetEmpty() );
+	bNeedsMaximumLengthChecked = false;
 
 	TSharedPtr<SWidget> CurrentFocus = FSlateApplication::Get().GetKeyboardFocusedWidget();
 	if (CurrentFocus.IsValid())
@@ -179,9 +235,9 @@ void SInlineEditableTextBlock::SetReadOnly(const TAttribute<bool>& bInIsReadOnly
 	bIsReadOnly = bInIsReadOnly;
 }
 
-void SInlineEditableTextBlock::SetMaximumLength(const TAttribute<int32>& InMaximumLength)
+void SInlineEditableTextBlock::SetMaximumLength(TAttribute<int32> InMaximumLength)
 {
-	MaximumLength = InMaximumLength;
+	MaximumLengthAttribute.Assign(*this, MoveTemp(InMaximumLength));
 }
 
 void SInlineEditableTextBlock::SetText( const TAttribute< FText >& InText )
@@ -327,8 +383,7 @@ void SInlineEditableTextBlock::OnTextChanged(const FText& InText)
 	if(IsInEditMode())
 	{
 		FText OutErrorMessage;
-		if(!FTextEditHelper::VerifyTextLength(InText, OutErrorMessage, MaximumLength.Get()) ||
-			(OnVerifyTextChanged.IsBound() && !OnVerifyTextChanged.Execute(InText, OutErrorMessage)))
+		if(!IsTextAllowed(InText, OutErrorMessage))
 		{
 			SetTextBoxError( OutErrorMessage );
 		}
@@ -337,6 +392,13 @@ void SInlineEditableTextBlock::OnTextChanged(const FText& InText)
 			SetTextBoxError( FText::GetEmpty() );
 		}
 	}
+	bNeedsMaximumLengthChecked = false;
+}
+
+bool SInlineEditableTextBlock::IsTextAllowed(const FText& InText, FText& OutErrorMessage) const
+{
+	return FTextEditHelper::VerifyTextLength(InText, OutErrorMessage, MaximumLengthAttribute.Get()) &&
+		(!OnVerifyTextChanged.IsBound() || OnVerifyTextChanged.Execute(InText, OutErrorMessage));
 }
 
 void SInlineEditableTextBlock::OnTextBoxCommitted(const FText& InText, ETextCommit::Type InCommitType)
@@ -353,8 +415,7 @@ void SInlineEditableTextBlock::OnTextBoxCommitted(const FText& InText, ETextComm
 		if(InCommitType == ETextCommit::OnEnter)
 		{
 			FText OutErrorMessage;
-			if (!FTextEditHelper::VerifyTextLength(InText, OutErrorMessage, MaximumLength.Get()) ||
-				(OnVerifyTextChanged.IsBound() && !OnVerifyTextChanged.Execute(InText, OutErrorMessage)))
+			if (!IsTextAllowed(InText, OutErrorMessage))
 			{
 				// Display as an error.
 				SetTextBoxError( OutErrorMessage );
@@ -364,8 +425,7 @@ void SInlineEditableTextBlock::OnTextBoxCommitted(const FText& InText, ETextComm
 		else if(InCommitType == ETextCommit::OnUserMovedFocus || InCommitType == ETextCommit::Default)
 		{
 			FText OutErrorMessage;
-			if (!FTextEditHelper::VerifyTextLength(InText, OutErrorMessage, MaximumLength.Get()) ||
-				(OnVerifyTextChanged.IsBound() && !OnVerifyTextChanged.Execute(InText, OutErrorMessage)))
+			if (!IsTextAllowed(InText, OutErrorMessage))
 			{
 				CancelEditMode();
 
@@ -427,6 +487,15 @@ void SInlineEditableTextBlock::SetEditableText( const TAttribute< FText >& InNew
 				 TextBox->SetText( Text );
 }
 
+FText SInlineEditableTextBlock::GetEditableText() const
+{
+	return 
+#if WITH_FANCY_TEXT
+	bIsMultiLine ? MultiLineTextBox->GetText() :
+#endif //WITH_FANCY_TEXT
+				 TextBox->GetText();
+}
+
 void SInlineEditableTextBlock::SetTextBoxError( const FText& ErrorText )
 {
 #if WITH_FANCY_TEXT
@@ -434,3 +503,12 @@ void SInlineEditableTextBlock::SetTextBoxError( const FText& ErrorText )
 #endif //WITH_FANCY_TEXT
 				 TextBox->SetError( ErrorText );
 }
+
+bool SInlineEditableTextBlock::HasTextBoxError() const
+{
+	return 
+#if WITH_FANCY_TEXT
+	bIsMultiLine ? MultiLineTextBox->HasError() :
+#endif //WITH_FANCY_TEXT
+				 TextBox->HasError();
+}
diff --git a/Engine/Source/Runtime/Slate/Public/Widgets/Input/SComboButton.h b/Engine/Source/Runtime/Slate/Public/Widgets/Input/SComboButton.h
--- a/Engine/Source/Runtime/Slate/Public/Widgets/Input/SComboButton.h
+++ b/Engine/Source/Runtime/Slate/Public/Widgets/Input/SComboButton.h
@@ -18,6 +18,12 @@ DECLARE_DELEGATE(FOnComboBoxOpened)
 class SButton;
 class SImage;
 
+template <>
+struct TWidgetTypeTraits<class SComboButton>
+{
+	static constexpr bool SupportsInvalidation() { return true; }
+};
+
 /**
  * A button that, when clicked, brings up a popup.
  */
diff --git a/Engine/Source/Runtime/Slate/Public/Widgets/Input/SMenuAnchor.h b/Engine/Source/Runtime/Slate/Public/Widgets/Input/SMenuAnchor.h
--- a/Engine/Source/Runtime/Slate/Public/Widgets/Input/SMenuAnchor.h
+++ b/Engine/Source/Runtime/Slate/Public/Widgets/Input/SMenuAnchor.h
@@ -23,6 +23,11 @@ DECLARE_DELEGATE_OneParam(FOnIsOpenChanged, bool)
 
 
 
+template <>
+struct TWidgetTypeTraits<class SMenuAnchor>
+{
+	static constexpr bool SupportsInvalidation() { return true; }
+};
 
 /**
  * A PopupAnchor summons a Popup relative to its content.
@@ -30,6 +35,7 @@ DECLARE_DELEGATE_OneParam(FOnIsOpenChanged, bool)
  */
 class SMenuAnchor : public SPanel, public IMenuHost
 {
+	SLATE_DECLARE_WIDGET_API(SMenuAnchor, SPanel, SLATE_API)
 public:
 	SLATE_BEGIN_ARGS( SMenuAnchor )
 		: _Content()
@@ -153,6 +159,11 @@ protected:
 
 	static SLATE_API TArray<TWeakPtr<IMenu>> OpenApplicationMenus;
 
+	TSlateAttributeRef<EMenuPlacement> GetPlacementAttribute() const
+	{
+		return TSlateAttributeRef<EMenuPlacement>{SharedThis(this), PlacementAttribute};
+	}
+
 protected:
 	struct FPopupPlacement
 	{
@@ -195,8 +206,11 @@ protected:
 	/** Callback invoked when the popup is being opened/closed */
 	FOnIsOpenChanged OnMenuOpenChanged;
 
+#if WITH_EDITORONLY_DATA
 	/** How should the popup be placed relative to the anchor. */
-	TAttribute<EMenuPlacement> Placement;
+	UE_DEPRECATED(5.8, "Use SetMenuPlacement / GetPlacementAttribute")
+	TSlateDeprecatedTAttribute<EMenuPlacement> Placement;
+#endif
 
 	/** Should the menu anchor fit inside the window? */
 	bool bFitInWindow;
@@ -238,4 +252,8 @@ protected:
 
 	/** The currently arranged children in the menu anchor.  Changes as the opened/closed state of the widget changes. */
 	TPanelChildren<FBasicLayoutWidgetSlot> Children;
+
+private:
+	/** How should the popup be placed relative to the anchor. */
+	TSlateAttribute<EMenuPlacement> PlacementAttribute;
 };
diff --git a/Engine/Source/Runtime/Slate/Public/Widgets/Input/SMultiLineEditableTextBox.h b/Engine/Source/Runtime/Slate/Public/Widgets/Input/SMultiLineEditableTextBox.h
--- a/Engine/Source/Runtime/Slate/Public/Widgets/Input/SMultiLineEditableTextBox.h
+++ b/Engine/Source/Runtime/Slate/Public/Widgets/Input/SMultiLineEditableTextBox.h
@@ -400,6 +400,7 @@ public:
 	 */
 	SLATE_API void SetError( const FText& InError );
 	SLATE_API void SetError( const FString& InError );
+	SLATE_API bool HasError() const;
 
 	// SWidget overrides
 	SLATE_API virtual bool SupportsKeyboardFocus() const override;
diff --git a/Engine/Source/Runtime/Slate/Public/Widgets/Text/SInlineEditableTextBlock.h b/Engine/Source/Runtime/Slate/Public/Widgets/Text/SInlineEditableTextBlock.h
--- a/Engine/Source/Runtime/Slate/Public/Widgets/Text/SInlineEditableTextBlock.h
+++ b/Engine/Source/Runtime/Slate/Public/Widgets/Text/SInlineEditableTextBlock.h
@@ -23,11 +23,19 @@ class SHorizontalBox;
 class SMultiLineEditableTextBox;
 class STextBlock;
 
+template <>
+struct TWidgetTypeTraits<class SInlineEditableTextBlock>
+{
+	static constexpr bool SupportsInvalidation() { return true; }
+};
+
 /**
  * Slate's InlineEditableTextBlock's are double selectable to go from a STextBlock to become SEditableTextBox.
  */
 class SInlineEditableTextBlock: public SCompoundWidget
 {
+	SLATE_DECLARE_WIDGET_API(SInlineEditableTextBlock, SCompoundWidget, SLATE_API)
+
 	SLATE_BEGIN_ARGS( SInlineEditableTextBlock )
 		: _Text()
 		, _HintText()
@@ -129,6 +137,8 @@ class SInlineEditableTextBlock: public SCompoundWidget
 		SLATE_ARGUMENT(TOptional<ETextOverflowPolicy>, OverflowPolicy)
 	SLATE_END_ARGS()
 
+	SLATE_API SInlineEditableTextBlock();
+
 	SLATE_API ~SInlineEditableTextBlock();
 
 	/**
@@ -139,6 +149,8 @@ class SInlineEditableTextBlock: public SCompoundWidget
 	 */
 	SLATE_API void Construct( const FArguments& InArgs );
 
+	SLATE_API virtual int32 OnPaint(const FPaintArgs& Args, const FGeometry& AllottedGeometry, const FSlateRect& MyCullingRect, FSlateWindowElementList& OutDrawElements, int32 LayerId, const FWidgetStyle& InWidgetStyle, bool bParentEnabled) const override;
+
 	/**
 	 * See SWidget::SupportsKeyboardFocus().
 	 *
@@ -154,8 +166,6 @@ class SInlineEditableTextBlock: public SCompoundWidget
 
 	SLATE_API virtual FReply OnKeyDown( const FGeometry& MyGeometry, const FKeyEvent& InKeyEvent ) override;
 
-	//virtual void Tick( const FGeometry& AllottedGeometry, const double InCurrentTime, const float InDeltaTime ) override;
-	 
 	/** Switches the widget to editing mode */
 	SLATE_API void EnterEditingMode();
 
@@ -172,7 +182,11 @@ class SInlineEditableTextBlock: public SCompoundWidget
 	 *
 	 * @param  InMaximumLength		Length to check for. Only values superior to 0 will be taken into account, otherwise no length check will occur.
 	 */
-	SLATE_API void SetMaximumLength(const TAttribute<int32>& InMaximumLength);
+	SLATE_API void SetMaximumLength(TAttribute<int32> InMaximumLength);
+	TSlateAttributeRef<int32> GetMaximumLengthAttribute() const
+	{
+		return TSlateAttributeRef<int32>{SharedThis(this), MaximumLengthAttribute};
+	}
 
 	SLATE_API void SetText( const TAttribute< FText >& InText );
 	SLATE_API void SetText( const FString& InText );
@@ -204,6 +218,9 @@ protected:
 	/** Callback for the text box's OnTextChanged event */
 	SLATE_API void OnTextChanged(const FText& InText);
 
+	/** Validate the given text by verifying its length and by calling OnVerifyTextChanged if bound. */
+	SLATE_API bool IsTextAllowed(const FText& InText, FText& OutErrorMessage) const;
+
 	/** Callback when the text box is committed, switches back to label mode. */ 
 	SLATE_API void OnTextBoxCommitted(const FText& InText, ETextCommit::Type InCommitType);
 
@@ -248,8 +265,11 @@ protected:
 	/** Attribute to look up if the widget is read-only */
 	TAttribute< bool > bIsReadOnly;
 
+#if WITH_EDITORONLY_DATA
 	/** Maximum text length that will be accepted by the widget */
-	TAttribute<int32> MaximumLength;
+	UE_DEPRECATED(5.8, "Use SetMaximumLength / GetMaximumLengthAttribute")
+	TSlateDeprecatedTAttribute<int32> MaximumLength;
+#endif
 
 	/** Widget to focus when we finish editing */
 	TWeakPtr<SWidget> WidgetToFocus;
@@ -262,8 +282,11 @@ private:
 	/** Set editable text */
 	void SetEditableText( const TAttribute< FText >& InNewText );
 
+	FText GetEditableText() const;
+
 	/** Set textbox error */
 	void SetTextBoxError( const FText& ErrorText );
+	bool HasTextBoxError() const;
 
 	/** Active timer to trigger entry into edit mode after a delay */
 	EActiveTimerReturnType TriggerEditMode(double InCurrentTime, float InDeltaTime);
@@ -271,6 +294,11 @@ private:
 	/** The handle to the active timer */
 	TWeakPtr<FActiveTimerHandle> ActiveTimerHandle;
 
+	/** Maximum text length that will be accepted by the widget */
+	TSlateAttribute<int32> MaximumLengthAttribute;
+
+	mutable bool bNeedsMaximumLengthChecked = false;
+
 protected:
 	/** When selection of widget is managed by another widget, this delays the "double select" from 
 	occurring immediately, offering a chance for double clicking to take action. */
diff --git a/Engine/Source/Runtime/SlateCore/Public/FastUpdate/SlateInvalidationRoot.h b/Engine/Source/Runtime/SlateCore/Public/FastUpdate/SlateInvalidationRoot.h
--- a/Engine/Source/Runtime/SlateCore/Public/FastUpdate/SlateInvalidationRoot.h
+++ b/Engine/Source/Runtime/SlateCore/Public/FastUpdate/SlateInvalidationRoot.h
@@ -119,6 +119,9 @@ public:
 	/** @return the cached draw elements for this window and its widget hierarchy */
 	SLATECORE_API FSlateInvalidationResult PaintInvalidationRoot(const FSlateInvalidationContext& Context);
 
+	/** @return if the InvalidationRoot is currently updating the attributes. */
+	bool IsProcessingAttributeUpdate() const { return bProcessingAttributeUpdate; }
+
 	void OnWidgetDestroyed(const SWidget* Widget);
 
 	SLATECORE_API void Advanced_ResetInvalidation(bool bClearResourcesImmediately);
