Slate Dynamic Invalidation - WidgetSwitcher

diff --git a/Engine/Source/Runtime/Slate/Private/Widgets/Layout/SWidgetSwitcher.cpp b/Engine/Source/Runtime/Slate/Private/Widgets/Layout/SWidgetSwitcher.cpp
--- a/Engine/Source/Runtime/Slate/Private/Widgets/Layout/SWidgetSwitcher.cpp
+++ b/Engine/Source/Runtime/Slate/Private/Widgets/Layout/SWidgetSwitcher.cpp
@@ -3,14 +3,35 @@
 #include "Widgets/Layout/SWidgetSwitcher.h"
 #include "Layout/LayoutUtils.h"
 #include "Types/ReflectionMetadata.h"
+#include "FastUpdate/SlateInvalidationRoot.h"
 
 #if WITH_ACCESSIBILITY
 #include "Framework/Application/SlateApplication.h"
 #include "Widgets/Accessibility/SlateAccessibleMessageHandler.h"
 #endif
 
+SLATE_IMPLEMENT_WIDGET(SWidgetSwitcher)
+void SWidgetSwitcher::PrivateRegisterAttributes(FSlateAttributeInitializer& AttributeInitializer)
+{
+	// This attribute should invalidate ChildOrder but we cannot invalidate ChildOrder from an attribute.
+	// The ChildOrder invalidation will be handled in CustomPrepass if it cannot be processed now
+	SLATE_ADD_MEMBER_ATTRIBUTE_DEFINITION(AttributeInitializer, WidgetIndex, EInvalidateWidgetReason::Prepass)
+		.OnValueChanged(FSlateAttributeDescriptor::FAttributeValueChangedDelegate::CreateLambda([](SWidget& Widget)
+		{
+			const bool bIsProcessingAttributeUpdate = Widget.GetProxyHandle().IsValid(Widget) &&
+				Widget.GetProxyHandle().GetInvalidationRootHandle().GetInvalidationRoot() &&
+				Widget.GetProxyHandle().GetInvalidationRootHandle().GetInvalidationRoot()->IsProcessingAttributeUpdate();
+
+			// If we are not currently updating attributes, we can safely process the widget index change
+			if (!bIsProcessingAttributeUpdate)
+			{
+				static_cast<SWidgetSwitcher&>(Widget).SetActiveWidgetIndexInternal();
+			}
+		}));
+}
+
 SWidgetSwitcher::SWidgetSwitcher()
-	: WidgetIndex()
+	: WidgetIndex(*this, 0)
 	, AllChildren(this)
 	, OneDynamicChild(this, &AllChildren, &WidgetIndex)
 { 
@@ -24,6 +45,12 @@ SWidgetSwitcher::FSlot::FSlotArguments SWidgetSwitcher::Slot()
 	return FSlot::FSlotArguments(MakeUnique<FSlot>());
 }
 
+bool SWidgetSwitcher::CustomPrepass(float LayoutScaleMultiplier)
+{
+	SetActiveWidgetIndexInternal();
+	return true;
+}
+
 SWidgetSwitcher::FScopedWidgetSlotArguments SWidgetSwitcher::AddSlot(int32 SlotIndex)
 {
 	TWeakPtr<SWidgetSwitcher> WeakSwitcher = SharedThis(this);
@@ -45,9 +72,9 @@ SWidgetSwitcher::FScopedWidgetSlotArguments SWidgetSwitcher::AddSlot(int32 SlotI
 				if (TSharedPtr<SWidgetSwitcher> Switcher = WeakSwitcher.Pin())
 				{
 					// adjust the active WidgetIndex based on this slot change
-					if (!Switcher->WidgetIndex.IsBound() && Switcher->WidgetIndex.Get() >= SlotIndex)
+					if (!Switcher->WidgetIndex.IsBound(*Switcher.Get()) && Switcher->WidgetIndex.Get() >= SlotIndex)
 					{
-						Switcher->WidgetIndex = Switcher->WidgetIndex.Get() + 1;
+						Switcher->SetActiveWidgetIndex(Switcher->WidgetIndex.Get() + 1);
 					}
 					Switcher->OnSlotAdded(SlotIndex);
 				}
@@ -59,7 +86,10 @@ void SWidgetSwitcher::Construct( const FArguments& InArgs )
 {
 	AllChildren.AddSlots(MoveTemp(const_cast<TArray<FSlot::FSlotArguments>&>(InArgs._Slots)));
 
-	WidgetIndex = InArgs._WidgetIndex;
+	PreviousWidgetIndex = INDEX_NONE;
+	WidgetIndex.Assign(*this, InArgs._WidgetIndex);
+	SetCanTick(false);
+	bHasCustomPrepass = true;
 }
 
 TSharedPtr<SWidget> SWidgetSwitcher::GetActiveWidget( ) const
@@ -106,22 +136,21 @@ int32 SWidgetSwitcher::RemoveSlot( TSharedRef<SWidget> WidgetToRemove )
 		if (AllChildren[SlotIndex].GetWidget() == WidgetToRemove)
 		{
 			// adjust the active WidgetIndex based on this slot change
-			bool bWasActiveSlot = false;
-			if (!WidgetIndex.IsBound())
-			{
-				int32 ActiveWidgetIndex = WidgetIndex.Get();
-
-				if (SlotIndex == ActiveWidgetIndex)
-				{
-					bWasActiveSlot = true;
-					Invalidate(EInvalidateWidget::ChildOrder);
-				}
+			const int32 ActiveWidgetIndex = WidgetIndex.Get();
+			const bool bWasActiveSlot = (SlotIndex == ActiveWidgetIndex);
 
+			if (!WidgetIndex.IsBound(*this))
+			{
 				if (ActiveWidgetIndex > 0 && ActiveWidgetIndex >= SlotIndex)
 				{
-					WidgetIndex = ActiveWidgetIndex - 1;
+					WidgetIndex.Assign(*this, ActiveWidgetIndex - 1);
 				}
 			}
+			else
+			{
+				// We want to give a chance for the attribute to update its index now that we have removed a child
+				WidgetIndex.UpdateNow(*this);
+			}
 
 			AllChildren.RemoveAt(SlotIndex);
 			OnSlotRemoved(SlotIndex, WidgetToRemove, bWasActiveSlot);
@@ -134,35 +163,7 @@ int32 SWidgetSwitcher::RemoveSlot( TSharedRef<SWidget> WidgetToRemove )
 
 void SWidgetSwitcher::SetActiveWidgetIndex( int32 Index )
 {
-	const int32 OldIndex = WidgetIndex.Get();
-
-	if (OldIndex != Index)
-	{
-		Invalidate(EInvalidateWidget::ChildOrder);
-
-		const FSlot* ActiveSlot = GetActiveSlot();
-
-		// Active slot can be null if the widget switcher was initialized to an invalid index.
-		if (ActiveSlot)
-		{
-			SWidget& OldActiveWidget = ActiveSlot->GetWidget().Get();
-			InvalidateChildRemovedFromTree(OldActiveWidget);
-
-#if WITH_SLATE_DEBUGGING
-			UE_LOG(LogSlate, Verbose, TEXT("WidgetSwitcher ('%s') Active Slot Changed: %d(%s) -FROM- %d(%s)"), *FReflectionMetaData::GetWidgetDebugInfo(this), 
-				Index, *FReflectionMetaData::GetWidgetDebugInfo(GetWidget(Index).Get()),
-				OldIndex, *FReflectionMetaData::GetWidgetDebugInfo(OldActiveWidget));
-#endif
-		}
-		else
-		{
-#if WITH_SLATE_DEBUGGING
-			UE_LOG(LogSlate, Verbose, TEXT("WidgetSwitcher ('%s') Active Slot Changed: %d(%s)"), *FReflectionMetaData::GetWidgetDebugInfo(this), Index, *FReflectionMetaData::GetWidgetDebugInfo(GetWidget(Index).Get()));
-#endif
-		}
-
-		WidgetIndex = Index;
-	}
+	WidgetIndex.Assign(*this, Index);
 }
 
 bool SWidgetSwitcher::ValidatePathToChild(SWidget* InChild)
@@ -221,3 +222,38 @@ const SWidgetSwitcher::FSlot* SWidgetSwitcher::GetActiveSlot() const
 
 	return nullptr;
 }
+
+void SWidgetSwitcher::SetActiveWidgetIndexInternal()
+{
+	const int NewIndex = WidgetIndex.Get();
+
+	if (NewIndex == PreviousWidgetIndex)
+	{
+		return;
+	}
+
+	Invalidate(EInvalidateWidgetReason::ChildOrder);
+
+	const FSlot* ActiveSlot = (PreviousWidgetIndex >= 0 && PreviousWidgetIndex < AllChildren.Num()) ? &AllChildren[PreviousWidgetIndex] : nullptr;
+
+	// Active slot can be null if the widget switcher was initialized to an invalid index.
+	if (ActiveSlot)
+	{
+		SWidget& OldActiveWidget = ActiveSlot->GetWidget().Get();
+		InvalidateChildRemovedFromTree(OldActiveWidget);
+
+#if WITH_SLATE_DEBUGGING
+		UE_LOG(LogSlate, Verbose, TEXT("WidgetSwitcher ('%s') Active Slot Changed: %d(%s) -FROM- %d(%s)"), *FReflectionMetaData::GetWidgetDebugInfo(this), 
+			NewIndex, *FReflectionMetaData::GetWidgetDebugInfo(GetWidget(NewIndex).Get()),
+			PreviousWidgetIndex, *FReflectionMetaData::GetWidgetDebugInfo(OldActiveWidget));
+#endif
+	}
+	else
+	{
+#if WITH_SLATE_DEBUGGING
+		UE_LOG(LogSlate, Verbose, TEXT("WidgetSwitcher ('%s') Active Slot Changed: %d(%s)"), *FReflectionMetaData::GetWidgetDebugInfo(this), NewIndex, *FReflectionMetaData::GetWidgetDebugInfo(GetWidget(NewIndex).Get()));
+#endif
+	}
+
+	PreviousWidgetIndex = NewIndex;
+}
diff --git a/Engine/Source/Runtime/Slate/Public/Widgets/Layout/SWidgetSwitcher.h b/Engine/Source/Runtime/Slate/Public/Widgets/Layout/SWidgetSwitcher.h
--- a/Engine/Source/Runtime/Slate/Public/Widgets/Layout/SWidgetSwitcher.h
+++ b/Engine/Source/Runtime/Slate/Public/Widgets/Layout/SWidgetSwitcher.h
@@ -13,6 +13,12 @@
 
 class FArrangedChildren;
 
+template <>
+struct TWidgetTypeTraits<class SWidgetSwitcher>
+{
+	static constexpr bool SupportsInvalidation() { return true; }
+};
+
 /**
  * Implements a widget switcher.
  *
@@ -21,6 +27,7 @@ class FArrangedChildren;
 class SWidgetSwitcher
 	: public SPanel
 {
+	SLATE_DECLARE_WIDGET_API(SWidgetSwitcher, SPanel, SLATE_API)
 public:
 
 	using FSlot = FBasicLayoutWidgetSlot;
@@ -135,23 +142,22 @@ public:
 	static SLATE_API FSlot::FSlotArguments Slot();
 
 protected:
-
-	// SCompoundWidget interface
-
+	SLATE_API virtual bool CustomPrepass(float LayoutScaleMultiplier) override;
 	SLATE_API virtual void OnArrangeChildren( const FGeometry& AllottedGeometry, FArrangedChildren& ArrangedChildren ) const override;
 	SLATE_API virtual FVector2D ComputeDesiredSize( float ) const override;
 	SLATE_API virtual FChildren* GetChildren( ) override;
 	virtual void OnSlotAdded(int32 AddedIndex) {}
 	virtual void OnSlotRemoved(int32 RemovedIndex, TSharedRef<SWidget> RemovedWidget, bool bWasActiveSlot) {}
-	virtual bool ComputeVolatility() const override { return WidgetIndex.IsBound(); }
 	SLATE_API const FSlot* GetActiveSlot() const;
 
 	TPanelChildren<FSlot>& GetTypedChildren() { return AllChildren; }
 
 private:
+	void SetActiveWidgetIndexInternal();
 
 	/** Holds the desired widget index */
-	TAttribute<int32> WidgetIndex;
+	TSlateAttribute<int32> WidgetIndex;
+	int32 PreviousWidgetIndex = INDEX_NONE;
 
 	/*
 	 * Holds the collection of all child widget, however the only one with a valid parent pointer
@@ -160,7 +166,7 @@ private:
 	TPanelChildren<FSlot> AllChildren;
 
 	/** Required to implement GetChildren() in a way that can dynamically return the currently active child. */
-	TOneDynamicChild<FSlot> OneDynamicChild;
+	TOneDynamicChildBase<FSlot, TSlateAttribute<int32>> OneDynamicChild;
 
 #if WITH_ACCESSIBILITY
 	/** Used to detect when WidgetIndex changes while bound. */
diff --git a/Engine/Source/Runtime/SlateCore/Public/Layout/Children.h b/Engine/Source/Runtime/SlateCore/Public/Layout/Children.h
--- a/Engine/Source/Runtime/SlateCore/Public/Layout/Children.h
+++ b/Engine/Source/Runtime/SlateCore/Public/Layout/Children.h
@@ -1023,11 +1023,11 @@ private:
 
 
 /** Required to implement GetChildren() in a way that can dynamically return the currently active child. */
-template<typename SlotType>
-class TOneDynamicChild final : public FChildren
+template<typename SlotType, typename Attr>
+class TOneDynamicChildBase final : public FChildren
 {
 public:
-	TOneDynamicChild(SWidget* InOwner, TPanelChildren<SlotType>* InAllChildren, const TAttribute<int32>* InWidgetIndex)
+	TOneDynamicChildBase(SWidget* InOwner, TPanelChildren<SlotType>* InAllChildren, const Attr* InWidgetIndex)
 		: FChildren(InOwner)
 		, AllChildren(InAllChildren)
 		, WidgetIndex(InWidgetIndex)
@@ -1036,12 +1036,12 @@ public:
 		check(WidgetIndex);
 	}
 
-	TOneDynamicChild(SWidget* InOwner, TPanelChildren<SlotType>* InAllChildren, std::nullptr_t) = delete;
-	TOneDynamicChild(SWidget* InOwner, std::nullptr_t, const TAttribute<int32>* InWidgetIndex) = delete;
-	TOneDynamicChild(SWidget* InOwner, std::nullptr_t, std::nullptr_t) = delete;
-	TOneDynamicChild(std::nullptr_t, TPanelChildren<SlotType>* InAllChildren, std::nullptr_t) = delete;
-	TOneDynamicChild(std::nullptr_t, std::nullptr_t, const TAttribute<int32>* InWidgetIndex) = delete;
-	TOneDynamicChild(std::nullptr_t, std::nullptr_t, std::nullptr_t) = delete;
+	TOneDynamicChildBase(SWidget* InOwner, TPanelChildren<SlotType>* InAllChildren, std::nullptr_t) = delete;
+	TOneDynamicChildBase(SWidget* InOwner, std::nullptr_t, const Attr* InWidgetIndex) = delete;
+	TOneDynamicChildBase(SWidget* InOwner, std::nullptr_t, std::nullptr_t) = delete;
+	TOneDynamicChildBase(std::nullptr_t, TPanelChildren<SlotType>* InAllChildren, std::nullptr_t) = delete;
+	TOneDynamicChildBase(std::nullptr_t, std::nullptr_t, const Attr* InWidgetIndex) = delete;
+	TOneDynamicChildBase(std::nullptr_t, std::nullptr_t, std::nullptr_t) = delete;
 
 	virtual int32 Num() const override { return AllChildren->Num() > 0 ? 1 : 0; }
 
@@ -1076,5 +1076,8 @@ private:
 	}
 
 	TPanelChildren<SlotType>* AllChildren;
-	const TAttribute<int32>* WidgetIndex;
+	const Attr* WidgetIndex;
 };
+
+template <typename SlotType>
+using TOneDynamicChild = TOneDynamicChildBase<SlotType, TAttribute<int32>>;
diff --git a/Engine/Source/Runtime/SlateCore/Public/Layout/ChildrenBase.h b/Engine/Source/Runtime/SlateCore/Public/Layout/ChildrenBase.h
--- a/Engine/Source/Runtime/SlateCore/Public/Layout/ChildrenBase.h
+++ b/Engine/Source/Runtime/SlateCore/Public/Layout/ChildrenBase.h
@@ -118,8 +118,8 @@ public:
 
 protected:
 	friend class FCombinedChildren;
-	template<typename T>
-	friend class TOneDynamicChild;
+	template<typename SlotType, typename Attr>
+	friend class TOneDynamicChildBase;
 
 	enum ECopyConstruct { CopyConstruct };
 	enum ERefConstruct { ReferenceConstruct };
