UAF Browser Editor Framework changes
- Added methods to TabManager to toggle / dismiss sidebar tabs
- Added 'StatusbarTab' as a valid tab state.
- Add ability to modify top bar of an asset viewer, to add custom widgets / etc.

review-45330403

#tests Editor

diff --git a/Engine/Source/Editor/ContentBrowser/Private/SAssetPicker.cpp b/Engine/Source/Editor/ContentBrowser/Private/SAssetPicker.cpp
--- a/Engine/Source/Editor/ContentBrowser/Private/SAssetPicker.cpp
+++ b/Engine/Source/Editor/ContentBrowser/Private/SAssetPicker.cpp
@@ -401,6 +401,7 @@ void SAssetPicker::Construct( const FArguments& InArgs )
 		]	
 	];
 
+	InArgs._AssetPickerConfig.OnExtendAssetPickerTopBar.ExecuteIfBound(HorizontalBox);
 
 	LoadSettings();
 
diff --git a/Engine/Source/Editor/ContentBrowser/Public/IContentBrowserSingleton.h b/Engine/Source/Editor/ContentBrowser/Public/IContentBrowserSingleton.h
--- a/Engine/Source/Editor/ContentBrowser/Public/IContentBrowserSingleton.h
+++ b/Engine/Source/Editor/ContentBrowser/Public/IContentBrowserSingleton.h
@@ -334,6 +334,16 @@ struct FAssetPickerConfig
 
 	/** Can be used to freely change the Add Filter menu. This is executed on the dynamically instanced menu, not the registered menu. */
 	TDelegate<void(UToolMenu*)> OnExtendAddFilterMenu;
+
+	/**
+	 * Can be used to extend the asset picker top bar via slot manipulation. Current default slots are as follows:
+	 * 
+	 * 0. Add Filter Menu
+	 * 1. SearchBar
+	 * 2. Show Developer Content Checkbox
+	 * 3. Asset View Options
+	 */
+	TDelegate<void(TSharedRef<SHorizontalBox>)> OnExtendAssetPickerTopBar;
 	
 	/** Override the default filter context menu layout */
 	EAssetTypeCategories::Type DefaultFilterMenuExpansion;
diff --git a/Engine/Source/Editor/UnrealEd/Public/EdMode.h b/Engine/Source/Editor/UnrealEd/Public/EdMode.h
--- a/Engine/Source/Editor/UnrealEd/Public/EdMode.h
+++ b/Engine/Source/Editor/UnrealEd/Public/EdMode.h
@@ -36,6 +36,10 @@ struct FGizmoState;
 
 /**
  * Base class for all editor modes.
+ * 
+ * Modes can be switched. Different modes may have different widgets / editing functionality / layouts.
+ * Modes can also be compatible, so multiple modes can concurrently affect a single asset editor. 
+ * Modes defined in separate modules, say from a plugin, can modify existing asset editors in a different module.
  */
 class FEdMode : public TSharedFromThis<FEdMode>, public FGCObject, public FEditorCommonDrawHelper, public FLegacyEdModeWidgetHelper
 {
diff --git a/Engine/Source/Editor/UnrealEd/Public/Tools/UEdMode.h b/Engine/Source/Editor/UnrealEd/Public/Tools/UEdMode.h
--- a/Engine/Source/Editor/UnrealEd/Public/Tools/UEdMode.h
+++ b/Engine/Source/Editor/UnrealEd/Public/Tools/UEdMode.h
@@ -66,6 +66,10 @@ enum class EToolsContextScope
 
 /**
  * Base class for all editor modes.
+ * 
+ * Modes can be switched. Different modes may have different widgets / editing functionality / layouts.
+ * Modes can also be compatible, so multiple modes can concurrently affect a single asset editor. 
+ * Modes defined in separate modules, say from a plugin, can modify existing asset editors in a different module.
  */
 UCLASS(Abstract, MinimalAPI)
 class UEdMode : public UObject
diff --git a/Engine/Source/Runtime/Slate/Private/Framework/Docking/SDockingArea.cpp b/Engine/Source/Runtime/Slate/Private/Framework/Docking/SDockingArea.cpp
--- a/Engine/Source/Runtime/Slate/Private/Framework/Docking/SDockingArea.cpp
+++ b/Engine/Source/Runtime/Slate/Private/Framework/Docking/SDockingArea.cpp
@@ -522,9 +522,19 @@ bool SDockingArea::RemoveTabFromSidebar(TSharedRef<SDockTab> Tab)
 	return LeftSidebar->RemoveTab(Tab) || RightSidebar->RemoveTab(Tab);
 }
 
-bool SDockingArea::TryOpenSidebarDrawer(TSharedRef<SDockTab> TabToOpen) const
+bool SDockingArea::TryOpenSidebarDrawer(TSharedRef<SDockTab> TabToModify) const
 {
-	return LeftSidebar->TryOpenSidebarDrawer(TabToOpen) || RightSidebar->TryOpenSidebarDrawer(TabToOpen);
+	return LeftSidebar->TryOpenSidebarDrawer(TabToModify) || RightSidebar->TryOpenSidebarDrawer(TabToModify);
+}
+
+bool SDockingArea::TryDismissSidebarDrawer(TSharedRef<SDockTab> TabToModify) const
+{
+	return LeftSidebar->TryDismissSidebarDrawer(TabToModify) || RightSidebar->TryDismissSidebarDrawer(TabToModify);
+}
+
+bool SDockingArea::TryToggleSidebarDrawer(TSharedRef<SDockTab> TabToModify) const
+{
+	return LeftSidebar->TryToggleSidebarDrawer(TabToModify) || RightSidebar->TryToggleSidebarDrawer(TabToModify);
 }
 
 void SDockingArea::AddSidebarTabsFromRestoredLayout(const FSidebarTabLists& SidebarTabs)
diff --git a/Engine/Source/Runtime/Slate/Private/Framework/Docking/SDockingArea.h b/Engine/Source/Runtime/Slate/Private/Framework/Docking/SDockingArea.h
--- a/Engine/Source/Runtime/Slate/Private/Framework/Docking/SDockingArea.h
+++ b/Engine/Source/Runtime/Slate/Private/Framework/Docking/SDockingArea.h
@@ -124,7 +124,21 @@ public:
 	 *
 	 * @true if the drawer was opened, false if the tab is not in a drawer
 	 */
-	SLATE_API bool TryOpenSidebarDrawer(TSharedRef<SDockTab> TabToOpen) const;
+	SLATE_API bool TryOpenSidebarDrawer(TSharedRef<SDockTab> TabToModify) const;
+
+	/**
+	 * Attempts to dismiss a sidebar drawer that might be open for the tab
+	 *
+	 * @true if the drawer was dismissed, false if the tab is pinned or not in a drawer
+	 */
+	SLATE_API bool TryDismissSidebarDrawer(TSharedRef<SDockTab> TabToModify) const;
+
+	/**
+	 * Attempts to toggle a sidebar drawer that may own the given tab
+	 *
+	 * @true if the drawer toggled, false if not found or some other setting prevented toggling
+	 */
+	SLATE_API bool TryToggleSidebarDrawer(TSharedRef<SDockTab> TabToModify) const;
 
 	/**
 	 * Adds all tabs back to a sidebar that were saved in a sidebar from a previous session
diff --git a/Engine/Source/Runtime/Slate/Private/Framework/Docking/STabSidebar.cpp b/Engine/Source/Runtime/Slate/Private/Framework/Docking/STabSidebar.cpp
--- a/Engine/Source/Runtime/Slate/Private/Framework/Docking/STabSidebar.cpp
+++ b/Engine/Source/Runtime/Slate/Private/Framework/Docking/STabSidebar.cpp
@@ -575,21 +575,59 @@ bool STabSidebar::TryOpenSidebarDrawer(TSharedRef<SDockTab> ForTab)
 	return false;
 }
 
-void STabSidebar::OnTabDrawerButtonPressed(TSharedRef<SDockTab> ForTab)
+bool STabSidebar::TryDismissSidebarDrawer(TSharedRef<SDockTab> ForTab)
 {
-	if (ForTab->IsActive())
+	int32 FoundIndex = Tabs.IndexOfByPredicate(
+		[ForTab](auto TabPair)
+		{
+			return TabPair.Key == ForTab;
+		});
+
+	if (FoundIndex != INDEX_NONE)
 	{
-		// When clicking on the button of an active (but unpinned) tab, close that tab drawer
 		if (!IsTabPinned(ForTab))
 		{
 			CloseDrawerInternal(ForTab);
+			return true;
 		}
 	}
-	else
+
+	return false;
+}
+
+bool STabSidebar::TryToggleSidebarDrawer(TSharedRef<SDockTab> ForTab)
+{
+	int32 FoundIndex = Tabs.IndexOfByPredicate(
+		[ForTab](auto TabPair)
+		{
+			return TabPair.Key == ForTab;
+		});
+
+	if (FoundIndex != INDEX_NONE)
 	{
-		// Otherwise clicking on an inactive tab should open the drawer
-		OpenDrawerInternal(ForTab, /*bAnimateOpen=*/ true);
+		if (ForTab->IsActive())
+		{
+			// When clicking on the button of an active (but unpinned) tab, close that tab drawer
+			if (!IsTabPinned(ForTab))
+			{
+				CloseDrawerInternal(ForTab);
+				return true;
+			}
+		}
+		else
+		{
+			// Otherwise clicking on an inactive tab should open the drawer
+			OpenDrawerInternal(ForTab, /*bAnimateOpen=*/ true);
+			return true;
+		}
 	}
+
+	return false;
+}
+
+void STabSidebar::OnTabDrawerButtonPressed(TSharedRef<SDockTab> ForTab)
+{
+	TryToggleSidebarDrawer(ForTab);
 }
 
 void STabSidebar::OnTabDrawerPinButtonToggled(TSharedRef<SDockTab> ForTab, bool bIsPinned)
diff --git a/Engine/Source/Runtime/Slate/Private/Framework/Docking/STabSidebar.h b/Engine/Source/Runtime/Slate/Private/Framework/Docking/STabSidebar.h
--- a/Engine/Source/Runtime/Slate/Private/Framework/Docking/STabSidebar.h
+++ b/Engine/Source/Runtime/Slate/Private/Framework/Docking/STabSidebar.h
@@ -79,6 +79,20 @@ public:
 	 * @return true if the passed in tab is contained in this sidebar and was opened
 	 */
 	bool TryOpenSidebarDrawer(TSharedRef<SDockTab> ForTab);
+
+	/**
+	 * Attempt to dismiss a drawer in the sidebar for a specified tab
+	 *
+	 * @return true if the passed in tab is contained in this sidebar and was dismissed
+	 */
+	bool TryDismissSidebarDrawer(TSharedRef<SDockTab> ForTab);
+
+	/**
+	 * Attempt to toggle a drawer in the sidebar for a specified tab
+	 *
+	 * @return true if the passed in tab is contained in this sidebar and was toggled
+	 */
+	bool TryToggleSidebarDrawer(TSharedRef<SDockTab> ForTab);
 private:
 	void OnTabDrawerButtonPressed(TSharedRef<SDockTab> ForTab);
 	void OnTabDrawerPinButtonToggled(TSharedRef<SDockTab> ForTab, bool bIsPinned);
diff --git a/Engine/Source/Runtime/Slate/Private/Framework/Docking/TabManager.cpp b/Engine/Source/Runtime/Slate/Private/Framework/Docking/TabManager.cpp
--- a/Engine/Source/Runtime/Slate/Private/Framework/Docking/TabManager.cpp
+++ b/Engine/Source/Runtime/Slate/Private/Framework/Docking/TabManager.cpp
@@ -1555,6 +1555,50 @@ TSharedPtr<SDockTab> FTabManager::TryInvokeTab(const FTabId& TabId, bool bInvoke
 	return NewTab;
 }
 
+TSharedPtr<SDockTab> FTabManager::TryToggleSidebarTab(const FTabId& TabId)
+{
+	if (!InvokeTab_CanInvokeTab(TabId))
+	{
+		return {};
+	}
+
+	TSharedPtr<SDockTab> ExistingTab = InvokeTab_FindOrReuseExistingTab(TabId);
+
+	bool bToggled = false;
+	if (ExistingTab.IsValid())
+	{
+		const TSharedPtr<SDockingArea> DockingArea = ExistingTab->GetDockArea();
+		if (DockingArea.IsValid())
+		{
+			bToggled = DockingArea->TryToggleSidebarDrawer(ExistingTab.ToSharedRef());
+		}
+	}
+
+	return bToggled ? ExistingTab : TSharedPtr<SDockTab>();
+}
+
+TSharedPtr<SDockTab> FTabManager::TryDismissSidebarTab(const FTabId& TabId)
+{
+	if (!InvokeTab_CanInvokeTab(TabId))
+	{
+		return {};
+	}
+
+	TSharedPtr<SDockTab> ExistingTab = InvokeTab_FindOrReuseExistingTab(TabId);
+
+	bool bDismissed = false;
+	if (ExistingTab.IsValid())
+	{
+		const TSharedPtr<SDockingArea> DockingArea = ExistingTab->GetDockArea();
+		if (DockingArea.IsValid())
+		{
+			bDismissed = DockingArea->TryDismissSidebarDrawer(ExistingTab.ToSharedRef());
+		}
+	}
+
+	return bDismissed ? ExistingTab : TSharedPtr<SDockTab>();
+}
+
 bool FTabManager::InvokeTab_CanInvokeTab(const FTabId& TabId) const
 {
 	if (!IsAllowedTab(TabId))
diff --git a/Engine/Source/Runtime/Slate/Public/Framework/Docking/TabManager.h b/Engine/Source/Runtime/Slate/Public/Framework/Docking/TabManager.h
--- a/Engine/Source/Runtime/Slate/Public/Framework/Docking/TabManager.h
+++ b/Engine/Source/Runtime/Slate/Public/Framework/Docking/TabManager.h
@@ -396,11 +396,12 @@ namespace ETabState
 		OpenedTab = 0x1 << 0,
 		ClosedTab = 0x1 << 1,
 		SidebarTab = 0x1 << 2,
+		StatusbarTab = 0x1 << 3,
 
 		/**
 		 * InvalidTab refers to tabs that were not recognized by the Editor (e.g., LiveLink when its plugin its disabled).
 		 */
-		InvalidTab = 0x1 << 3
+		InvalidTab = 0x1 << 4
 	};
 }
 
@@ -475,7 +476,7 @@ class FTabManager : public TSharedFromThis<FTabManager>
 				, bPinnedInSidebar(false)
 				, bIsLockedInPlace(TOptional<bool>())
 			{
-				check(InTabState != ETabState::SidebarTab);
+				checkf(InTabState != ETabState::SidebarTab, TEXT("Cannot add sidebar tab without specifying sidebar location."));
 			}
 
 			FTab(const FTabId& InTabId, ETabState::Type InTabState, ESidebarLocation InSidebarLocation, float InSidebarSizeCoefficient, bool bInPinnedInSidebar)
@@ -486,7 +487,7 @@ class FTabManager : public TSharedFromThis<FTabManager>
 				, bPinnedInSidebar(bInPinnedInSidebar)
 				, bIsLockedInPlace(TOptional<bool>())
 			{
-				check(InTabState != ETabState::SidebarTab || InSidebarLocation != ESidebarLocation::None);
+				checkf(InTabState != ETabState::SidebarTab || InSidebarLocation != ESidebarLocation::None, TEXT("Cannot add sidebar tab without specifying sidebar location."));
 			}
 
 			bool operator==( const FTab& Other ) const
@@ -1049,6 +1050,22 @@ class FTabManager : public TSharedFromThis<FTabManager>
 		 */
 		SLATE_API virtual TSharedPtr<SDockTab> TryInvokeTab(const FTabId& TabId, bool bInvokeAsInactive = false);
 
+		/**
+		 * Try to toggle a sidebar tab. 
+		 *
+		 * @param TabId The tab identifier.
+		 * @return The toggled tab if toggling was sucessful.
+		 */
+		SLATE_API virtual TSharedPtr<SDockTab> TryToggleSidebarTab(const FTabId& TabId);
+
+		/**
+		 * Try to dissmiss a sidebar tab.
+		 *
+		 * @param TabId The tab identifier.
+		 * @return The toggled tab if dissmissing was successful
+		 */
+		SLATE_API virtual TSharedPtr<SDockTab> TryDismissSidebarTab(const FTabId& TabId);
+
 		/**
 		 * Finds the first instance of an existing tab with the given tab id.
 		 *
