Slate Dynamic Invalidation - Content Browser List Widgets

Making List Widgets and their related types invalidation compatible.
As a result of this change, some modules will need to include Slate and Slate Core as dependencies because Widget classes have been added, and both `SListView` and `STreeView` need to call `STableViewBase::RequestLayoutRefresh` when `MaxPinnedItemsAttribute` changes.

diff --git a/Engine/Plugins/Experimental/GeometryCollectionPlugin/Source/GeometryCollectionTracks/GeometryCollectionTracks.Build.cs b/Engine/Plugins/Experimental/GeometryCollectionPlugin/Source/GeometryCollectionTracks/GeometryCollectionTracks.Build.cs
--- a/Engine/Plugins/Experimental/GeometryCollectionPlugin/Source/GeometryCollectionTracks/GeometryCollectionTracks.Build.cs
+++ b/Engine/Plugins/Experimental/GeometryCollectionPlugin/Source/GeometryCollectionTracks/GeometryCollectionTracks.Build.cs
@@ -34,6 +34,14 @@ namespace UnrealBuildTool.Rules
             {
                 PublicIncludePathModuleNames.Add("GeometryCollectionSequencer");
                 DynamicallyLoadedModuleNames.Add("GeometryCollectionSequencer");
+                
+                PrivateDependencyModuleNames.AddRange(
+	                new string[]
+	                {
+		                "SlateCore",
+		                "Slate",
+	                }
+                );
             }
         }
     }
diff --git a/Engine/Plugins/Interchange/Editor/Source/Pipelines/Private/SInterchangeGraphInspectorWindow.cpp b/Engine/Plugins/Interchange/Editor/Source/Pipelines/Private/SInterchangeGraphInspectorWindow.cpp
--- a/Engine/Plugins/Interchange/Editor/Source/Pipelines/Private/SInterchangeGraphInspectorWindow.cpp
+++ b/Engine/Plugins/Interchange/Editor/Source/Pipelines/Private/SInterchangeGraphInspectorWindow.cpp
@@ -375,7 +375,7 @@ void SInterchangeGraphInspectorTreeView::SetSelectionImportState(bool MarkForImp
 
 void SInterchangeGraphInspectorTreeView::OnTreeViewSelectionChanged(UInterchangeBaseNode* Item, ESelectInfo::Type SelectionType)
 {
-	if (SelectionMode.Get() == ESelectionMode::None)
+	if (GetSelectionMode() == ESelectionMode::None)
 	{
 		return;
 	}
diff --git a/Engine/Plugins/MovieScene/MovieRenderPipeline/Source/MovieRenderPipelineRenderPasses/MovieRenderPipelineRenderPasses.Build.cs b/Engine/Plugins/MovieScene/MovieRenderPipeline/Source/MovieRenderPipelineRenderPasses/MovieRenderPipelineRenderPasses.Build.cs
--- a/Engine/Plugins/MovieScene/MovieRenderPipeline/Source/MovieRenderPipelineRenderPasses/MovieRenderPipelineRenderPasses.Build.cs
+++ b/Engine/Plugins/MovieScene/MovieRenderPipeline/Source/MovieRenderPipelineRenderPasses/MovieRenderPipelineRenderPasses.Build.cs
@@ -24,6 +24,7 @@ public class MovieRenderPipelineRenderPasses : ModuleRules
 				"CinematicCamera", // For metadata
 				"MovieRenderPipelineSettings", // For settings
 				"MovieScene",
+				"Slate",
 				"SlateCore",
 				"ImageCore",
 				"OpenColorIOWrapper",
diff --git a/Engine/Plugins/Runtime/GeometryCache/Source/GeometryCacheTracks/GeometryCacheTracks.Build.cs b/Engine/Plugins/Runtime/GeometryCache/Source/GeometryCacheTracks/GeometryCacheTracks.Build.cs
--- a/Engine/Plugins/Runtime/GeometryCache/Source/GeometryCacheTracks/GeometryCacheTracks.Build.cs
+++ b/Engine/Plugins/Runtime/GeometryCache/Source/GeometryCacheTracks/GeometryCacheTracks.Build.cs
@@ -34,6 +34,14 @@ namespace UnrealBuildTool.Rules
             {
                 PublicIncludePathModuleNames.Add("GeometryCacheSequencer");
                 DynamicallyLoadedModuleNames.Add("GeometryCacheSequencer");
+                
+                PrivateDependencyModuleNames.AddRange(
+	                new string[]
+	                {
+		                "SlateCore",
+		                "Slate",
+	                }
+                );
             }
 
 			CppCompileWarningSettings.UnsafeTypeCastWarningLevel = WarningLevel.Error;
diff --git a/Engine/Source/Editor/ContentBrowser/Private/AssetViewWidgets.cpp b/Engine/Source/Editor/ContentBrowser/Private/AssetViewWidgets.cpp
--- a/Engine/Source/Editor/ContentBrowser/Private/AssetViewWidgets.cpp
+++ b/Engine/Source/Editor/ContentBrowser/Private/AssetViewWidgets.cpp
@@ -3284,6 +3284,12 @@ float SAssetTileItem::SmallFontHeight(0);
 // SAssetColumnItem
 ///////////////////////////////
 
+template <>
+struct TWidgetTypeTraits<class SAssetColumnItemNameBox>
+{
+	static constexpr bool SupportsInvalidation() { return true; }
+};
+
 /** Custom box for the Name column of an asset */
 class SAssetColumnItemNameBox : public SCompoundWidget
 {
diff --git a/Engine/Source/Editor/ContentBrowser/Public/AssetViewWidgets.h b/Engine/Source/Editor/ContentBrowser/Public/AssetViewWidgets.h
--- a/Engine/Source/Editor/ContentBrowser/Public/AssetViewWidgets.h
+++ b/Engine/Source/Editor/ContentBrowser/Public/AssetViewWidgets.h
@@ -80,6 +80,13 @@ private:
 	static TSharedRef<SWidget> CreateListTileItemContents(T* const InTileOrListItem, const TSharedRef<SWidget>& InThumbnail, FName& OutItemShadowBorder, const FMargin& InFolderOverlayPadding, bool bShowAssetAccessSpecifier);
 };
 
+
+template <>
+struct TWidgetTypeTraits<class SAssetTileView>
+{
+	static constexpr bool SupportsInvalidation() { return true; }
+};
+
 /** The tile view mode of the asset view */
 class SAssetTileView : public STileView<TSharedPtr<FAssetViewItem>>
 {
@@ -88,6 +95,13 @@ public:
 	virtual void Tick( const FGeometry& AllottedGeometry, const double InCurrentTime, const float InDeltaTime ) override;
 };
 
+
+template <>
+struct TWidgetTypeTraits<class SAssetListView>
+{
+	static constexpr bool SupportsInvalidation() { return true; }
+};
+
 /** The list view mode of the asset view */
 class SAssetListView : public SListView<TSharedPtr<FAssetViewItem>>
 {
@@ -96,6 +110,13 @@ public:
 	virtual void Tick(const FGeometry& AllottedGeometry, const double InCurrentTime, const float InDeltaTime) override;
 };
 
+
+template <>
+struct TWidgetTypeTraits<class SAssetColumnView>
+{
+	static constexpr bool SupportsInvalidation() { return true; }
+};
+
 /** The columns view mode of the asset view */
 class SAssetColumnView : public SListView<TSharedPtr<FAssetViewItem>>
 {
@@ -104,6 +125,13 @@ public:
 	virtual void Tick(const FGeometry& AllottedGeometry, const double InCurrentTime, const float InDeltaTime) override;
 };
 
+
+template <>
+struct TWidgetTypeTraits<class SAssetViewItem>
+{
+	static constexpr bool SupportsInvalidation() { return true; }
+};
+
 /** A base class for all asset view items */
 class SAssetViewItem : public SCompoundWidget
 {
@@ -391,6 +419,12 @@ protected:
 };
 
 
+template <>
+struct TWidgetTypeTraits<class SAssetListItem>
+{
+	static constexpr bool SupportsInvalidation() { return true; }
+};
+
 /** An item in the asset list view */
 class SAssetListItem : public SAssetViewItem
 {
@@ -568,6 +602,13 @@ private:
 	TSharedPtr<STextBlock> PathTextWidget;
 };
 
+
+template <>
+struct TWidgetTypeTraits<class SAssetListViewRow>
+{
+	static constexpr bool SupportsInvalidation() { return true; }
+};
+
 class SAssetListViewRow : public SMultiColumnTableRow< TSharedPtr<FAssetViewItem> >
 {
 public:
@@ -667,6 +708,12 @@ public:
 	TSharedPtr<SAssetListItem> AssetListItem;
 };
 
+template <>
+struct TWidgetTypeTraits<class SAssetTileItem>
+{
+	static constexpr bool SupportsInvalidation() { return true; }
+};
+
 /** An item in the asset tile view */
 class SAssetTileItem : public SAssetViewItem
 {
@@ -875,6 +922,12 @@ private:
 	TAttribute<EThumbnailSize> CurrentThumbnailSize;
 };
 
+template <>
+struct TWidgetTypeTraits<class SAssetColumnItem>
+{
+	static constexpr bool SupportsInvalidation() { return true; }
+};
+
 /** An item in the asset column view */
 class SAssetColumnItem : public SAssetViewItem
 {
@@ -944,6 +997,13 @@ private:
 	TSharedPtr<STextBlock> PathText;
 };
 
+
+template <>
+struct TWidgetTypeTraits<class SAssetColumnViewRow>
+{
+	static constexpr bool SupportsInvalidation() { return true; }
+};
+
 class SAssetColumnViewRow : public SMultiColumnTableRow< TSharedPtr<FAssetViewItem> >
 {
 public:
diff --git a/Engine/Source/Editor/SkeletonEditor/Private/SSkeletonTree.cpp b/Engine/Source/Editor/SkeletonEditor/Private/SSkeletonTree.cpp
--- a/Engine/Source/Editor/SkeletonEditor/Private/SSkeletonTree.cpp
+++ b/Engine/Source/Editor/SkeletonEditor/Private/SSkeletonTree.cpp
@@ -90,7 +90,7 @@ public:
 
 	virtual void Private_SelectRangeFromCurrentTo( ItemType InRangeSelectionEnd ) override
 	{
-		if ( this->SelectionMode.Get() == ESelectionMode::None )
+		if ( this->GetSelectionMode() == ESelectionMode::None )
 		{
 			return;
 		}
diff --git a/Engine/Source/Runtime/Slate/Private/Widgets/Views/SListPanel.cpp b/Engine/Source/Runtime/Slate/Private/Widgets/Views/SListPanel.cpp
--- a/Engine/Source/Runtime/Slate/Private/Widgets/Views/SListPanel.cpp
+++ b/Engine/Source/Runtime/Slate/Private/Widgets/Views/SListPanel.cpp
@@ -10,6 +10,9 @@ static const float FloatingPointPrecisionOffset = 0.001f;
 SLATE_IMPLEMENT_WIDGET(SListPanel)
 void SListPanel::PrivateRegisterAttributes(FSlateAttributeInitializer& AttributeInitializer)
 {
+	FSlateWidgetSlotAttributeInitializer Initializer = SLATE_ADD_PANELCHILDREN_DEFINITION(AttributeInitializer, Children);
+	FSlot::RegisterAttributes(Initializer);
+
 	SLATE_ADD_MEMBER_ATTRIBUTE_DEFINITION_WITH_NAME(AttributeInitializer, "ItemWidth", ItemWidthAttribute, EInvalidateWidgetReason::Layout);
 	SLATE_ADD_MEMBER_ATTRIBUTE_DEFINITION_WITH_NAME(AttributeInitializer, "ItemHeight", ItemHeightAttribute, EInvalidateWidgetReason::Layout);
 	SLATE_ADD_MEMBER_ATTRIBUTE_DEFINITION_WITH_NAME(AttributeInitializer, "NumDesiredItems", NumDesiredItemsAttribute, EInvalidateWidgetReason::Layout);
@@ -141,6 +144,8 @@ void SListPanel::OnArrangeChildren( const FGeometry& AllottedGeometry, FArranged
 	
 void SListPanel::Tick( const FGeometry& AllottedGeometry, const double InCurrentTime, const float InDeltaTime )
 {
+	int32 OldPreferredNumLines = PreferredNumLines;
+	
 	if (ShouldArrangeAsTiles())
 	{
 		const FTableViewDimensions AllottedDimensions(Orientation, AllottedGeometry.GetLocalSize());
@@ -166,6 +171,11 @@ void SListPanel::Tick( const FGeometry& AllottedGeometry, const double InCurrent
 	{
 		PreferredNumLines = NumDesiredItemsAttribute.Get();
 	}
+
+	if (OldPreferredNumLines != PreferredNumLines)
+	{
+		Invalidate(EInvalidateWidgetReason::Layout);
+	}
 }
 
 FVector2D SListPanel::ComputeDesiredSize( float ) const
@@ -220,12 +230,20 @@ FChildren* SListPanel::GetChildren()
 
 void SListPanel::SetFirstLineScrollOffset(float InFirstLineScrollOffset)
 {
-	FirstLineScrollOffset = InFirstLineScrollOffset;
+	if (FirstLineScrollOffset != InFirstLineScrollOffset)
+	{
+		FirstLineScrollOffset = InFirstLineScrollOffset;
+		Invalidate(EInvalidateWidgetReason::Layout);
+	}
 }
 
 void SListPanel::SetOverscrollAmount( float InOverscrollAmount )
 {
-	OverscrollAmount = InOverscrollAmount;
+	if (OverscrollAmount != InOverscrollAmount)
+	{
+		OverscrollAmount = InOverscrollAmount;
+		Invalidate(EInvalidateWidgetReason::Layout);
+	}
 }
 
 void SListPanel::ClearItems()
diff --git a/Engine/Source/Runtime/Slate/Private/Widgets/Views/SListPanel.h b/Engine/Source/Runtime/Slate/Private/Widgets/Views/SListPanel.h
--- a/Engine/Source/Runtime/Slate/Private/Widgets/Views/SListPanel.h
+++ b/Engine/Source/Runtime/Slate/Private/Widgets/Views/SListPanel.h
@@ -13,6 +13,12 @@
 
 class FArrangedChildren;
 
+template <>
+struct TWidgetTypeTraits<class SListPanel>
+{
+	static constexpr bool SupportsInvalidation() { return true; }
+};
+
 /**
  * A really simple panel that arranges its children in a vertical list with no spacing.
  * Items in this panel have a uniform height.
@@ -24,11 +30,11 @@ class SListPanel : public SPanel
 
 public:
 	/** A ListPanel slot is very simple - it just stores a widget. */
-	class FSlot : public TSlotBase<FSlot>
+	class FSlot : public TWidgetSlotWithAttributeSupport<FSlot>
 	{
 	public:
 		FSlot()
-		: TSlotBase<FSlot>()
+		: TWidgetSlotWithAttributeSupport<FSlot>()
 		{}
 	};
 	
diff --git a/Engine/Source/Runtime/Slate/Private/Widgets/Views/STableViewBase.cpp b/Engine/Source/Runtime/Slate/Private/Widgets/Views/STableViewBase.cpp
--- a/Engine/Source/Runtime/Slate/Private/Widgets/Views/STableViewBase.cpp
+++ b/Engine/Source/Runtime/Slate/Private/Widgets/Views/STableViewBase.cpp
@@ -67,6 +67,20 @@ FTableViewDimensions::FTableViewDimensions(EOrientation InOrientation, const UE:
 	}
 }
 
+SLATE_IMPLEMENT_WIDGET(STableViewBase)
+void STableViewBase::PrivateRegisterAttributes(FSlateAttributeInitializer& AttributeInitializer)
+{
+	SLATE_ADD_MEMBER_ATTRIBUTE_DEFINITION(AttributeInitializer, SelectionModeAttribute, EInvalidateWidgetReason::None)
+		.OnValueChanged(FSlateAttributeDescriptor::FAttributeValueChangedDelegate::CreateLambda([](SWidget& Widget)
+		{
+			STableViewBase* TableView = static_cast<STableViewBase*>(&Widget);
+			const ESelectionMode::Type OldSelectionMode = TableView->PreviousSelectionMode;
+			TableView->PreviousSelectionMode = TableView->GetSelectionMode();
+			TableView->OnSelectionModeChanged(OldSelectionMode, TableView->GetSelectionMode());
+		}));
+	SLATE_ADD_MEMBER_ATTRIBUTE_DEFINITION(AttributeInitializer, BackgroundBrushAttribute, EInvalidateWidgetReason::Layout);
+}
+
 TSharedRef<SWidget> STableViewBase::ConstructVerticalShadowBox(TSharedRef<SWidget> Content)
 {
 	return SNew(SOverlay)
@@ -80,20 +94,20 @@ TSharedRef<SWidget> STableViewBase::ConstructVerticalShadowBox(TSharedRef<SWidge
 		.VAlign(VAlign_Top)
 		[
 			// Shadow: Hint to scroll up
-			SNew(SImage)
-			.Visibility(EVisibility::HitTestInvisible)
+			SAssignNew(TopLeftShadowImage, SImage)
+			.Visibility(ShadowBoxStyle ? EVisibility::HitTestInvisible : EVisibility::Collapsed)
 			.ColorAndOpacity(this, &STableViewBase::GetStartShadowOpacity)
-			.Image(&ShadowBoxStyle->TopShadowBrush)
+			.Image(ShadowBoxStyle ? &ShadowBoxStyle->TopShadowBrush : nullptr)
 		]
 		+ SOverlay::Slot()
 		.HAlign(HAlign_Fill)
 		.VAlign(VAlign_Bottom)
 		[
 			// Shadow: a hint to scroll down
-			SNew(SImage)
-			.Visibility(EVisibility::HitTestInvisible)
+			SAssignNew(BottomRightShadowImage, SImage)
+			.Visibility(ShadowBoxStyle ? EVisibility::HitTestInvisible : EVisibility::Collapsed)
 			.ColorAndOpacity(this, &STableViewBase::GetEndShadowOpacity)
-			.Image(&ShadowBoxStyle->BottomShadowBrush)
+			.Image(ShadowBoxStyle ? &ShadowBoxStyle->BottomShadowBrush : nullptr)
 		];
 }
 
@@ -110,20 +124,20 @@ TSharedRef<SWidget> STableViewBase::ConstructHorizontalShadowBox(TSharedRef<SWid
 		.VAlign(VAlign_Fill)
 		[
 			// Shadow: Hint to left
-			SNew(SImage)
-			.Visibility(EVisibility::HitTestInvisible)
+			SAssignNew(TopLeftShadowImage, SImage)
+			.Visibility(ShadowBoxStyle ? EVisibility::HitTestInvisible : EVisibility::Collapsed)
 			.ColorAndOpacity(this, &STableViewBase::GetStartShadowOpacity)
-			.Image(&ShadowBoxStyle->LeftShadowBrush)
+			.Image(ShadowBoxStyle ? &ShadowBoxStyle->LeftShadowBrush : nullptr)
 		]
 		+ SOverlay::Slot()
 		.HAlign(HAlign_Right)
 		.VAlign(VAlign_Fill)
 		[
 			// Shadow: a hint to scroll right
-			SNew(SImage)
-			.Visibility(EVisibility::HitTestInvisible)
+			SAssignNew(BottomRightShadowImage, SImage)
+			.Visibility(ShadowBoxStyle ? EVisibility::HitTestInvisible : EVisibility::Collapsed)
 			.ColorAndOpacity(this, &STableViewBase::GetEndShadowOpacity)
-			.Image(&ShadowBoxStyle->RightShadowBrush)
+			.Image(ShadowBoxStyle ? &ShadowBoxStyle->RightShadowBrush : nullptr)
 		];
 }
 
@@ -188,7 +202,7 @@ void STableViewBase::ConstructChildren( const TAttribute<float>& InItemWidth, co
 			];		
 
 		// If we use Shadow Box Style, wrap the Container in a Shadow Box.
-		if (bShouldUseShadowBoxStyle && ShadowBoxStyle)
+		if (bShouldUseShadowBoxStyle)
 		{
 			if (Orientation == Orient_Vertical)
 			{
@@ -844,16 +858,11 @@ bool STableViewBase::IsPendingRefresh() const
 	return bItemsNeedRefresh || ItemsPanel->IsRefreshPending();
 }
 
-bool STableViewBase::ComputeVolatility() const
-{
-	return BackgroundBrush.IsBound();
-}
-
 int32 STableViewBase::OnPaint( const FPaintArgs& Args, const FGeometry& AllottedGeometry, const FSlateRect& MyCullingRect, FSlateWindowElementList& OutDrawElements, int32 LayerId, const FWidgetStyle& InWidgetStyle, bool bParentEnabled ) const
 {
 	int32 NewLayerId = LayerId;
 
-	const FSlateBrush* BackgroundBrushResource = BackgroundBrush.Get();
+	const FSlateBrush* BackgroundBrushResource = GetBackgroundBrush();
 	if ( BackgroundBrushResource && BackgroundBrushResource->DrawAs != ESlateBrushDrawType::NoDrawType )
 	{
 
@@ -897,18 +906,19 @@ STableViewBase::STableViewBase( ETableViewMode::Type InTableViewMode )
 	, TickScrollDelta( 0 )
 	, LastGenerateResults( 0,0,0,false )
 	, bWasAtEndOfList(false)
-	, SelectionMode( ESelectionMode::Multi )
 	, SoftwareCursorPosition( ForceInitToZero )
 	, bShowSoftwareCursor( false )
 	, WheelScrollMultiplier(GetGlobalScrollAmount())
 	, bShouldUseShadowBoxStyle(false)
-	, ShadowBoxStyle( &FAppStyle::Get().GetWidgetStyle<FScrollBoxStyle>("ScrollBox"))
-	, BackgroundBrush(FStyleDefaults::GetNoBrush())
 	, bIsScrollingActiveTimerRegistered( false )
 	, Overscroll()
 	, AllowOverscroll(EAllowOverscroll::Yes)
 	, ConsumeMouseWheel(EConsumeMouseWheel::WhenScrollingPossible)
 	, bItemsNeedRefresh( false )	
+	, ShadowBoxStyle(&FAppStyle::Get().GetWidgetStyle<FScrollBoxStyle>("ScrollBox"))
+	, SelectionModeAttribute(*this, ESelectionMode::Multi)
+	, PreviousSelectionMode(ESelectionMode::Multi)
+	, BackgroundBrushAttribute(*this, FStyleDefaults::GetNoBrush())
 {
 	FixedLineScrollOffset = 0.25f;
 }
@@ -1157,9 +1167,13 @@ void STableViewBase::SetIsGamepadScrollingEnabled(bool bInIsGamepadScrollingEnab
 	bIsGamepadScrollingEnabled = bInIsGamepadScrollingEnabled;
 }
 
-void STableViewBase::SetBackgroundBrush(const TAttribute<const FSlateBrush*>& InBackgroundBrush)
+void STableViewBase::SetBackgroundBrush(TAttribute<const FSlateBrush*> InBackgroundBrush)
 {
-	BackgroundBrush.SetImage(*this, InBackgroundBrush);
+	BackgroundBrushAttribute.Assign(*this, TAttribute<FSlateBrushCached>::CreateLambda(
+		[BackgroundBrush = MoveTemp(InBackgroundBrush)]()
+		{
+			return FSlateBrushCached(BackgroundBrush.Get(nullptr));
+		}));
 }
 
 void STableViewBase::InsertWidget( const TSharedRef<ITableRow> & WidgetToInset )
@@ -1352,6 +1366,34 @@ FVector2D STableViewBase::GetScrollDistanceRemaining()
 	return FVector2D( 0, ScrollBar->DistanceFromBottom() );
 }
 
+void STableViewBase::SetShadowBoxStyle(const FScrollBoxStyle* InShadowBoxStyle)
+{
+	if (ShadowBoxStyle != InShadowBoxStyle)
+	{
+		ShadowBoxStyle = InShadowBoxStyle;
+		if (TopLeftShadowImage)
+		{
+			const FSlateBrush* Image = nullptr;
+			if (ShadowBoxStyle)
+			{
+				Image = Orientation == Orient_Vertical ? &ShadowBoxStyle->TopShadowBrush : &ShadowBoxStyle->LeftShadowBrush;
+			}
+			TopLeftShadowImage->SetImage(Image);
+			TopLeftShadowImage->SetVisibility(ShadowBoxStyle ? EVisibility::HitTestInvisible : EVisibility::Collapsed);
+		}
+		if (BottomRightShadowImage)
+		{
+			const FSlateBrush* Image = nullptr;
+			if (ShadowBoxStyle)
+			{
+				Image = Orientation == Orient_Vertical ? &ShadowBoxStyle->BottomShadowBrush : &ShadowBoxStyle->RightShadowBrush;
+			}
+			BottomRightShadowImage->SetImage(Image);
+			BottomRightShadowImage->SetVisibility(ShadowBoxStyle ? EVisibility::HitTestInvisible : EVisibility::Collapsed);
+		}
+	}
+}
+
 TSharedRef<class SWidget> STableViewBase::GetScrollWidget()
 {
 	return SharedThis(this);
diff --git a/Engine/Source/Runtime/Slate/Public/Widgets/Views/SListView.h b/Engine/Source/Runtime/Slate/Public/Widgets/Views/SListView.h
--- a/Engine/Source/Runtime/Slate/Public/Widgets/Views/SListView.h
+++ b/Engine/Source/Runtime/Slate/Public/Widgets/Views/SListView.h
@@ -32,6 +32,172 @@
 #include "Widgets/Accessibility/SlateAccessibleMessageHandler.h"
 #endif
 
+template <typename ItemType>
+class SListView;
+
+template <typename ItemType>
+struct TWidgetTypeTraits<SListView<ItemType>>
+{
+	static constexpr bool SupportsInvalidation() { return true; }
+};
+
+namespace UE::Slate::Private
+{
+	template <typename ItemType>
+	class SListViewPinnedRowWidget;
+}
+
+template <typename ItemType>
+struct TWidgetTypeTraits<class UE::Slate::Private::SListViewPinnedRowWidget<ItemType>>
+{
+	static constexpr bool SupportsInvalidation() { return true; }
+};
+
+
+namespace UE::Slate::Private
+{
+	// Private class that acts as a wrapper around PinnedRows, to allow for customized styling
+	template <typename ItemType>
+	class SListViewPinnedRowWidget : public SCompoundWidget
+	{
+	public:
+		SLATE_BEGIN_ARGS(SListViewPinnedRowWidget) {}
+
+		SLATE_END_ARGS()
+
+		void Construct(const FArguments& InArgs, TSharedPtr<ITableRow> InPinnedItemRow, TSharedRef<SListView<ItemType>> InOwnerListView, const int32 ItemIndex, const int32 NumPinnedItems)
+		{
+			PinnedItemRow = InPinnedItemRow;
+			OwnerListView = InOwnerListView;
+
+			TSharedPtr<STableRow<ItemType>> PinnedRow = StaticCastSharedPtr<STableRow<ItemType>>(InPinnedItemRow);
+
+			// If the PinnedRow inherits from STableRow (i.e has a custom border already), remove it since we will be adding our own border
+			// Also set the expander arrow to Hidden so it is not available for pinned rows, but still occupies the same space
+			if (PinnedRow.IsValid())
+			{
+				PinnedRow->SetBorderImage(FAppStyle::Get().GetBrush("NoBrush"));
+				PinnedRow->SetExpanderArrowVisibility(EVisibility::Hidden);
+			}
+
+			TSharedRef<SWidget> InPinnedItemRowWidget = InPinnedItemRow->AsWidget();
+			InPinnedItemRowWidget->SetVisibility(EVisibility::HitTestInvisible);
+
+			ChildSlot
+				[
+					SNew(SOverlay)
+					.Visibility(TAttribute<EVisibility>::CreateSP(this, &SListViewPinnedRowWidget::SetPinnedItemVisibility, ItemIndex, NumPinnedItems))
+
+					+ SOverlay::Slot()
+					.Padding(FMargin(0.0f, 0.0f, 0.0f, 0.0f))
+					[
+						SNew(SBorder)
+						.BorderImage_Lambda([this]()
+							{
+								return this->IsHovered() ? FAppStyle::Get().GetBrush("Brushes.Hover") : FAppStyle::Get().GetBrush("Brushes.Header");
+							})
+						.Padding(0.f)
+						.Content()
+						[
+							SNew(SHorizontalBox)
+							+ SHorizontalBox::Slot()
+							.AutoWidth()
+							[
+								InPinnedItemRowWidget
+							]
+							+ SHorizontalBox::Slot()
+							.FillWidth(1.0f)
+							.Padding(2.0f, 2.0f, 0.0f, 0.0f)
+							[
+								// Text Block for ellipses, shows up when some items in the pinned list are collapsed when the number of items > MaxPinnedItems
+								SNew(STextBlock).Text(NSLOCTEXT("SListView", "Ellipses", "..."))
+								.Visibility(this, &SListViewPinnedRowWidget::SetPinnedItemEllipsesVisibility, ItemIndex)
+							]
+						]
+					]
+					+ SOverlay::Slot()
+					.HAlign(HAlign_Fill)
+					.VAlign(VAlign_Top)
+					[
+						// A shadow to indicate parent/child relationship
+						SNew(SImage)
+						.Visibility(EVisibility::HitTestInvisible)
+						.Image(FAppStyle::Get().GetBrush("ListView.PinnedItemShadow"))
+					]
+					
+				];
+
+		}
+
+	protected:
+
+		virtual FReply OnMouseButtonDown(const FGeometry& MyGeometry, const FPointerEvent& MouseEvent) override
+		{
+			if (MouseEvent.GetEffectingButton() == EKeys::LeftMouseButton)
+			{
+				const TObjectPtrWrapTypeOf<ItemType>* PinnedItem = OwnerListView->ItemFromWidget(PinnedItemRow.Get());
+
+				if (PinnedItem)
+				{
+					// Select the pinned item on click 
+					OwnerListView->NavigationSelect(*PinnedItem, MouseEvent);
+					return FReply::Handled();
+				}
+			}
+
+			return FReply::Unhandled();
+		}
+
+	private:
+
+		EVisibility SetPinnedItemVisibility(const int32 IndexInList, const int32 NumPinnedItems) const
+		{
+			// If the hierarchy is not collapsed (i.e all items are visible)
+			if (!OwnerListView->bIsHierarchyCollapsed)
+			{
+				return EVisibility::Visible;
+			}
+
+			int32 CurrentMaxPinnedItems = OwnerListView->GetMaxPinnedItems();
+
+			// If this is the last item, it is visible
+			if (IndexInList == NumPinnedItems - 1)
+			{
+				return EVisibility::Visible;
+			}
+			// Only show a limited number of items depending on MaxPinnedItems
+			else if (IndexInList < CurrentMaxPinnedItems - 1)
+			{
+				return EVisibility::Visible;
+			}
+
+			return EVisibility::Collapsed;
+		}
+
+		EVisibility SetPinnedItemEllipsesVisibility(const int32 IndexInList) const
+		{
+			// If all items are visible, the ...'s are never visible
+			if (!OwnerListView->bIsHierarchyCollapsed)
+			{
+				return EVisibility::Collapsed;
+			}
+
+			int32 CurrentMaxPinnedItems = OwnerListView->GetMaxPinnedItems();
+
+			// If the hierarchy is collapsed, the last item before the collapsed items gets the ellipses
+			return IndexInList == CurrentMaxPinnedItems - 2 ? EVisibility::Visible : EVisibility::Collapsed;
+		}
+
+	private:
+		// A pointer to the ListView that owns this widget
+		TSharedPtr<SListView<ItemType>> OwnerListView;
+
+		// A pointer to the row contained by this widget
+ 		TSharedPtr<ITableRow> PinnedItemRow;
+	};
+}
+
+
 /**
  * A ListView widget observes an array of data items and creates visual representations of these items.
  * ListView relies on the property that holding a reference to a value ensures its existence. In other words,
@@ -64,6 +230,8 @@
 template <typename ItemType>
 class SListView : public STableViewBase, TListTypeTraits<ItemType>::SerializerType, public ITypedTableView< ItemType >
 {
+	SLATE_DECLARE_WIDGET(SListView, STableViewBase)
+	static inline SLATE_IMPLEMENT_WIDGET(SListView)
 public:
 	using NullableItemType  = typename TListTypeTraits< ItemType >::NullableType;
 	using MapKeyFuncs       = typename TListTypeTraits<ItemType>::MapKeyFuncs;
@@ -246,7 +414,7 @@
 		this->OnDoubleClick = InArgs._OnMouseButtonDoubleClick;
 		this->OnSelectionChanged = InArgs._OnSelectionChanged;
 		this->OnIsSelectableOrNavigable = InArgs._OnIsSelectableOrNavigable;
-		this->SelectionMode = InArgs._SelectionMode;
+		this->SetSelectionMode(InArgs._SelectionMode);
 
 		this->bClearSelectionOnClick = InArgs._ClearSelectionOnClick;
 
@@ -281,7 +449,7 @@
 
 		this->SetShadowBoxStyle(InArgs._ShadowBoxStyle);
 			
-		this->MaxPinnedItems = InArgs._MaxPinnedItems;
+		this->SetMaxPinnedItems(InArgs._MaxPinnedItems);
 		this->DefaultMaxPinnedItems = InArgs._MaxPinnedItems;
 		this->ScrollBarSlotPadding = InArgs._ScrollBarPadding;
 
@@ -336,7 +504,8 @@
 		, ItemToNotifyWhenInView(TListTypeTraits<ItemType>::MakeNullPtr())
 		, ScrollIntoViewAlignment(EScrollIntoViewAlignment::CenterAligned)
 		, IsFocusable(true)
-	{ 
+		, MaxPinnedItemsAttribute(*this, 6)
+	{
 #if WITH_ACCESSIBILITY
 		AccessibleBehavior = EAccessibleBehavior::Auto;
 		bCanChildrenBeAccessible = true;
@@ -452,7 +621,7 @@
 					ItemType SelectorItemDereference(TListTypeTraits<ItemType>::NullableItemTypeConvertToItemType(SelectorItem));
 
 					// Deselect.
-					if (InKeyEvent.IsControlDown() || SelectionMode.Get() == ESelectionMode::SingleToggle)
+					if (InKeyEvent.IsControlDown() || GetSelectionMode() == ESelectionMode::SingleToggle)
 					{
 						this->Private_SetItemSelection(SelectorItemDereference, !(this->Private_IsItemSelected(SelectorItemDereference)), true);
 						this->Private_SignalSelectionChanged(ESelectInfo::OnKeyPress);
@@ -484,7 +653,7 @@
 					}
 				}
 				// Select all items
-				else if ( (!InKeyEvent.IsShiftDown() && !InKeyEvent.IsAltDown() && InKeyEvent.IsControlDown() && InKeyEvent.GetKey() == EKeys::A) && SelectionMode.Get() == ESelectionMode::Multi )
+				else if ( (!InKeyEvent.IsShiftDown() && !InKeyEvent.IsAltDown() && InKeyEvent.IsControlDown() && InKeyEvent.GetKey() == EKeys::A) && GetSelectionMode() == ESelectionMode::Multi )
 				{
 					this->Private_ClearSelection();
 
@@ -573,7 +742,7 @@
 	virtual FReply OnMouseButtonDown( const FGeometry& MyGeometry, const FPointerEvent& MouseEvent ) override
 	{
 		if ( bClearSelectionOnClick
-			&& SelectionMode.Get() != ESelectionMode::None
+			&& GetSelectionMode() != ESelectionMode::None
 			&& MouseEvent.GetEffectingButton() == EKeys::LeftMouseButton
 			&& !MouseEvent.IsControlDown()
 			&& !MouseEvent.IsShiftDown()
@@ -596,7 +765,7 @@
 	virtual FReply OnMouseButtonUp( const FGeometry& MyGeometry, const FPointerEvent& MouseEvent ) override
 	{
 		if ( bClearSelectionOnClick
-			&& SelectionMode.Get() != ESelectionMode::None
+			&& GetSelectionMode() != ESelectionMode::None
 			&& MouseEvent.GetEffectingButton() == EKeys::RightMouseButton
 			&& !MouseEvent.IsControlDown()
 			&& !MouseEvent.IsShiftDown()
@@ -675,7 +844,7 @@
 			if (Widget.IsValid())
 			{
 				TSharedPtr<SListView<ItemType>> ListView = StaticCastSharedPtr<SListView<ItemType>>(Widget.Pin());
-				return ListView->SelectionMode.Get() == ESelectionMode::Multi;
+				return ListView->GetSelectionMode() == ESelectionMode::Multi;
 			}
 			return false; 
 		}
@@ -925,7 +1094,7 @@
 
 	virtual void Private_SetItemSelection( ItemType TheItem, bool bShouldBeSelected, bool bWasUserDirected = false ) override
 	{
-		if ( SelectionMode.Get() == ESelectionMode::None )
+		if ( GetSelectionMode() == ESelectionMode::None )
 		{
 			return;
 		}
@@ -940,7 +1109,7 @@
 		}
 
 		// Move the selector item and range selection start if the user directed this change in selection or if the list view is single selection
-		if( bWasUserDirected || SelectionMode.Get() == ESelectionMode::Single || SelectionMode.Get() == ESelectionMode::SingleToggle )
+		if( bWasUserDirected || GetSelectionMode() == ESelectionMode::Single || GetSelectionMode() == ESelectionMode::SingleToggle )
 		{
 			SelectorItem = TheItem;
 			RangeSelectionStart = TheItem;
@@ -983,7 +1152,7 @@
 
 	virtual void Private_SelectRangeFromCurrentTo( ItemType InRangeSelectionEnd ) override
 	{
-		if ( SelectionMode.Get() == ESelectionMode::None )
+		if ( GetSelectionMode() == ESelectionMode::None )
 		{
 			return;
 		}
@@ -1024,7 +1193,7 @@
 
 	virtual void Private_SignalSelectionChanged(ESelectInfo::Type SelectInfo) override
 	{
-		if ( SelectionMode.Get() == ESelectionMode::None )
+		if ( GetSelectionMode() == ESelectionMode::None )
 		{
 			return;
 		}
@@ -1132,7 +1301,7 @@
 
 	virtual ESelectionMode::Type Private_GetSelectionMode() const override
 	{
-		return SelectionMode.Get();
+		return GetSelectionMode();
 	}
 
 	virtual EOrientation Private_GetOrientation() const override
@@ -1181,147 +1350,8 @@
 	}
 
 private:
+	typedef UE::Slate::Private::SListViewPinnedRowWidget<ItemType> SListViewPinnedRowWidget;
+	friend SListViewPinnedRowWidget;
 
-	friend class SListViewPinnedRowWidget;
-
-	// Private class that acts as a wrapper around PinnedRows, to allow for customized styling
-	class SListViewPinnedRowWidget : public SCompoundWidget
-	{
-	public:
-		SLATE_BEGIN_ARGS(SListViewPinnedRowWidget) {}
-
-		SLATE_END_ARGS()
-
-		void Construct(const FArguments& InArgs, TSharedPtr<ITableRow> InPinnedItemRow, TSharedRef<SListView> InOwnerListView, const int32 ItemIndex, const int32 NumPinnedItems)
-		{
-			PinnedItemRow = InPinnedItemRow;
-			OwnerListView = InOwnerListView;
-
-			TSharedPtr<STableRow<ItemType>> PinnedRow = StaticCastSharedPtr<STableRow<ItemType>>(InPinnedItemRow);
-
-			// If the PinnedRow inherits from STableRow (i.e has a custom border already), remove it since we will be adding our own border
-			// Also set the expander arrow to Hidden so it is not available for pinned rows, but still occupies the same space
-			if (PinnedRow.IsValid())
-			{
-				PinnedRow->SetBorderImage(FAppStyle::Get().GetBrush("NoBrush"));
-				PinnedRow->SetExpanderArrowVisibility(EVisibility::Hidden);
-			}
-
-			TSharedRef<SWidget> InPinnedItemRowWidget = InPinnedItemRow->AsWidget();
-			InPinnedItemRowWidget->SetVisibility(EVisibility::HitTestInvisible);
-
-			ChildSlot
-				[
-					SNew(SOverlay)
-					.Visibility(TAttribute<EVisibility>::CreateSP(this, &SListViewPinnedRowWidget::SetPinnedItemVisibility, ItemIndex, NumPinnedItems))
-
-					+ SOverlay::Slot()
-					.Padding(FMargin(0.0f, 0.0f, 0.0f, 0.0f))
-					[
-						SNew(SBorder)
-						.BorderImage_Lambda([this]()
-							{
-								return this->IsHovered() ? FAppStyle::Get().GetBrush("Brushes.Hover") : FAppStyle::Get().GetBrush("Brushes.Header");
-							})
-						.Padding(0.f)
-						.Content()
-						[
-							SNew(SHorizontalBox)
-							+ SHorizontalBox::Slot()
-							.AutoWidth()
-							[
-								InPinnedItemRowWidget
-							]
-							+ SHorizontalBox::Slot()
-							.FillWidth(1.0f)
-							.Padding(2.0f, 2.0f, 0.0f, 0.0f)
-							[
-								// Text Block for ellipses, shows up when some items in the pinned list are collapsed when the number of items > MaxPinnedItems
-								SNew(STextBlock).Text(NSLOCTEXT("SListView", "Ellipses", "..."))
-								.Visibility(this, &SListViewPinnedRowWidget::SetPinnedItemEllipsesVisibility, ItemIndex)
-							]
-						]
-					]
-					+ SOverlay::Slot()
-					.HAlign(HAlign_Fill)
-					.VAlign(VAlign_Top)
-					[
-						// A shadow to indicate parent/child relationship
-						SNew(SImage)
-						.Visibility(EVisibility::HitTestInvisible)
-						.Image(FAppStyle::Get().GetBrush("ListView.PinnedItemShadow"))
-					]
-					
-				];
-
-		}
-
-	protected:
-
-		virtual FReply OnMouseButtonDown(const FGeometry& MyGeometry, const FPointerEvent& MouseEvent) override
-		{
-			if (MouseEvent.GetEffectingButton() == EKeys::LeftMouseButton)
-			{
-				const TObjectPtrWrapTypeOf<ItemType>* PinnedItem = OwnerListView->ItemFromWidget(PinnedItemRow.Get());
-
-				if (PinnedItem)
-				{
-					// Select the pinned item on click 
-					OwnerListView->NavigationSelect(*PinnedItem, MouseEvent);
-					return FReply::Handled();
-				}
-			}
-
-			return FReply::Unhandled();
-		}
-
-	private:
-
-		EVisibility SetPinnedItemVisibility(const int32 IndexInList, const int32 NumPinnedItems) const
-		{
-			// If the hierarchy is not collapsed (i.e all items are visible)
-			if (!OwnerListView->bIsHierarchyCollapsed)
-			{
-				return EVisibility::Visible;
-			}
-
-			int32 CurrentMaxPinnedItems = OwnerListView->MaxPinnedItems.Get();
-
-			// If this is the last item, it is visible
-			if (IndexInList == NumPinnedItems - 1)
-			{
-				return EVisibility::Visible;
-			}
-			// Only show a limited number of items depending on MaxPinnedItems
-			else if (IndexInList < CurrentMaxPinnedItems - 1)
-			{
-				return EVisibility::Visible;
-			}
-
-			return EVisibility::Collapsed;
-		}
-
-		EVisibility SetPinnedItemEllipsesVisibility(const int32 IndexInList) const
-		{
-			// If all items are visible, the ...'s are never visible
-			if (!OwnerListView->bIsHierarchyCollapsed)
-			{
-				return EVisibility::Collapsed;
-			}
-
-			int32 CurrentMaxPinnedItems = OwnerListView->MaxPinnedItems.Get();
-
-			// If the hierarchy is collapsed, the last item before the collapsed items gets the ellipses
-			return IndexInList == CurrentMaxPinnedItems - 2 ? EVisibility::Visible : EVisibility::Collapsed;
-		}
-
-	private:
-		// A pointer to the ListView that owns this widget
-		TSharedPtr<SListView> OwnerListView;
-
-		// A pointer to the row contained by this widget
- 		TSharedPtr<ITableRow> PinnedItemRow;
-	};
-
 public:	
 
@@ -1331,11 +1361,6 @@
 		SetBackgroundBrush( Style != nullptr ? &Style->BackgroundBrush : FStyleDefaults::GetNoBrush() );
 	}
 
-	void SetShadowBoxStyle(const FScrollBoxStyle* InScrollBoxStyle)
-	{
-		ShadowBoxStyle = InScrollBoxStyle;
-	}
-
 	/** Sets the OnEntryInitializer delegate. This delegate is invoked after initializing an entry being generated, before it may be added to the actual widget hierarchy. */
 	void SetOnEntryInitialized(const FOnEntryInitialized& Delegate)
 	{
@@ -1592,15 +1617,15 @@
 		// Check if the User provided an override for MaxPinnedItems, valid until the next time ReGeneratePinnedItems is called
 		if (MaxPinnedItemsOverride != -1)
 		{
-			MaxPinnedItems.Set(MaxPinnedItemsOverride);
+			SetMaxPinnedItems(MaxPinnedItemsOverride);
 		}
 		else
 		{
 			// Reset it back to the default value if there is no override
-			MaxPinnedItems = DefaultMaxPinnedItems;
+			SetMaxPinnedItems(DefaultMaxPinnedItems);
 		}
 
-		int32 CurrentMaxPinnedItems = MaxPinnedItems.Get();
+		int32 CurrentMaxPinnedItems = MaxPinnedItemsAttribute.Get();
 		// There are more items than what we allow to show
 		if (InItems.Num() > CurrentMaxPinnedItems)
 		{
@@ -1818,7 +1843,7 @@
 	 */
 	bool IsItemSelected( const ItemType& InItem ) const
 	{
-		if ( SelectionMode.Get() == ESelectionMode::None )
+		if ( GetSelectionMode() == ESelectionMode::None )
 		{
 			return false;
 		}
@@ -1835,7 +1860,7 @@
 	 */
 	void SetItemSelection( const ItemType& InItem, bool bSelected, ESelectInfo::Type SelectInfo = ESelectInfo::Direct )
 	{
-		if ( SelectionMode.Get() == ESelectionMode::None )
+		if ( GetSelectionMode() == ESelectionMode::None )
 		{
 			return;
 		}
@@ -1853,7 +1878,7 @@
 	 */
 	void SetItemSelection(TConstArrayView<ItemType> InItems, bool bSelected, ESelectInfo::Type SelectInfo = ESelectInfo::Direct)
 	{
-		if ( InItems.Num() == 0 || SelectionMode.Get() == ESelectionMode::None )
+		if ( InItems.Num() == 0 || GetSelectionMode() == ESelectionMode::None )
 		{
 			return;
 		}
@@ -1873,11 +1898,6 @@
 	 */
 	void ClearSelection()
 	{
-		if ( SelectionMode.Get() == ESelectionMode::None )
-		{
-			return;
-		}
-
 		if ( SelectedItems.Num() == 0 )
 		{
 			return;
@@ -2050,32 +2070,9 @@
 	 * If going from multi-select to a type of single-select and one item is selected, it will be maintained (otherwise all will be cleared).
 	 * If disabling selection, any current selections will be cleared.
 	 */
-	void SetSelectionMode(const TAttribute<ESelectionMode::Type>& NewSelectionMode)
+	void SetSelectionMode(TAttribute<ESelectionMode::Type> NewSelectionMode)
 	{
-		const ESelectionMode::Type PreviousMode = SelectionMode.Get();
-		SelectionMode = NewSelectionMode;
-		const ESelectionMode::Type NewMode = NewSelectionMode.Get();
-		if (PreviousMode != NewMode)
-		{
-			if (NewMode == ESelectionMode::None)
-			{
-				ClearSelection();
-			}
-			else if (PreviousMode == ESelectionMode::Multi)
-			{
-				// We've gone to a single-selection mode, so if we already had a single item selected, preserve it
-				if (SelectedItems.Num() == 1)
-				{
-					SetSelection(*SelectedItems.CreateIterator());
-				}
-				else
-				{
-					// Otherwise, there's no way to know accurately which item was selected most recently, so just wipe it all
-					// The caller responsible for changing the mode can decide themselves which item they want to be selected
-					ClearSelection();
-				}
-			}
-		}
+		Private_SetSelectionMode(MoveTemp(NewSelectionMode));
 	}
 
 	/* Sets ScrollIntoViewAlignment which allows to stick the selected item to either side or center */
@@ -2141,6 +2138,30 @@
 	}
 
 protected:
+	virtual void OnSelectionModeChanged(ESelectionMode::Type OldSelectionMode, ESelectionMode::Type NewSelectionMode) override
+	{
+		if (OldSelectionMode != NewSelectionMode)
+		{
+			if (NewSelectionMode == ESelectionMode::None)
+			{
+				ClearSelection();
+			}
+			else if (OldSelectionMode == ESelectionMode::Multi)
+			{
+				// We've gone to a single-selection mode, so if we already had a single item selected, preserve it
+				if (SelectedItems.Num() == 1)
+				{
+					SetSelection(*SelectedItems.CreateIterator());
+				}
+				else
+				{
+					// Otherwise, there's no way to know accurately which item was selected most recently, so just wipe it all
+					// The caller responsible for changing the mode can decide themselves which item they want to be selected
+					ClearSelection();
+				}
+			}
+		}
+	}
 
 	static FOnItemToString_Debug GetDefaultDebugDelegate()
 	{
@@ -2571,7 +2592,7 @@
 			return;
 		}
 
-		const ESelectionMode::Type CurrentSelectionMode = SelectionMode.Get();
+		const ESelectionMode::Type CurrentSelectionMode = GetSelectionMode();
 
 		if (CurrentSelectionMode != ESelectionMode::None)
 		{
@@ -2624,6 +2645,21 @@
 		NavigationTypeOnScrollIntoView = bNavigateOnScrollIntoView ? InNavigationEvent.GetNavigationType() : EUINavigation::Invalid;
 	}
 
+protected:
+	int32 GetMaxPinnedItems() const
+	{
+		return MaxPinnedItemsAttribute.Get();
+	}
+	void SetMaxPinnedItems(TAttribute<int32> InMaxPinnedItems)
+	{
+		MaxPinnedItemsAttribute.Assign(*this, MoveTemp(InMaxPinnedItems));
+	}
+	TSlateAttributeRef<int32> GetMaxPinnedItemsAttribute() const
+	{
+		return TSlateAttributeRef<int32>{SharedThis(this), MaxPinnedItemsAttribute};
+	}
+
+	const FTableViewStyle* GetStyle() const { return Style;}
 
 protected:
 	/** A widget generator component */
@@ -2732,11 +2768,11 @@
 	/** Cached navigation direction to forward to navigate widget call */
 	EUINavigation NavigationTypeOnScrollIntoView = EUINavigation::Invalid;
 
-	/** Style resource for the list */
-	const FTableViewStyle* Style;
-
+#if WITH_EDITORONLY_DATA
 	/** The maximum number of pinned items allowed */
-	TAttribute<int32> MaxPinnedItems;
+	UE_DEPRECATED(5.8, "Use SetMaxPinnedItems / GetMaxPinnedItems / GetMaxPinnedItemsAttribute")
+	TSlateDeprecatedTAttribute<int32> MaxPinnedItems;
+#endif
 
 	/** The initial value of MaxPinnedItems (used to restore it back if overriden) */
 	TAttribute<int32> DefaultMaxPinnedItems;
@@ -2745,9 +2781,15 @@
 	bool bIsHierarchyCollapsed = false;
 
 private:
+	/** Style resource for the list */
+	const FTableViewStyle* Style;
+
 	/** Pointer to the source data that we are observing */
 	TUniquePtr<UE::Slate::ItemsSource::IItemsSource<ItemType>> ViewSource;
 
+	/** The maximum number of pinned items allowed */
+	TSlateAttribute<int32> MaxPinnedItemsAttribute;
+
 private:
 	struct FGenerationPassGuard
 	{
@@ -2766,3 +2808,17 @@
 		}
 	};
 };
+
+template<typename ItemType>
+inline void SListView<ItemType>::PrivateRegisterAttributes(FSlateAttributeInitializer& AttributeInitializer)
+{
+	struct FInvalidation
+	{
+		static void RequestLayoutRefresh(SWidget& OwningWidget)
+		{
+			static_cast<SListView<ItemType>&>(OwningWidget).RequestLayoutRefresh();
+		}
+	};
+	SLATE_ADD_MEMBER_ATTRIBUTE_DEFINITION(AttributeInitializer, MaxPinnedItemsAttribute, EInvalidateWidgetReason::Layout)
+		.OnValueChanged(FSlateAttributeDescriptor::FAttributeValueChangedDelegate::CreateStatic(FInvalidation::RequestLayoutRefresh));
+}
diff --git a/Engine/Source/Runtime/Slate/Public/Widgets/Views/STableViewBase.h b/Engine/Source/Runtime/Slate/Public/Widgets/Views/STableViewBase.h
--- a/Engine/Source/Runtime/Slate/Public/Widgets/Views/STableViewBase.h
+++ b/Engine/Source/Runtime/Slate/Public/Widgets/Views/STableViewBase.h
@@ -101,6 +101,13 @@ struct FTableViewDimensions
 	float LineAxis = 0.f;
 };
 
+
+template <>
+struct TWidgetTypeTraits<class STableViewBase>
+{
+	static constexpr bool SupportsInvalidation() { return true; }
+};
+
 /**
  * Contains ListView functionality that does not depend on the type of data being observed by the ListView.
  */
@@ -108,6 +115,7 @@ class STableViewBase
 	: public SCompoundWidget
 	, public IScrollableWidget
 {
+	SLATE_DECLARE_WIDGET_API(STableViewBase, SCompoundWidget, SLATE_API)
 public:
 
 	/** Create the child widgets that comprise the list */
@@ -222,7 +230,7 @@ public:
 	SLATE_API void SetIsGamepadScrollingEnabled(bool bInIsGamepadScrollingEnabled);
 
 	/** Sets the Background Brush */
-	SLATE_API void SetBackgroundBrush(const TAttribute<const FSlateBrush*>& InBackgroundBrush);
+	SLATE_API void SetBackgroundBrush(TAttribute<const FSlateBrush*> InBackgroundBrush);
 
 public:
 
@@ -243,7 +251,6 @@ public:
 	SLATE_API virtual FReply OnKeyDown( const FGeometry& MyGeometry, const FKeyEvent& InKeyEvent ) override;
 	SLATE_API virtual FCursorReply OnCursorQuery( const FGeometry& MyGeometry, const FPointerEvent& CursorEvent ) const override;
 	SLATE_API virtual int32 OnPaint(const FPaintArgs& Args, const FGeometry& AllottedGeometry, const FSlateRect& MyCullingRect, FSlateWindowElementList& OutDrawElements, int32 LayerId, const FWidgetStyle& InWidgetStyle, bool bParentEnabled) const override;
-	SLATE_API virtual bool ComputeVolatility() const override;
 	SLATE_API virtual FReply OnTouchStarted( const FGeometry& MyGeometry, const FPointerEvent& InTouchEvent ) override;
 	SLATE_API virtual FReply OnTouchMoved( const FGeometry& MyGeometry, const FPointerEvent& InTouchEvent ) override;
 	SLATE_API virtual FReply OnTouchEnded( const FGeometry& MyGeometry, const FPointerEvent& InTouchEvent ) override;
@@ -256,6 +263,12 @@ public:
 	SLATE_API virtual FVector2D GetScrollDistanceRemaining() override;
 	SLATE_API virtual TSharedRef<class SWidget> GetScrollWidget() override;
 
+	SLATE_API void SetShadowBoxStyle(const FScrollBoxStyle* InShadowBoxStyle);
+	const FScrollBoxStyle* GetShadowBoxStyle() const 
+	{
+		return ShadowBoxStyle;
+	}
+
 protected:
 
 	SLATE_API STableViewBase( ETableViewMode::Type InTableViewMode );
@@ -357,6 +370,11 @@ protected:
 	 */
 	virtual void UpdateSelectionSet() = 0;
 
+	/**
+	 * Automatically called when the Selection Mode has changed
+	 */
+	virtual void OnSelectionModeChanged(ESelectionMode::Type OldSelectionMode, ESelectionMode::Type NewSelectionMode) {}
+
 	/** Internal request for a layout update on the next tick (i.e. a refresh without implication that the source items have changed) */
 	SLATE_API void RequestLayoutRefresh();
 
@@ -433,6 +451,30 @@ protected:
 	/** Util function to find the child index under the given position. */
 	SLATE_API int32 FindChildUnderPosition(FArrangedChildren& ArrangedChildren, const FVector2D& ArrangedSpacePosition) const;
 
+protected:
+	ESelectionMode::Type GetSelectionMode() const
+	{
+		return SelectionModeAttribute.Get();
+	}
+	void Private_SetSelectionMode(TAttribute<ESelectionMode::Type> InSelectionMode)
+	{
+		SelectionModeAttribute.Assign(*this, MoveTemp(InSelectionMode));
+	}
+	TSlateAttributeRef<ESelectionMode::Type> GetSelectionModeAttribute() const
+	{
+		return TSlateAttributeRef<ESelectionMode::Type>{SharedThis(this), SelectionModeAttribute};
+	}
+
+	const FSlateBrush* GetBackgroundBrush() const
+	{
+		return BackgroundBrushAttribute.Get().GetImage();
+	}
+	TSlateAttributeRef<FSlateBrushCached> GetBackgroundBrushAttribute() const
+	{
+		return TSlateAttributeRef<FSlateBrushCached>{SharedThis(this), BackgroundBrushAttribute};
+	}
+
+protected:
 	/** The panel which holds the visible widgets in this list */
 	TSharedPtr< SListPanel > ItemsPanel;
 
@@ -442,6 +484,11 @@ protected:
 	/** The scroll bar widget */
 	TSharedPtr< SScrollBar > ScrollBar;
 
+	/** The ShadowImage widget for either the Top or the Left shadow depending on the Orientation */
+	TSharedPtr< class SImage > TopLeftShadowImage;
+	/** The ShadowImage widget for either the Bottom or the Right shadow depending on the Orientation */
+	TSharedPtr< class SImage > BottomRightShadowImage;
+
 	/** Padding to the scrollbox */
 	FMargin ScrollBarSlotPadding;
 
@@ -522,8 +569,11 @@ protected:
 	/** Delegate to invoke after the children have been regenerated. */
 	FSimpleDelegate OnItemsRebuilt;
 
+#if WITH_EDITORONLY_DATA
 	/** The selection mode that this tree/list is in. Note that it is up to the generated ITableRows to respect this setting. */
-	TAttribute<ESelectionMode::Type> SelectionMode;
+	UE_DEPRECATED(5.8, "Use Private_SetSelectionMode / GetSelectionMode / GetSelectionModeAttribute")
+	TSlateDeprecatedTAttribute<ESelectionMode::Type> SelectionMode;
+#endif
 
 	/** Column headers that describe which columns this list shows */
 	TSharedPtr< SHeaderRow > HeaderRow;
@@ -555,12 +605,12 @@ protected:
 	/** Should Use ShdadowBox Styling */
 	bool bShouldUseShadowBoxStyle = false;
 
-	/** Style resource for the Shadow Box when scrolling */
-	const FScrollBoxStyle* ShadowBoxStyle;
-
+#if WITH_EDITORONLY_DATA
 	/** Brush resource representing the background area of the view */
-	FInvalidatableBrushAttribute BackgroundBrush;
-	
+	UE_DEPRECATED(5.8, "Use SetBackgroundBrush / GetBackgroundBrush / GetBackgroundBrushAttribute")
+	FDeprecatedInvalidatableBrushAttribute BackgroundBrush;
+#endif
+
 	/** Does the user need a hint that they can scroll to the start of the list? */
 	FSlateColor GetStartShadowOpacity() const;
 	
@@ -612,6 +662,9 @@ private:
 
 	const float ShadowFadeDistance = 32.0f;
 
+	/** Style resource for the Shadow Box when scrolling */
+	const FScrollBoxStyle* ShadowBoxStyle;
+
 	TArray<float> DeltaTimeQueue;
 	int32 DeltaTimeCircularIndex = 0;
 	TArray<float> SortedDeltaTimes;
@@ -619,6 +672,11 @@ private:
 	// Initialized at -1.0f to detect if the variable has been set and is usable since DeltaTimes cannot be negative.
 	float SmoothDeltaTime = -1.0f;
 
+	TSlateAttribute<ESelectionMode::Type> SelectionModeAttribute;
+	ESelectionMode::Type PreviousSelectionMode;
+
+	TSlateAttribute<FSlateBrushCached> BackgroundBrushAttribute;
+
 	union
 	{
 		// vertical scroll bar is stored in horizontal box and vice versa
diff --git a/Engine/Source/Runtime/Slate/Public/Widgets/Views/STileView.h b/Engine/Source/Runtime/Slate/Public/Widgets/Views/STileView.h
--- a/Engine/Source/Runtime/Slate/Public/Widgets/Views/STileView.h
+++ b/Engine/Source/Runtime/Slate/Public/Widgets/Views/STileView.h
@@ -15,6 +15,16 @@
 #include "Framework/Layout/Overscroll.h"
 #include "Widgets/Views/SListView.h"
 
+
+template <typename ItemType>
+class STileView;
+
+template <typename ItemType>
+struct TWidgetTypeTraits<STileView<ItemType>>
+{
+	static constexpr bool SupportsInvalidation() { return true; }
+};
+
 /**
  * A TileView widget is a list which arranges its items horizontally until there is no more space then creates a new row.
  * Items are spaced evenly horizontally.
@@ -168,7 +178,7 @@ public:
 		this->OnDoubleClick = InArgs._OnMouseButtonDoubleClick;
 		this->OnSelectionChanged = InArgs._OnSelectionChanged;
 		this->OnIsSelectableOrNavigable = InArgs._OnIsSelectableOrNavigable;
-		this->SelectionMode = InArgs._SelectionMode;
+		this->SetSelectionMode(InArgs._SelectionMode);
 
 		this->bClearSelectionOnClick = InArgs._ClearSelectionOnClick;
 
diff --git a/Engine/Source/Runtime/Slate/Public/Widgets/Views/STreeView.h b/Engine/Source/Runtime/Slate/Public/Widgets/Views/STreeView.h
--- a/Engine/Source/Runtime/Slate/Public/Widgets/Views/STreeView.h
+++ b/Engine/Source/Runtime/Slate/Public/Widgets/Views/STreeView.h
@@ -16,6 +16,7 @@
 #include "Widgets/Views/SListView.h"
 #include "Algo/Reverse.h"
 #include "Math/NumericLimits.h"
+#include "Widgets/SWidget.h"
 
 /** Info needed by a (relatively) small fraction of the tree items; some of them may not be visible. */
 struct FSparseItemInfo
@@ -75,6 +76,15 @@ struct FItemInfo
 };
 
 
+template <typename ItemType>
+class STreeView;
+
+template <typename ItemType>
+struct TWidgetTypeTraits<STreeView<ItemType>>
+{
+	static constexpr bool SupportsInvalidation() { return true; }
+};
+
 /**
  * This assumes you are familiar with SListView; see SListView.
  *
@@ -88,6 +98,8 @@ struct FItemInfo
 template<typename ItemType>
 class STreeView : public SListView< ItemType >
 {
+	SLATE_DECLARE_WIDGET(STreeView, SListView<ItemType>)
+	static inline SLATE_IMPLEMENT_WIDGET(STreeView)
 public:
 
 	using NullableItemType  = typename TListTypeTraits< ItemType >::NullableType;
@@ -272,7 +284,7 @@ public:
 		this->OnSelectionChanged = InArgs._OnSelectionChanged;
 		this->OnExpansionChanged = InArgs._OnExpansionChanged;
 		this->OnIsSelectableOrNavigable = InArgs._OnIsSelectableOrNavigable;
-		this->SelectionMode = InArgs._SelectionMode;
+		this->SetSelectionMode(InArgs._SelectionMode);
 
 		this->bClearSelectionOnClick = InArgs._ClearSelectionOnClick;
 		this->ConsumeMouseWheel = InArgs._ConsumeMouseWheel;
@@ -295,14 +307,14 @@ public:
 
 		this->bReturnFocusToSelection = InArgs._ReturnFocusToSelection;
 
-		this->bShouldStackHierarchyHeaders = InArgs._ShouldStackHierarchyHeaders;
+		this->bShouldStackHierarchyHeaders.Assign(*this, InArgs._ShouldStackHierarchyHeaders);
 
 		this->SetStyle(InArgs._TreeViewStyle);
 
 		this->SetShadowBoxStyle(InArgs._ShadowBoxStyle);
 		this->bShouldUseShadowBoxStyle = InArgs._bEnableShadowBoxStyle;
 
-		this->MaxPinnedItems = InArgs._MaxPinnedItems;
+		this->SetMaxPinnedItems(InArgs._MaxPinnedItems);
 		this->DefaultMaxPinnedItems = InArgs._MaxPinnedItems;
 
 		// Check for any parameters that the coder forgot to specify.
@@ -351,6 +363,7 @@ public:
 	STreeView()
 		: SListView< ItemType >( ETableViewMode::Tree )
 		, bTreeItemsAreDirty( true )
+		, bShouldStackHierarchyHeaders( *this, false )
 	{
 		SListView<ItemType>::SetItemsSource(&LinearizedItems);
 	}
@@ -779,7 +792,7 @@ public:
 		if (Results.bGeneratedPastLastItem && Results.NewScrollOffset > 0)
 		{
 			int32 LastItem = InItemsSource.Num() - 1;
-			int32 CurrentMaxPinnedItems = this->MaxPinnedItems.Get();
+			int32 CurrentMaxPinnedItems = this->GetMaxPinnedItems();
 
 			// Could be different than reported by STableViewBase if some items are collapsed
 			int32 NumPinnedItems = (this->GetNumPinnedItems() < CurrentMaxPinnedItems) ? this->GetNumPinnedItems() : CurrentMaxPinnedItems;
@@ -887,12 +900,6 @@ public:
 		SListView<ItemType>::RebuildList();
 	}
 
-	void SetStyle(const FTableViewStyle* InStyle)
-	{
-		Style = InStyle;
-		STableViewBase::SetBackgroundBrush( Style != nullptr ? &Style->BackgroundBrush : FStyleDefaults::GetNoBrush() );
-	}
-
 	/**
 	 * Set whether some data item is expanded or not.
 	 * 
@@ -1099,9 +1106,6 @@ protected:
 	/** The delegate that is invoked whenever an item in the tree is expanded or collapsed. */
 	FOnExpansionChanged OnExpansionChanged;
 
-	/** Style resource for the tree */
-	const FTableViewStyle* Style;
-
 private:		
 	/** Pointer to the source data that we are observing */
 	TUniquePtr<UE::Slate::ItemsSource::IItemsSource<ItemType>> TreeViewSource;
@@ -1116,5 +1120,19 @@ private:
 	bool bHighlightParentNodesForSelection = false;
 
 	/** true if we want to show the hierarchy of items pinned at the top */
-	TAttribute<bool> bShouldStackHierarchyHeaders = false;
+	SWidget::TSlateAttribute<bool> bShouldStackHierarchyHeaders;
 };
+
+template<typename ItemType>
+inline void STreeView<ItemType>::PrivateRegisterAttributes(FSlateAttributeInitializer& AttributeInitializer)
+{
+	struct FInvalidation
+	{
+		static void RequestLayoutRefresh(SWidget& OwningWidget)
+		{
+			static_cast<STreeView<ItemType>&>(OwningWidget).RequestLayoutRefresh();
+		}
+	};
+	SLATE_ADD_MEMBER_ATTRIBUTE_DEFINITION(AttributeInitializer, bShouldStackHierarchyHeaders, EInvalidateWidgetReason::Layout)
+		.OnValueChanged(FSlateAttributeDescriptor::FAttributeValueChangedDelegate::CreateStatic(FInvalidation::RequestLayoutRefresh));
+}
\ No newline at end of file
diff --git a/Engine/Source/Runtime/SlateCore/Public/Styling/SlateTypes.h b/Engine/Source/Runtime/SlateCore/Public/Styling/SlateTypes.h
--- a/Engine/Source/Runtime/SlateCore/Public/Styling/SlateTypes.h
+++ b/Engine/Source/Runtime/SlateCore/Public/Styling/SlateTypes.h
@@ -2273,6 +2273,62 @@ private:
 	FSlateBrush ImageCache;
 };
 
+/** A structure used to help the user identify deprecated FInvalidatableBrushAttribute that are now TSlateAttribute. */
+class FDeprecatedInvalidatableBrushAttribute
+{
+public:
+	FDeprecatedInvalidatableBrushAttribute() { }
+	FDeprecatedInvalidatableBrushAttribute(const TAttribute< const FSlateBrush* >& InImage) {}
+
+	bool IsBound() const { return false; }
+	const FSlateBrush* Get() const { return nullptr; }
+	TAttribute< const FSlateBrush* > GetImage() const { return {}; }
+	void SetImage(SWidget& ThisWidget, const TAttribute< const FSlateBrush* >& InImage) {}
+};
+
+/**
+ * Class similar to FInvalidatableBrushAttribute but to be used with TSlateAttribute.
+ * Note that this class does not invalidate on value change unless used inside a TSlateAttribute
+ */
+class FSlateBrushCached
+{
+public:
+	FSlateBrushCached() : Image(nullptr), ImageCache(FSlateBrush()) {}
+	FSlateBrushCached(const FSlateBrush* InImage)
+		: Image(InImage), ImageCache(InImage ? *InImage : FSlateBrush())
+	{
+	}
+
+	const FSlateBrush* GetImage() const
+	{
+		return Image;
+	}
+	void SetImage(const FSlateBrush* InImage)
+	{
+		Image = InImage;
+		ImageCache = InImage ? *InImage : FSlateBrush();
+	}
+	
+	bool operator==(const FSlateBrushCached& Rhs) const
+	{
+		return Image == Rhs.Image;
+	}
+	bool operator==(const FSlateBrush* Rhs) const
+	{
+		if (Image == Rhs)
+		{
+			const FSlateBrush NewImageToCache = Image ? *Image : FSlateBrush();
+			return ImageCache == NewImageToCache;
+		}
+		return false;
+	}
+private:
+	/** The slate brush to draw for the image. */
+	const FSlateBrush* Image;
+
+	/** The copy of the image data, some users reuse the same FSlateBrush pointer, so we need to check it against the last true data to see what changed. */
+	FSlateBrush ImageCache;
+};
 
 /** HACK: We need a UClass here or UHT will complain. */
 UCLASS()
