Fix MenuAnchor Placement AboveRightAnchor.

Currently the Menu would be located in the wrong location with this placement.

diff --git a/Engine/Source/Runtime/Slate/Private/Widgets/Input/SMenuAnchor.cpp b/Engine/Source/Runtime/Slate/Private/Widgets/Input/SMenuAnchor.cpp
--- a/Engine/Source/Runtime/Slate/Private/Widgets/Input/SMenuAnchor.cpp
+++ b/Engine/Source/Runtime/Slate/Private/Widgets/Input/SMenuAnchor.cpp
@@ -43,7 +43,7 @@ static FVector2D GetMenuOffsetForPlacement(const FGeometry& AllottedGeometry, EM
 		case MenuPlacement_CenteredAboveAnchor:
 			return FVector2D(-((PopupSizeLocalSpace.X / 2) - (AllottedGeometry.GetLocalSize().X / 2)), -PopupSizeLocalSpace.Y);
 		case MenuPlacement_AboveRightAnchor:
-			return FVector2D( -( PopupSizeLocalSpace.X ) - ( AllottedGeometry.GetLocalSize().X ), -PopupSizeLocalSpace.Y );
+			return FVector2D( -( PopupSizeLocalSpace.X ) + ( AllottedGeometry.GetLocalSize().X ), -PopupSizeLocalSpace.Y );
 		case MenuPlacement_MenuLeft:
 			return FVector2D(-PopupSizeLocalSpace.X, 0.0f);
 		case MenuPlacement_Center:
