Slate Dynamic Invalidation - SExpandable Area

diff --git a/Engine/Source/Runtime/Slate/Private/Widgets/Layout/SExpandableArea.cpp b/Engine/Source/Runtime/Slate/Private/Widgets/Layout/SExpandableArea.cpp
--- a/Engine/Source/Runtime/Slate/Private/Widgets/Layout/SExpandableArea.cpp
+++ b/Engine/Source/Runtime/Slate/Private/Widgets/Layout/SExpandableArea.cpp
@@ -15,21 +15,11 @@ void SExpandableArea::Construct( const FArguments& InArgs )
 {
 	check(InArgs._Style);
 
-	WidgetStyle = InArgs._Style;
 	bAreaCollapsed = InArgs._InitiallyCollapsed;
 	MinWidth = InArgs._MinWidth;
 	MaxHeight = InArgs._MaxHeight;
 	OnAreaExpansionChanged = InArgs._OnAreaExpansionChanged;
-	CollapsedImage = &InArgs._Style->CollapsedImage;
-	ExpandedImage = &InArgs._Style->ExpandedImage;
-
-	// If it should be initially visible, show it now
-	RolloutCurve = FCurveSequence(0.0f, InArgs._Style->RolloutAnimationSeconds, ECurveEaseFunction::CubicOut);
-
-	if (!bAreaCollapsed)
-	{
-		RolloutCurve.JumpToEnd();
-	}
+	SetStyle(InArgs._Style);
 
 	TSharedRef<SWidget> HeaderContent = InArgs._HeaderContent.Widget;
 	if( HeaderContent == SNullWidget::NullWidget )
@@ -238,7 +228,19 @@ void SExpandableArea::SetMaxHeight(float InMaxHeight)
 
 void SExpandableArea::SetStyle(const FExpandableAreaStyle* InStyle)
 {
+	check(InStyle);
 	WidgetStyle = InStyle;
+	CollapsedImage = &WidgetStyle->CollapsedImage;
+	ExpandedImage = &WidgetStyle->ExpandedImage;
+	RolloutCurve = FCurveSequence(0.0f, WidgetStyle->RolloutAnimationSeconds, ECurveEaseFunction::CubicOut);
+
+	// If it should be visible, show it now
+	if (!bAreaCollapsed)
+	{
+		RolloutCurve.JumpToEnd();
+	}
+
+	Invalidate(EInvalidateWidgetReason::Layout);
 }
 
 void SExpandableArea::SetBorderBrush(const FSlateBrush* InBorderBrush)
diff --git a/Engine/Source/Runtime/Slate/Public/Widgets/Layout/SExpandableArea.h b/Engine/Source/Runtime/Slate/Public/Widgets/Layout/SExpandableArea.h
--- a/Engine/Source/Runtime/Slate/Public/Widgets/Layout/SExpandableArea.h
+++ b/Engine/Source/Runtime/Slate/Public/Widgets/Layout/SExpandableArea.h
@@ -20,6 +20,12 @@
 class SBorder;
 class SButton;
 
+template<>
+struct TWidgetTypeTraits<class SExpandableArea>
+{
+	static constexpr bool SupportsInvalidation() { return true; }
+};
+
 /**
  * Represents an expandable area of content            
  */
@@ -139,11 +145,13 @@ public:
 	/**
 	 * Invalidates the style
 	 */
+	UE_DEPRECATED(5.8, "Changing the style needs to update other properties. Use SetStyle if the style has changed")
 	SLATE_API void InvalidateStyle();
 
 	/**
 	 * Invalidates the border brush
 	 */
+	UE_DEPRECATED(5.8, "No need to call this function anymore, use SetBorderBrush")
 	SLATE_API void InvalidateBorderBrush();
 
 	/**
diff --git a/Engine/Source/Runtime/UMG/Private/Components/ExpandableArea.cpp b/Engine/Source/Runtime/UMG/Private/Components/ExpandableArea.cpp
--- a/Engine/Source/Runtime/UMG/Private/Components/ExpandableArea.cpp
+++ b/Engine/Source/Runtime/UMG/Private/Components/ExpandableArea.cpp
@@ -136,7 +136,7 @@ void UExpandableArea::SetStyle(const FExpandableAreaStyle& InStyle)
 	Style = InStyle;
 	if (MyExpandableArea.IsValid())
 	{
-		MyExpandableArea->InvalidateStyle();
+		MyExpandableArea->SetStyle(&Style);
 	}
 }
 
@@ -249,12 +249,10 @@ void UExpandableArea::SynchronizeProperties()
 	}
 	
 	MyExpandableArea->SetStyle(&Style);
-	MyExpandableArea->InvalidateStyle();
 	MyExpandableArea->SetExpanded(bIsExpanded);
 	MyExpandableArea->SetAreaPadding(AreaPadding);
 	MyExpandableArea->SetHeaderPadding(HeaderPadding);
 	MyExpandableArea->SetBorderBrush(&BorderBrush);
-	MyExpandableArea->InvalidateBorderBrush();
 	MyExpandableArea->SetBorderBackgroundColor(BorderColor);
 	MyExpandableArea->SetMaxHeight(MaxHeight);
 
