Fix Getter name for IsFakeTouch

diff --git a/Engine/Source/Runtime/Slate/Private/Framework/Application/SlateApplication.cpp b/Engine/Source/Runtime/Slate/Private/Framework/Application/SlateApplication.cpp
--- a/Engine/Source/Runtime/Slate/Private/Framework/Application/SlateApplication.cpp
+++ b/Engine/Source/Runtime/Slate/Private/Framework/Application/SlateApplication.cpp
@@ -4083,7 +4083,7 @@ void FSlateApplication::SetIsFakeTouch(bool bInIsFakingTouch)
 	bIsFakingTouch = bInIsFakingTouch;
 }
 
-bool FSlateApplication::SetIsFakeTouch() const
+bool FSlateApplication::GetIsFakeTouch() const
 {
 	return bIsFakingTouch;
 }
diff --git a/Engine/Source/Runtime/Slate/Public/Framework/Application/SlateApplication.h b/Engine/Source/Runtime/Slate/Public/Framework/Application/SlateApplication.h
--- a/Engine/Source/Runtime/Slate/Public/Framework/Application/SlateApplication.h
+++ b/Engine/Source/Runtime/Slate/Public/Framework/Application/SlateApplication.h
@@ -1564,7 +1564,7 @@ public:
 	
 	/** Setter and getter for the touch simulation */
 	SLATE_API void SetIsFakeTouch(bool bInIsFakingTouch);
-	SLATE_API bool SetIsFakeTouch() const;
+	SLATE_API bool GetIsFakeTouch() const;
 
 	bool IsRenderingOffScreen() const { return bRenderOffScreen; }
 
