Slate Dynamic Invalidation - SCheckbox

diff --git a/Engine/Source/Runtime/Slate/Private/Framework/MultiBox/SToolBarStackButtonBlock.cpp b/Engine/Source/Runtime/Slate/Private/Framework/MultiBox/SToolBarStackButtonBlock.cpp
--- a/Engine/Source/Runtime/Slate/Private/Framework/MultiBox/SToolBarStackButtonBlock.cpp
+++ b/Engine/Source/Runtime/Slate/Private/Framework/MultiBox/SToolBarStackButtonBlock.cpp
@@ -67,7 +67,7 @@ FReply SCheckBoxStack::OnMouseButtonDown( const FGeometry& MyGeometry, const FPo
 		if(InputClickMethod == EButtonClickMethod::MouseDown )
 		{
 			ToggleCheckedState();
-			const ECheckBoxState State = IsCheckboxChecked.Get();
+			const ECheckBoxState State = GetCheckedState();
 			if(State == ECheckBoxState::Checked)
 			{
 				PlayCheckedSound();
diff --git a/Engine/Source/Runtime/Slate/Private/Widgets/Input/SCheckBox.cpp b/Engine/Source/Runtime/Slate/Private/Widgets/Input/SCheckBox.cpp
--- a/Engine/Source/Runtime/Slate/Private/Widgets/Input/SCheckBox.cpp
+++ b/Engine/Source/Runtime/Slate/Private/Widgets/Input/SCheckBox.cpp
@@ -12,7 +12,20 @@
 #include "Widgets/Accessibility/SlateAccessibleMessageHandler.h"
 #endif
 
+SLATE_IMPLEMENT_WIDGET(SCheckBox)
+void SCheckBox::PrivateRegisterAttributes(FSlateAttributeInitializer& AttributeInitializer)
+{
+	SLATE_ADD_MEMBER_ATTRIBUTE_DEFINITION(AttributeInitializer, IsCheckboxCheckedAttribute, EInvalidateWidgetReason::None);
+	SLATE_ADD_MEMBER_ATTRIBUTE_DEFINITION(AttributeInitializer, PaddingOverrideAttribute, EInvalidateWidgetReason::None);
+	SLATE_ADD_MEMBER_ATTRIBUTE_DEFINITION(AttributeInitializer, ForegroundColorOverrideAttribute, EInvalidateWidgetReason::None);
+	SLATE_ADD_MEMBER_ATTRIBUTE_DEFINITION(AttributeInitializer, BorderBackgroundColorOverrideAttribute, EInvalidateWidgetReason::None);
+}
+
 SCheckBox::SCheckBox()
+	: IsCheckboxCheckedAttribute(*this, ECheckBoxState::Unchecked)
+	, PaddingOverrideAttribute(*this)
+	, ForegroundColorOverrideAttribute(*this, FSlateColor::UseStyle())
+	, BorderBackgroundColorOverrideAttribute(*this)
 {
 #if WITH_ACCESSIBILITY
 	AccessibleBehavior = EAccessibleBehavior::Summary;
@@ -47,10 +60,10 @@ void SCheckBox::Construct( const SCheckBox::FArguments& InArgs )
 	BackgroundImage = InArgs._BackgroundImage;
 	BackgroundHoveredImage = InArgs._BackgroundHoveredImage;
 	BackgroundPressedImage = InArgs._BackgroundPressedImage;
-	
-	PaddingOverride = InArgs._Padding;
-	ForegroundColorOverride = InArgs._ForegroundColor;
-	BorderBackgroundColorOverride = InArgs._BorderBackgroundColor;
+
+	SetPaddingOverride(InArgs._Padding);
+	SetForegroundColorOverride(InArgs._ForegroundColor);
+	SetBorderBackgroundColorOverride(InArgs._BorderBackgroundColor);
 	CheckBoxTypeOverride = InArgs._Type;
 
 	HorizontalAlignment = InArgs._HAlign;
@@ -62,7 +75,7 @@ void SCheckBox::Construct( const SCheckBox::FArguments& InArgs )
 
 	BuildCheckBox(InArgs._Content.Widget);
 
-	IsCheckboxChecked = InArgs._IsChecked;
+	SetIsChecked(InArgs._IsChecked);
 	OnCheckStateChanged = InArgs._OnCheckStateChanged;
 
 	ClickMethod = InArgs._ClickMethod;
@@ -71,9 +84,9 @@ void SCheckBox::Construct( const SCheckBox::FArguments& InArgs )
 
 	OnGetMenuContent = InArgs._OnGetMenuContent;
 
-	HoveredSound = InArgs._HoveredSoundOverride.Get(InArgs._Style->HoveredSlateSound);
-	CheckedSound = InArgs._CheckedSoundOverride.Get(InArgs._Style->CheckedSlateSound);
-	UncheckedSound = InArgs._UncheckedSoundOverride.Get(InArgs._Style->UncheckedSlateSound);
+	SetHoveredSound(InArgs._HoveredSoundOverride);
+	SetCheckedSound(InArgs._CheckedSoundOverride);
+	SetUncheckedSound(InArgs._UncheckedSoundOverride);
 }
 
 /**
@@ -97,7 +110,7 @@ FReply SCheckBox::OnKeyDown(const FGeometry& MyGeometry, const FKeyEvent& InKeyE
 		{
 			ToggleCheckedState();
 
-			const ECheckBoxState State = IsCheckboxChecked.Get();
+			const ECheckBoxState State = GetCheckedState();
 			if (State == ECheckBoxState::Checked)
 			{
 				PlayCheckedSound();
@@ -126,7 +139,7 @@ FReply SCheckBox::OnKeyUp( const FGeometry& MyGeometry, const FKeyEvent& InKeyEv
 		{
 			ToggleCheckedState();
 
-			const ECheckBoxState State = IsCheckboxChecked.Get();
+			const ECheckBoxState State = GetCheckedState();
 			if (State == ECheckBoxState::Checked)
 			{
 				PlayCheckedSound();
@@ -162,7 +175,7 @@ FReply SCheckBox::OnMouseButtonDown( const FGeometry& MyGeometry, const FPointer
 		if(InputClickMethod == EButtonClickMethod::MouseDown )
 		{
 			ToggleCheckedState();
-			const ECheckBoxState State = IsCheckboxChecked.Get();
+			const ECheckBoxState State = GetCheckedState();
 			if(State == ECheckBoxState::Checked)
 			{
 				PlayCheckedSound();
@@ -209,9 +222,9 @@ FReply SCheckBox::OnMouseButtonDown( const FGeometry& MyGeometry, const FPointer
  *
  * @return Whether the event was handled along with possible requests for the system to take action.
  */
-FReply SCheckBox::OnMouseButtonDoubleClick( const FGeometry& InMyGeometry, const FPointerEvent& InMouseEvent )
+FReply SCheckBox::OnMouseButtonDoubleClick( const FGeometry& MyGeometry, const FPointerEvent& MouseEvent )
 {
-	return OnMouseButtonDown( InMyGeometry, InMouseEvent );
+	return OnMouseButtonDown( MyGeometry, MouseEvent );
 }
 
 /**
@@ -248,7 +261,7 @@ FReply SCheckBox::OnMouseButtonUp( const FGeometry& MyGeometry, const FPointerEv
 				if(InputClickMethod == EButtonClickMethod::MouseUp || HasMouseCapture() )
 				{
 					ToggleCheckedState();
-					const ECheckBoxState State = IsCheckboxChecked.Get();
+					const ECheckBoxState State = GetCheckedState();
 					if(State == ECheckBoxState::Checked)
 					{
 						PlayCheckedSound();
@@ -295,11 +308,11 @@ bool SCheckBox::IsInteractable() const
 
 FSlateColor SCheckBox::GetForegroundColor() const
 {
-	FSlateColor UserColor = ForegroundColorOverride.Get();
+	FSlateColor UserColor = ForegroundColorOverrideAttribute.Get();
 
 	if (UserColor == FSlateColor::UseStyle())
 	{
-		ECheckBoxState State = IsCheckboxChecked.Get();
+		ECheckBoxState State = GetCheckedState();
 
 		switch (State)
 		{
@@ -322,7 +335,7 @@ FSlateColor SCheckBox::GetForegroundColor() const
  */
 const FSlateBrush* SCheckBox::OnGetCheckImage() const
 {
-	ECheckBoxState State = IsCheckboxChecked.Get();
+	ECheckBoxState State = GetCheckedState();
 
 	const FSlateBrush* ImageToUse;
 	switch( State )
@@ -352,7 +365,7 @@ const FSlateBrush* SCheckBox::OnGetBackgroundImage() const
 
 ECheckBoxState SCheckBox::GetCheckedState() const
 {
-	return IsCheckboxChecked.Get();
+	return IsCheckboxCheckedAttribute.Get();
 }
 
 /**
@@ -360,15 +373,15 @@ ECheckBoxState SCheckBox::GetCheckedState() const
  */
 void SCheckBox::ToggleCheckedState()
 {
-	const ECheckBoxState State = IsCheckboxChecked.Get();
+	const ECheckBoxState State = GetCheckedState();
 
-	// If the current check box state is checked OR undetermined we set the check box to checked.
+	// If the current check box state is checked OR undetermined we set the check box to unchecked.
 	if( State == ECheckBoxState::Checked || State == ECheckBoxState::Undetermined )
 	{
-		if ( !IsCheckboxChecked.IsBound() )
+		if ( !IsCheckboxCheckedAttribute.IsBound(*this) )
 		{
 			// When we are not bound, just toggle the current state.
-			IsCheckboxChecked.Set( ECheckBoxState::Unchecked );
+			SetIsChecked( ECheckBoxState::Unchecked );
 		}
 
 		// The state of the check box changed.  Execute the delegate to notify users
@@ -376,10 +389,10 @@ void SCheckBox::ToggleCheckedState()
 	}
 	else if( State == ECheckBoxState::Unchecked )
 	{
-		if ( !IsCheckboxChecked.IsBound() )
+		if ( !IsCheckboxCheckedAttribute.IsBound(*this) )
 		{
 			// When we are not bound, just toggle the current state.
-			IsCheckboxChecked.Set( ECheckBoxState::Checked );
+			SetIsChecked( ECheckBoxState::Checked );
 		}
 
 		// The state of the check box changed.  Execute the delegate to notify users
@@ -388,13 +401,19 @@ void SCheckBox::ToggleCheckedState()
 
 #if WITH_ACCESSIBILITY
 	// @TODOAccessibility: Technically we should pass the Id of the user that toggled the checkbox, but we don't want to change the Slate API as much as possible
-	FSlateApplicationBase::Get().GetAccessibleMessageHandler()->OnWidgetEventRaised(FSlateAccessibleMessageHandler::FSlateWidgetAccessibleEventArgs(AsShared(), EAccessibleEvent::Activate, State == ECheckBoxState::Checked, IsCheckboxChecked.Get() == ECheckBoxState::Checked));
+	FSlateApplicationBase::Get().GetAccessibleMessageHandler()->OnWidgetEventRaised(
+		FSlateAccessibleMessageHandler::FSlateWidgetAccessibleEventArgs(
+			AsShared(),
+			EAccessibleEvent::Activate,
+			State == ECheckBoxState::Checked,
+			IsChecked()
+		));
 #endif
 }
 
 void SCheckBox::SetIsChecked(TAttribute<ECheckBoxState> InIsChecked)
 {
-	IsCheckboxChecked = InIsChecked;
+	IsCheckboxCheckedAttribute.Assign(*this, MoveTemp(InIsChecked));
 }
 
 TEnumAsByte<EButtonClickMethod::Type> SCheckBox::GetClickMethodFromInputType(const FPointerEvent& MouseEvent) const
@@ -422,17 +441,17 @@ bool SCheckBox::IsPreciseTapOrClick(const FPointerEvent& MouseEvent) const
 
 void SCheckBox::PlayCheckedSound() const
 {
-	FSlateApplication::Get().PlaySound( CheckedSound );
+	FSlateApplication::Get().PlaySound( GetCheckedSound() );
 }
 
 void SCheckBox::PlayUncheckedSound() const
 {
-	FSlateApplication::Get().PlaySound( UncheckedSound );
+	FSlateApplication::Get().PlaySound( GetUncheckedSound() );
 }
 
 void SCheckBox::PlayHoverSound() const
 {
-	FSlateApplication::Get().PlaySound( HoveredSound );
+	FSlateApplication::Get().PlaySound( GetHoveredSound() );
 }
 
 void SCheckBox::SetContent(const TSharedRef< SWidget >& InContent)
@@ -544,7 +563,7 @@ void SCheckBox::BuildCheckBox(TSharedRef<SWidget> InContent)
 				.BorderImage(FStyleDefaults::GetNoBrush())
 				.Padding(0.0f)
 				[
-					InContent
+					MoveTemp(InContent)
 				]
 			]
 		];
@@ -564,7 +583,7 @@ void SCheckBox::BuildCheckBox(TSharedRef<SWidget> InContent)
 			.BorderBackgroundColor(this, &SCheckBox::OnGetBorderBackgroundColor)
 			.HAlign(HorizontalAlignment)
 			[
-				InContent
+				MoveTemp(InContent)
 			]
 		];
 	}
@@ -573,12 +592,12 @@ void SCheckBox::BuildCheckBox(TSharedRef<SWidget> InContent)
 
 FMargin SCheckBox::OnGetPadding() const
 {
-	return PaddingOverride.IsSet() ? PaddingOverride.Get() : Style->Padding;
+	return bIsPaddingAttrSet ? PaddingOverrideAttribute.Get() : Style->Padding;
 }
 
 FSlateColor SCheckBox::OnGetBorderBackgroundColor() const
 {
-	return BorderBackgroundColorOverride.IsSet() ? BorderBackgroundColorOverride.Get() : Style->BorderBackgroundColor;
+	return bIsBorderBackgroundColorOverrideAttrSet ? BorderBackgroundColorOverrideAttribute.Get() : Style->BorderBackgroundColor;
 }
 
 ESlateCheckBoxType::Type SCheckBox::OnGetCheckBoxType() const
@@ -646,6 +665,53 @@ const FSlateBrush* SCheckBox::GetBackgroundPressedImage() const
 	return BackgroundPressedImage ? BackgroundPressedImage : &Style->BackgroundPressedImage;
 }
 
+void SCheckBox::SetPaddingOverride(TAttribute<FMargin> InPaddingOverride)
+{
+	bIsPaddingAttrSet = InPaddingOverride.IsSet();
+	PaddingOverrideAttribute.Assign(*this, MoveTemp(InPaddingOverride));
+}
+
+void SCheckBox::SetForegroundColorOverride(TAttribute<FSlateColor> InForegroundColorOverride)
+{
+	ForegroundColorOverrideAttribute.Assign(*this, MoveTemp(InForegroundColorOverride));
+}
+
+void SCheckBox::SetBorderBackgroundColorOverride(TAttribute<FSlateColor> InBorderBackgroundColorOverride)
+{
+	bIsBorderBackgroundColorOverrideAttrSet = InBorderBackgroundColorOverride.IsSet();
+	BorderBackgroundColorOverrideAttribute.Assign(*this, MoveTemp(InBorderBackgroundColorOverride));
+}
+
+void SCheckBox::SetHoveredSound(TOptional<FSlateSound> InHoveredSound)
+{
+	HoveredSoundOverride = MoveTemp(InHoveredSound);
+}
+
+FSlateSound SCheckBox::GetHoveredSound() const
+{
+	return HoveredSoundOverride.Get(Style->HoveredSlateSound);
+}
+
+void SCheckBox::SetCheckedSound(TOptional<FSlateSound> InCheckedSound)
+{
+	CheckedSoundOverride = MoveTemp(InCheckedSound);
+}
+
+FSlateSound SCheckBox::GetCheckedSound() const
+{
+	return CheckedSoundOverride.Get(Style->CheckedSlateSound);
+}
+
+void SCheckBox::SetUncheckedSound(TOptional<FSlateSound> InUncheckedSound)
+{
+	UncheckedSoundOverride = MoveTemp(InUncheckedSound);
+}
+
+FSlateSound SCheckBox::GetUncheckedSound() const
+{
+	return UncheckedSoundOverride.Get(Style->UncheckedSlateSound);
+}
+
 void SCheckBox::SetClickMethod(EButtonClickMethod::Type InClickMethod)
 {
 	ClickMethod = InClickMethod;
diff --git a/Engine/Source/Runtime/Slate/Public/Widgets/Input/SCheckBox.h b/Engine/Source/Runtime/Slate/Public/Widgets/Input/SCheckBox.h
--- a/Engine/Source/Runtime/Slate/Public/Widgets/Input/SCheckBox.h
+++ b/Engine/Source/Runtime/Slate/Public/Widgets/Input/SCheckBox.h
@@ -22,12 +22,18 @@ class SBorder;
 /** Delegate that is executed when the check box state changes */
 DECLARE_DELEGATE_OneParam( FOnCheckStateChanged, ECheckBoxState );
 
+template<>
+struct TWidgetTypeTraits<class SCheckBox>
+{
+	static constexpr bool SupportsInvalidation() { return true; }
+};
 
 /**
  * Check box Slate control
  */
 class SCheckBox : public SCompoundWidget
 {
+	SLATE_DECLARE_WIDGET_API(SCheckBox, SCompoundWidget, SLATE_API)
 public:
 
 	SLATE_BEGIN_ARGS( SCheckBox )
@@ -177,7 +183,7 @@ public:
 	SLATE_API virtual void OnMouseEnter( const FGeometry& MyGeometry, const FPointerEvent& MouseEvent ) override;
 	SLATE_API virtual void OnMouseLeave( const FPointerEvent& MouseEvent ) override;
 	SLATE_API virtual bool IsInteractable() const override;
-	SLATE_API virtual FSlateColor GetForegroundColor() const;
+	SLATE_API virtual FSlateColor GetForegroundColor() const override;
 #if WITH_ACCESSIBILITY
 	SLATE_API virtual TSharedRef<FSlateAccessibleWidget> CreateAccessibleWidget() override;
 #endif
@@ -190,7 +196,7 @@ public:
 	 */
 	bool IsChecked() const
 	{
-		return ( IsCheckboxChecked.Get() == ECheckBoxState::Checked );
+		return ( GetCheckedState() == ECheckBoxState::Checked );
 	}
 
 	/** @return The current checked state of the checkbox. */
@@ -282,7 +288,42 @@ protected:
 	SLATE_API const FSlateBrush* GetBackgroundHoveredImage() const;
 	SLATE_API const FSlateBrush* GetBackgroundPressedImage() const;
 
-	
+protected:
+	TSlateAttributeRef<ECheckBoxState> GetIsCheckboxCheckedAttribute() const
+	{
+		return TSlateAttributeRef<ECheckBoxState>{SharedThis(this), IsCheckboxCheckedAttribute};
+	}
+
+	SLATE_API void SetPaddingOverride(TAttribute<FMargin> InPaddingOverride);
+	TSlateAttributeRef<FMargin> GetPaddingOverrideAttribute() const
+	{
+		return TSlateAttributeRef<FMargin>{SharedThis(this), PaddingOverrideAttribute};
+	}
+
+	SLATE_API void SetForegroundColorOverride(TAttribute<FSlateColor> InForegroundColorOverride);
+	TSlateAttributeRef<FSlateColor> GetForegroundColorOverrideAttribute() const
+	{
+		return TSlateAttributeRef<FSlateColor>{SharedThis(this), ForegroundColorOverrideAttribute};
+	}
+
+	SLATE_API void SetBorderBackgroundColorOverride(TAttribute<FSlateColor> InBorderBackgroundColorOverride);
+	TSlateAttributeRef<FSlateColor> GetBorderBackgroundColorOverrideAttribute() const
+	{
+		return TSlateAttributeRef<FSlateColor>{SharedThis(this), BorderBackgroundColorOverrideAttribute};
+	}
+
+	/** The Sound to play when the check box is hovered  */
+	void SetHoveredSound(TOptional<FSlateSound> InHoveredSound);
+	FSlateSound GetHoveredSound() const;
+
+	/** The Sound to play when the check box is checked */
+	void SetCheckedSound(TOptional<FSlateSound> InCheckedSound);
+	FSlateSound GetCheckedSound() const;
+
+	/** The Sound to play when the check box is unchecked */
+	void SetUncheckedSound(TOptional<FSlateSound> InUncheckedSound);
+	FSlateSound GetUncheckedSound() const;
+
 protected:
 	
 	const FCheckBoxStyle* Style;
@@ -290,8 +331,11 @@ protected:
 	/** True if this check box is currently in a pressed state */
 	bool bIsPressed;
 
+#if WITH_EDITORONLY_DATA
 	/** Are we checked */
-	TAttribute<ECheckBoxState> IsCheckboxChecked;
+	UE_DEPRECATED(5.8, "Use SetIsChecked / GetCheckedState / GetIsCheckboxCheckedAttribute")
+	TSlateDeprecatedTAttribute<ECheckBoxState> IsCheckboxChecked;
+#endif
 
 	/** Delegate called when the check box changes state */
 	FOnCheckStateChanged OnCheckStateChanged;
@@ -321,12 +365,17 @@ protected:
 	/** Image to use for the checkbox background when pressed*/
 	const FSlateBrush* BackgroundPressedImage;
 
+#if WITH_EDITORONLY_DATA
 	/** Overrides padding in the widget style, if set */
-	TAttribute<FMargin> PaddingOverride;
+	UE_DEPRECATED(5.8, "Use SetPaddingOverride / OnGetPadding / GetPaddingOverrideAttribute")
+	TSlateDeprecatedTAttribute<FMargin> PaddingOverride;
 	/** Overrides foreground color in the widget style, if set */
-	TAttribute<FSlateColor> ForegroundColorOverride;
+	UE_DEPRECATED(5.8, "Use SetForegroundColorOverride / GetForegroundColor / GetForegroundColorOverrideAttribute")
+	TSlateDeprecatedTAttribute<FSlateColor> ForegroundColorOverride;
 	/** Overrides border background color in the widget style, if set */
-	TAttribute<FSlateColor> BorderBackgroundColorOverride;
+	UE_DEPRECATED(5.8, "Use SetBorderBackgroundColorOverride / OnGetBorderBackgroundColor / GetBorderBackgroundColorOverrideAttribute")
+	TSlateDeprecatedTAttribute<FSlateColor> BorderBackgroundColorOverride;
+#endif
 	/** Overrides checkbox type in the widget style, if set */
 	TOptional<ESlateCheckBoxType::Type> CheckBoxTypeOverride;
 
@@ -369,16 +418,44 @@ protected:
 	/** Utility function to determine if the incoming mouse event is for a precise tap or click */
 	SLATE_API bool IsPreciseTapOrClick(const FPointerEvent& MouseEvent) const;
 
+#if WITH_EDITORONLY_DATA
 	/** The Sound to play when the check box is hovered  */
+	UE_DEPRECATED(5.8, "Use GetHoveredSound / SetHoveredSound")
 	FSlateSound HoveredSound;
 
 	/** The Sound to play when the check box is checked */
+	UE_DEPRECATED(5.8, "Use GetCheckedSound / SetCheckedSound")
 	FSlateSound CheckedSound;
 
 	/** The Sound to play when the check box is unchecked */
+	UE_DEPRECATED(5.8, "Use GetUncheckedSound / SetUncheckedSound")
 	FSlateSound UncheckedSound;
+#endif
 
 protected:
 	/** When in toggle button mode, this will hold the pointer to the toggle button's border */
 	TSharedPtr<SBorder> ContentContainer;
+
+private:
+	/** Are we checked */
+	TSlateAttribute<ECheckBoxState> IsCheckboxCheckedAttribute;
+
+	/** Overrides padding in the widget style, if set */
+	TSlateAttribute<FMargin> PaddingOverrideAttribute;
+	bool bIsPaddingAttrSet = false;
+	
+	/** Overrides foreground color in the widget style, if set */
+	TSlateAttribute<FSlateColor> ForegroundColorOverrideAttribute;
+	/** Overrides border background color in the widget style, if set */
+	TSlateAttribute<FSlateColor> BorderBackgroundColorOverrideAttribute;
+	bool bIsBorderBackgroundColorOverrideAttrSet = false;
+
+	/** The Sound to play when the check box is hovered  */
+	TOptional<FSlateSound> HoveredSoundOverride;
+
+	/** The Sound to play when the check box is checked */
+	TOptional<FSlateSound> CheckedSoundOverride;
+
+	/** The Sound to play when the check box is unchecked */
+	TOptional<FSlateSound> UncheckedSoundOverride;
 };
