Slate Dynamic Invalidation - Revert `SCheckBox` to use a `TAttribute` instead of a `TSlateAttribute` for `IsCheckboxChecked`

There was an issue with MetaSound where the value of `GetCheckedState` was read just after the widget was created. At this point, the `TSlateAttribute` had never been set (as `Assign` does not set the value by default) and was holding its default value of `ECheckBoxState::Unchecked`. This lead to the following code to update its value based on the checkbox being unchecked even though it was just checked.

diff --git a/Engine/Source/Runtime/Slate/Private/Widgets/Input/SCheckBox.cpp b/Engine/Source/Runtime/Slate/Private/Widgets/Input/SCheckBox.cpp
--- a/Engine/Source/Runtime/Slate/Private/Widgets/Input/SCheckBox.cpp
+++ b/Engine/Source/Runtime/Slate/Private/Widgets/Input/SCheckBox.cpp
@@ -15,15 +15,13 @@
 SLATE_IMPLEMENT_WIDGET(SCheckBox)
 void SCheckBox::PrivateRegisterAttributes(FSlateAttributeInitializer& AttributeInitializer)
 {
-	SLATE_ADD_MEMBER_ATTRIBUTE_DEFINITION(AttributeInitializer, IsCheckboxCheckedAttribute, EInvalidateWidgetReason::None);
 	SLATE_ADD_MEMBER_ATTRIBUTE_DEFINITION(AttributeInitializer, PaddingOverrideAttribute, EInvalidateWidgetReason::None);
 	SLATE_ADD_MEMBER_ATTRIBUTE_DEFINITION(AttributeInitializer, ForegroundColorOverrideAttribute, EInvalidateWidgetReason::None);
 	SLATE_ADD_MEMBER_ATTRIBUTE_DEFINITION(AttributeInitializer, BorderBackgroundColorOverrideAttribute, EInvalidateWidgetReason::None);
 }
 
 SCheckBox::SCheckBox()
-	: IsCheckboxCheckedAttribute(*this, ECheckBoxState::Unchecked)
-	, PaddingOverrideAttribute(*this)
+	: PaddingOverrideAttribute(*this)
 	, ForegroundColorOverrideAttribute(*this, FSlateColor::UseStyle())
 	, BorderBackgroundColorOverrideAttribute(*this)
 {
@@ -365,7 +363,7 @@ const FSlateBrush* SCheckBox::OnGetBackgroundImage() const
 
 ECheckBoxState SCheckBox::GetCheckedState() const
 {
-	return IsCheckboxCheckedAttribute.Get();
+	return IsCheckboxChecked.Get();
 }
 
 /**
@@ -378,7 +376,7 @@ void SCheckBox::ToggleCheckedState()
 	// If the current check box state is checked OR undetermined we set the check box to unchecked.
 	if( State == ECheckBoxState::Checked || State == ECheckBoxState::Undetermined )
 	{
-		if ( !IsCheckboxCheckedAttribute.IsBound(*this) )
+		if ( !IsCheckboxChecked.IsBound() )
 		{
 			// When we are not bound, just toggle the current state.
 			SetIsChecked( ECheckBoxState::Unchecked );
@@ -389,7 +387,7 @@ void SCheckBox::ToggleCheckedState()
 	}
 	else if( State == ECheckBoxState::Unchecked )
 	{
-		if ( !IsCheckboxCheckedAttribute.IsBound(*this) )
+		if ( !IsCheckboxChecked.IsBound() )
 		{
 			// When we are not bound, just toggle the current state.
 			SetIsChecked( ECheckBoxState::Checked );
@@ -413,7 +411,7 @@ void SCheckBox::ToggleCheckedState()
 
 void SCheckBox::SetIsChecked(TAttribute<ECheckBoxState> InIsChecked)
 {
-	IsCheckboxCheckedAttribute.Assign(*this, MoveTemp(InIsChecked));
+	IsCheckboxChecked = MoveTemp(InIsChecked);
 }
 
 TEnumAsByte<EButtonClickMethod::Type> SCheckBox::GetClickMethodFromInputType(const FPointerEvent& MouseEvent) const
diff --git a/Engine/Source/Runtime/Slate/Public/Widgets/Input/SCheckBox.h b/Engine/Source/Runtime/Slate/Public/Widgets/Input/SCheckBox.h
--- a/Engine/Source/Runtime/Slate/Public/Widgets/Input/SCheckBox.h
+++ b/Engine/Source/Runtime/Slate/Public/Widgets/Input/SCheckBox.h
@@ -289,11 +289,6 @@ protected:
 	SLATE_API const FSlateBrush* GetBackgroundPressedImage() const;
 
 protected:
-	TSlateAttributeRef<ECheckBoxState> GetIsCheckboxCheckedAttribute() const
-	{
-		return TSlateAttributeRef<ECheckBoxState>{SharedThis(this), IsCheckboxCheckedAttribute};
-	}
-
 	SLATE_API void SetPaddingOverride(TAttribute<FMargin> InPaddingOverride);
 	TSlateAttributeRef<FMargin> GetPaddingOverrideAttribute() const
 	{
@@ -331,11 +326,8 @@ protected:
 	/** True if this check box is currently in a pressed state */
 	bool bIsPressed;
 
-#if WITH_EDITORONLY_DATA
 	/** Are we checked */
-	UE_DEPRECATED(5.8, "Use SetIsChecked / GetCheckedState / GetIsCheckboxCheckedAttribute")
-	TSlateDeprecatedTAttribute<ECheckBoxState> IsCheckboxChecked;
-#endif
+	TAttribute<ECheckBoxState> IsCheckboxChecked;
 
 	/** Delegate called when the check box changes state */
 	FOnCheckStateChanged OnCheckStateChanged;
@@ -437,9 +429,6 @@ protected:
 	TSharedPtr<SBorder> ContentContainer;
 
 private:
-	/** Are we checked */
-	TSlateAttribute<ECheckBoxState> IsCheckboxCheckedAttribute;
-
 	/** Overrides padding in the widget style, if set */
 	TSlateAttribute<FMargin> PaddingOverrideAttribute;
 	bool bIsPaddingAttrSet = false;
