Slate Invalidation - Show Slate Invalidation status in Diagnostic Window

diff --git a/Engine/Source/Runtime/Slate/Private/Framework/Application/SlateApplication.cpp b/Engine/Source/Runtime/Slate/Private/Framework/Application/SlateApplication.cpp
--- a/Engine/Source/Runtime/Slate/Private/Framework/Application/SlateApplication.cpp
+++ b/Engine/Source/Runtime/Slate/Private/Framework/Application/SlateApplication.cpp
@@ -1437,6 +1437,10 @@ void FSlateApplication::PrivateDrawWindows( TSharedPtr<SWindow> DrawOnlyThisWind
 	// Prepass the window
 	DrawPrepass( DrawOnlyThisWindow );
 
+#if WITH_SLATE_DEBUGGING
+	FSlateDebugging::ResetInvalidationPanelCount();
+#endif
+
 	{
 		FSlateRenderer::FScopedAcquireDrawBuffer ScopedDrawBuffer{ *Renderer };
 		FDrawWindowArgs DrawWindowArgs( ScopedDrawBuffer.GetDrawBuffer(), WidgetsToVisualizeUnderCursor);
