Slate Dynamic Invalidation - Place Actor Panel

diff --git a/Engine/Source/Developer/WidgetRegistration/Private/Layout/CategoryDrivenContentBuilderBase.cpp b/Engine/Source/Developer/WidgetRegistration/Private/Layout/CategoryDrivenContentBuilderBase.cpp
--- a/Engine/Source/Developer/WidgetRegistration/Private/Layout/CategoryDrivenContentBuilderBase.cpp
+++ b/Engine/Source/Developer/WidgetRegistration/Private/Layout/CategoryDrivenContentBuilderBase.cpp
@@ -10,6 +10,7 @@
 #include "Layout/SeparatorBuilder.h"
 #include "Layout/SeparatorTemplates.h"
 #include "Layout/Visibility.h"
+#include "Widgets/SInvalidationPanel.h"
 
 namespace UE::CategoryDrivenContentBuilderBase::Private
 {
@@ -143,15 +144,23 @@ void FCategoryDrivenContentBuilderBase::CreateWidget()
 		+ SSplitter::Slot()
 		.Resizable(false)
 		.SizeRule(SSplitter::SizeToContent)
+		[
+			SNew(SInvalidationPanel)
+			.UseDynamicInvalidation(true)
 			[
 				CategoryToolbarVBox.ToSharedRef()
 			]
+		]
 
 		+ SSplitter::Slot()
 		.SizeRule(SSplitter::FractionOfParent)
+		[
+			SNew(SInvalidationPanel)
+			.UseDynamicInvalidation(true)
 			[
 				MainContentVerticalBox->AsShared()
-			];
+			]
+		];
 	
 		ToolkitWidgetContainerVBox->AddSlot()
     	.VAlign(VAlign_Fill)
diff --git a/Engine/Source/Editor/PlacementMode/Private/SPlacementModeTools.cpp b/Engine/Source/Editor/PlacementMode/Private/SPlacementModeTools.cpp
--- a/Engine/Source/Editor/PlacementMode/Private/SPlacementModeTools.cpp
+++ b/Engine/Source/Editor/PlacementMode/Private/SPlacementModeTools.cpp
@@ -91,6 +91,13 @@ namespace PlacementViewFilter
 	}
 } // namespace PlacementViewFilter
 
+
+template <>
+struct TWidgetTypeTraits<class SPlacementAssetThumbnail>
+{
+	static constexpr bool SupportsInvalidation() { return true; }
+};
+
 /**
  * These are the asset thumbnails.
  */
diff --git a/Engine/Source/Editor/PlacementMode/Private/SPlacementModeTools.h b/Engine/Source/Editor/PlacementMode/Private/SPlacementModeTools.h
--- a/Engine/Source/Editor/PlacementMode/Private/SPlacementModeTools.h
+++ b/Engine/Source/Editor/PlacementMode/Private/SPlacementModeTools.h
@@ -17,6 +17,12 @@
 class STextBlock;
 class FCategoryDrivenContentBuilder;
 
+template <>
+struct TWidgetTypeTraits<class SPlacementAssetEntry>
+{
+	static constexpr bool SupportsInvalidation() { return true; }
+};
+
 /**
  * A tile representation of the class or the asset.  These are embedded into the views inside
  * of each tab.
diff --git a/Engine/Source/Runtime/Slate/Private/Framework/MultiBox/SClippingHorizontalBox.h b/Engine/Source/Runtime/Slate/Private/Framework/MultiBox/SClippingHorizontalBox.h
--- a/Engine/Source/Runtime/Slate/Private/Framework/MultiBox/SClippingHorizontalBox.h
+++ b/Engine/Source/Runtime/Slate/Private/Framework/MultiBox/SClippingHorizontalBox.h
@@ -41,6 +41,13 @@ void PrioritizedResize(
 
 }; // namespace UE::Slate
 
+
+template <>
+struct TWidgetTypeTraits<class SClippingHorizontalBox>
+{
+	static constexpr bool SupportsInvalidation() { return true; }
+};
+
 /** Specialized control for handling the clipping of toolbars and menubars */
 class SClippingHorizontalBox : public SHorizontalBox
 {
diff --git a/Engine/Source/Runtime/Slate/Private/Framework/MultiBox/SClippingVerticalBox.h b/Engine/Source/Runtime/Slate/Private/Framework/MultiBox/SClippingVerticalBox.h
--- a/Engine/Source/Runtime/Slate/Private/Framework/MultiBox/SClippingVerticalBox.h
+++ b/Engine/Source/Runtime/Slate/Private/Framework/MultiBox/SClippingVerticalBox.h
@@ -14,6 +14,12 @@ class FPaintArgs;
 class FSlateWindowElementList;
 class SComboButton;
 
+template <>
+struct TWidgetTypeTraits<class SClippingVerticalBox>
+{
+	static constexpr bool SupportsInvalidation() { return true; }
+};
+
 /** Specialized control for handling the clipping of toolbars and menubars */
 class SClippingVerticalBox : public SVerticalBox
 {
diff --git a/Engine/Source/Runtime/Slate/Private/Framework/MultiBox/SPrioritizedWrapBox.h b/Engine/Source/Runtime/Slate/Private/Framework/MultiBox/SPrioritizedWrapBox.h
--- a/Engine/Source/Runtime/Slate/Private/Framework/MultiBox/SPrioritizedWrapBox.h
+++ b/Engine/Source/Runtime/Slate/Private/Framework/MultiBox/SPrioritizedWrapBox.h
@@ -21,6 +21,13 @@ namespace UE::Slate::PrioritizedWrapBox
 	class FChildArranger;
 }; // namespace UE::Slate
 
+
+template <>
+struct TWidgetTypeTraits<class SPrioritizedWrapBox>
+{
+	static constexpr bool SupportsInvalidation() { return true; }
+};
+
 /**
  * An extension of SHorizontalBox that adds wrapping behavior (similar to SWrapBox), primarily for use with MultiBox.
  * Wrapping candidates are determined by priority, so it doesn't necessarily occur sequentially.
diff --git a/Engine/Source/Runtime/Slate/Public/Framework/MultiBox/MultiBox.h b/Engine/Source/Runtime/Slate/Public/Framework/MultiBox/MultiBox.h
--- a/Engine/Source/Runtime/Slate/Public/Framework/MultiBox/MultiBox.h
+++ b/Engine/Source/Runtime/Slate/Public/Framework/MultiBox/MultiBox.h
@@ -608,6 +608,11 @@ public:
 };
 
 
+template <>
+struct TWidgetTypeTraits<class SMultiBlockBaseWidget>
+{
+	static constexpr bool SupportsInvalidation() { return true; }
+};
 
 /**
  * MultiBlock Slate base widget (pure virtual).  You'll derive your own MultiBlock class from this base class.
@@ -653,6 +658,13 @@ protected:
 	bool bSectionContainsIcons;
 };
 
+
+template <>
+struct TWidgetTypeTraits<class SMultiBoxWidget>
+{
+	static constexpr bool SupportsInvalidation() { return true; }
+};
+
 /**
  * MultiBox Slate widget
  */
@@ -683,7 +695,7 @@ public:
 	 *
 	 * @param	InMultiBox	The MultiBox we'll be associated with
 	 */
-	void SetMultiBox( TSharedRef< FMultiBox > InMultiBox )
+	void SetMultiBox( const TSharedRef< FMultiBox >& InMultiBox )
 	{
 		MultiBox = InMultiBox;
 	}
@@ -693,7 +705,7 @@ public:
 	*
 	* @param	InContent	The new content to place in the ChildSlot
 	*/
-	void SetContent( TSharedRef< SWidget > InContent )
+	void SetContent( const TSharedRef< SWidget >& InContent )
 	{	
 		ChildSlot
 		[
@@ -730,7 +742,7 @@ public:
 	/**
 	* Sets optional maximum height of widget
 	*/
-	void SetMaxHeight(TAttribute<float> InMaxHeight) { MaxHeight = InMaxHeight; }
+	void SetMaxHeight(TAttribute<float> InMaxHeight) { MaxHeight = MoveTemp(InMaxHeight); }
 
 	/**
 	 * Builds this MultiBox widget up from the MultiBox associated with it
diff --git a/Engine/Source/Runtime/Slate/Public/Framework/MultiBox/SToolBarButtonBlock.h b/Engine/Source/Runtime/Slate/Public/Framework/MultiBox/SToolBarButtonBlock.h
--- a/Engine/Source/Runtime/Slate/Public/Framework/MultiBox/SToolBarButtonBlock.h
+++ b/Engine/Source/Runtime/Slate/Public/Framework/MultiBox/SToolBarButtonBlock.h
@@ -178,6 +178,11 @@ private:
 };
 
 
+template <>
+struct TWidgetTypeTraits<class SToolBarButtonBlock>
+{
+	static constexpr bool SupportsInvalidation() { return true; }
+};
 
 /**
  * Tool bar button MultiBlock widget
diff --git a/Engine/Source/Runtime/Slate/Public/Widgets/Layout/SMenuOwner.h b/Engine/Source/Runtime/Slate/Public/Widgets/Layout/SMenuOwner.h
--- a/Engine/Source/Runtime/Slate/Public/Widgets/Layout/SMenuOwner.h
+++ b/Engine/Source/Runtime/Slate/Public/Widgets/Layout/SMenuOwner.h
@@ -7,6 +7,13 @@
 #include "Widgets/SCompoundWidget.h"
 #include "Widgets/Input/SMenuAnchor.h"
 
+template <>
+struct TWidgetTypeTraits<class SMenuOwner>
+{
+	static constexpr bool SupportsInvalidation() { return true; }
+};
+
+
 /**
  * Any widget that wants to own lists of items that may have sub-menus that should be handled in the same way as multibox menus (mouse can cross other items without sub-menu closing)
  */
diff --git a/Engine/Source/Runtime/Slate/Public/Widgets/Layout/SUniformWrapPanel.h b/Engine/Source/Runtime/Slate/Public/Widgets/Layout/SUniformWrapPanel.h
--- a/Engine/Source/Runtime/Slate/Public/Widgets/Layout/SUniformWrapPanel.h
+++ b/Engine/Source/Runtime/Slate/Public/Widgets/Layout/SUniformWrapPanel.h
@@ -14,6 +14,12 @@
 
 class FArrangedChildren;
 
+template <>
+struct TWidgetTypeTraits<class SUniformWrapPanel>
+{
+	static constexpr bool SupportsInvalidation() { return true; }
+};
+
 /** A panel that evenly divides up available space between all of its children. */
 class SUniformWrapPanel : public SPanel
 {
