Add SComboBox ReselectionTriggersOnSelectionChanged slate argument to trigger OnSelectionChanged even when the same item gets reselected (false by default) to basically get the same behavior as SListView

#rnx

diff --git a/Engine/Source/Runtime/Slate/Public/Widgets/Input/SComboBox.h b/Engine/Source/Runtime/Slate/Public/Widgets/Input/SComboBox.h
--- a/Engine/Source/Runtime/Slate/Public/Widgets/Input/SComboBox.h
+++ b/Engine/Source/Runtime/Slate/Public/Widgets/Input/SComboBox.h
@@ -122,6 +122,7 @@ public:
 		, _HasDownArrow( true )
 		, _EnableGamepadNavigationMode(false)
 		, _IsFocusable( true )
+		, _ReselectionTriggersOnSelectionChanged(false)
 		{}
 		
 		/** Slot for this button's content (optional) */
@@ -180,7 +181,10 @@ public:
 
 		/** True if this combo's menu should be collapsed when our parent receives focus, false (default) otherwise */
 		SLATE_ARGUMENT(bool, CollapseMenuOnParentFocus)
-				
+
+		/** Whether OnSelectionChanged event gets triggered when selecting the item that's already selected */
+		SLATE_ARGUMENT(bool, ReselectionTriggersOnSelectionChanged)
+
 	SLATE_END_ARGS()
 
 	/**
@@ -208,6 +212,7 @@ public:
 		this->OnGenerateWidget = InArgs._OnGenerateWidget;
 		this->EnableGamepadNavigationMode = InArgs._EnableGamepadNavigationMode;
 		this->bControllerInputCaptured = false;
+		this->bReselectionTriggersOnSelectionChanged = InArgs._ReselectionTriggersOnSelectionChanged;
 
 		CustomScrollbar = InArgs._CustomScrollbar;
 
@@ -607,7 +612,7 @@ private:
 		if(SelectInfo != ESelectInfo::OnNavigation)
 		{
 			// Ensure that the proposed selection is different from selected
-			if ( ProposedSelection != SelectedItem )
+			if (bReselectionTriggersOnSelectionChanged || ProposedSelection != SelectedItem)
 			{
 				PlaySelectionChangeSound();
 				SelectedItem = ProposedSelection;
@@ -702,6 +707,8 @@ private:
 	// Holds a flag indicating whether a controller/keyboard is manipulating the combobox's value. 
 	// When true, navigation away from the widget is prevented until a new value has been accepted or canceled. 
 	bool bControllerInputCaptured;
+	// Whether OnSelectionChanged event gets triggered when selecting the item that's already selected
+	bool bReselectionTriggersOnSelectionChanged;
 
 	TSharedPtr<SBox> ComboBoxMenuContent;
 };
