Add focused widget info to our Editor summary. Also exposes forground tab in TabStack, and adds the ability to get focused tab from SlateUser, if one exists.

diff --git a/Engine/Source/Runtime/Engine/Private/Analytics/EditorAnalyticsSessionSummary.cpp b/Engine/Source/Runtime/Engine/Private/Analytics/EditorAnalyticsSessionSummary.cpp
--- a/Engine/Source/Runtime/Engine/Private/Analytics/EditorAnalyticsSessionSummary.cpp
+++ b/Engine/Source/Runtime/Engine/Private/Analytics/EditorAnalyticsSessionSummary.cpp
@@ -7,11 +7,18 @@
 #include "Editor.h"
 #include "IVREditorModule.h"
 #include "Framework/Application/SlateApplication.h"
+
+// includes for editor focus summary test
+#include "Framework/Application/SlateUser.h"
+#include "Framework/Docking/SDockingTabStack.h"
+#include "Widgets/Docking/SDockTab.h"
+
 #include "Interfaces/IAnalyticsPropertyStore.h"
 #include "Kismet2/DebuggerCommands.h"
 #include "ProfilingDebugging/StallDetector.h"
 #include "ShaderCompiler.h"
 
+
 namespace EditorAnalyticsProperties
 {
 	// List of mutable keys.
@@ -39,6 +46,14 @@ namespace EditorAnalyticsProperties
 	static const TAnalyticsProperty<double> TimeShaderCompilationWasActive = TEXT("TimeShaderCompilationWasActive ");
 }
 
+namespace EditorAnalyticsCVars
+{
+	static TAutoConsoleVariable<bool> CVarDisableFocusTimes(
+		TEXT("EditorAnalytics.DisableFocusTimes"),
+		false,
+		TEXT("Disables collecting focused widget times for analytics summary."));
+}
+
 FEditorAnalyticsSessionSummary::FEditorAnalyticsSessionSummary(TSharedPtr<IAnalyticsPropertyStore> Store, uint32 MonitorProcessId)
 	: FEngineAnalyticsSessionSummary(MoveTemp(Store), MonitorProcessId)
 	, LastUserActivityTimeSecs(FPlatformTime::Seconds())
@@ -129,9 +144,92 @@ bool FEditorAnalyticsSessionSummary::UpdateSessionProgressInternal(bool bCrashin
 	EditorAnalyticsProperties::ShaderDDCHits.Set(GetStore(), GShaderCompilerStats->GetDDCHits());
 	EditorAnalyticsProperties::TimeShaderCompilationWasActive.Set(GetStore(), GShaderCompilerStats->GetTimeShaderCompilationWasActive());
 
+	if (!EditorAnalyticsCVars::CVarDisableFocusTimes.GetValueOnAnyThread())
+	{
+		UpdateUserFocusTimes(FPlatformTime::Seconds());
+	}
+
 	return bShouldPersist;
 }
 
+namespace EditorAnalyticsSessionSummaryHelpers
+{
+	bool GetFocusedTabIdentifierIfInFocusPath(TSharedPtr<const FSlateUser> SlateUser, FString& OutTabIdentifier)
+	{
+		if (!SlateUser.IsValid() || !SlateUser->HasValidFocusPath())
+		{
+			return false;
+		}
+
+		const FWeakWidgetPath& WeakFocusPath = SlateUser->GetWeakFocusPath();
+
+		for (int i = WeakFocusPath.Widgets.Num() - 1; i > 0; --i)
+		{
+			TWeakPtr<SWidget> WidgetInPathWeakPtr = WeakFocusPath.Widgets[i];
+
+			if (!WidgetInPathWeakPtr.IsValid())
+			{
+				continue;
+			}
+
+			TSharedPtr<SWidget> WidgetInPath = WidgetInPathWeakPtr.Pin();
+
+			// docking tab stack is the part of the docking stuff that shows up in focus paths
+			if (WidgetInPath->GetTypeAsString() == TEXT("SDockingTabStack"))
+			{
+				SDockingTabStack* TabStackWidget = reinterpret_cast<SDockingTabStack*>(WidgetInPath.Get());
+
+				if (TabStackWidget != nullptr)
+				{
+					TSharedPtr<SDockTab> TabWidget = TabStackWidget->GetSelectedTabFromWell();
+
+					if (TabWidget.IsValid())
+					{
+						OutTabIdentifier = TabWidget->GetLayoutIdentifier().ToString();
+
+						return true;
+					}
+
+				}
+			}
+		}
+
+		return false;
+	}
+}
+
+void FEditorAnalyticsSessionSummary::UpdateUserFocusTimes(double CurrTimeSecs)
+{
+	FSlateApplication& SlateApplication = FSlateApplication::Get();
+
+	TSharedPtr<const FSlateUser> CurrentUser = SlateApplication.GetCursorUser();
+
+	double PrevTimeOfLastUpdate = TimeOfLastUpdate.load();
+	// similar to our other timer updates, we want to ensure only a single thread performs this update
+	if (TimeOfLastUpdate.compare_exchange_strong(PrevTimeOfLastUpdate, CurrTimeSecs))
+	{
+		FString TabIdentifier;
+		if (EditorAnalyticsSessionSummaryHelpers::GetFocusedTabIdentifierIfInFocusPath(CurrentUser, TabIdentifier))
+		{
+			TabIdentifier = FString(TEXT("FocusTime")).Append(TabIdentifier);
+			float DeltaAsFloat = static_cast<float>(CurrTimeSecs - PrevTimeOfLastUpdate);
+
+			if (!GetStore()->Contains(TabIdentifier))
+			{
+				GetStore()->Set(TabIdentifier, DeltaAsFloat);
+			}
+			else
+			{
+				GetStore()->Update(TabIdentifier, [DeltaAsFloat](float& InOutValue)
+					{
+						InOutValue += DeltaAsFloat;
+						return true;
+					});
+			}
+		}
+	}
+}
+
 bool FEditorAnalyticsSessionSummary::UpdateUserIdleTime(double CurrTimeSecs, bool bReset)
 {
 	bool bSessionUpdated = false;
diff --git a/Engine/Source/Runtime/Engine/Private/Analytics/EditorAnalyticsSessionSummary.h b/Engine/Source/Runtime/Engine/Private/Analytics/EditorAnalyticsSessionSummary.h
--- a/Engine/Source/Runtime/Engine/Private/Analytics/EditorAnalyticsSessionSummary.h
+++ b/Engine/Source/Runtime/Engine/Private/Analytics/EditorAnalyticsSessionSummary.h
@@ -26,6 +26,7 @@ protected:
 
 private:
 	bool UpdateUserIdleTime(double CurrTimeSecs, bool bReset);
+	void UpdateUserFocusTimes(double CurrTimeSecs);
 	void OnSlateUserInteraction(double CurrSlateInteractionTime);
 	void OnEnterPIE(const bool /*bIsSimulating*/);
 	void OnExitPIE(const bool /*bIsSimulating*/);
@@ -37,6 +38,9 @@ private:
 
 	/** The number of idle seconds in the current idle sequence that were accounted (saved in the session) for the user idle counters. */
 	std::atomic<double> AccountedUserIdleSecs;
+
+	/** Time since we last updated our focus stats */
+	std::atomic<double> TimeOfLastUpdate;
 };
 
 #endif // WITH_EDITOR
diff --git a/Engine/Source/Runtime/Slate/Private/Framework/Application/SlateUser.cpp b/Engine/Source/Runtime/Slate/Private/Framework/Application/SlateUser.cpp
--- a/Engine/Source/Runtime/Slate/Private/Framework/Application/SlateUser.cpp
+++ b/Engine/Source/Runtime/Slate/Private/Framework/Application/SlateUser.cpp
@@ -620,6 +620,16 @@ void FSlateUser::SetUserNavigationConfig(TSharedPtr<FNavigationConfig> InNavigat
 #endif // WITH_SLATE_DEBUGGING
 }
 
+bool FSlateUser::HasValidFocusPath() const 
+{
+	return WeakFocusPath.IsValid();
+}
+
+const FWeakWidgetPath& FSlateUser::GetWeakFocusPath() const
+{
+	return WeakFocusPath;
+}
+
 bool FSlateUser::IsTouchPointerActive(int32 TouchPointerIndex) const
 {
 	return TouchPointerIndex < (int32)ETouchIndex::CursorPointerIndex && PointerPositionsByIndex.Contains(TouchPointerIndex);
diff --git a/Engine/Source/Runtime/Slate/Private/Framework/Docking/SDockingNode.h b/Engine/Source/Runtime/Slate/Private/Framework/Docking/SDockingNode.h
deleted file mode 100644
diff --git a/Engine/Source/Runtime/Slate/Private/Framework/Docking/SDockingTabStack.cpp b/Engine/Source/Runtime/Slate/Private/Framework/Docking/SDockingTabStack.cpp
--- a/Engine/Source/Runtime/Slate/Private/Framework/Docking/SDockingTabStack.cpp
+++ b/Engine/Source/Runtime/Slate/Private/Framework/Docking/SDockingTabStack.cpp
@@ -1216,6 +1216,18 @@ bool SDockingTabStack::IsTabAllowedInSidebar(TSharedPtr<SDockTab> Tab) const
 	return false;
 }
 
+TSharedPtr<SDockTab> SDockingTabStack::GetSelectedTabFromWell()
+{
+	if (TabWell.IsValid())
+	{
+		return TabWell->GetForegroundTab();
+	}
+	else
+	{
+		return nullptr;
+	}
+}
+
 SSplitter::ESizeRule SDockingTabStack::GetSizeRule() const
 {
 	int32 NumTabs = this->GetNumTabs();
diff --git a/Engine/Source/Runtime/Slate/Private/Framework/Docking/SDockingTabStack.h b/Engine/Source/Runtime/Slate/Private/Framework/Docking/SDockingTabStack.h
deleted file mode 100644
diff --git a/Engine/Source/Runtime/Slate/Public/Framework/Application/SlateUser.h b/Engine/Source/Runtime/Slate/Public/Framework/Application/SlateUser.h
--- a/Engine/Source/Runtime/Slate/Public/Framework/Application/SlateUser.h
+++ b/Engine/Source/Runtime/Slate/Public/Framework/Application/SlateUser.h
@@ -126,13 +126,13 @@ public:
 	SLATE_API void SetUserNavigationConfig(TSharedPtr<FNavigationConfig> InNavigationConfig);
 	TSharedPtr<FNavigationConfig> GetUserNavigationConfig() const { return UserNavigationConfig; }
 
+	SLATE_API bool HasValidFocusPath() const;
+	SLATE_API const FWeakWidgetPath& GetWeakFocusPath() const;
+
 SLATE_SCOPE:
 	static TSharedRef<FSlateUser> Create(int32 InUserIndex, TSharedPtr<ICursor> InCursor);
 	static SLATE_API TSharedRef<FSlateUser> Create(FPlatformUserId InPlatformUserId, TSharedPtr<ICursor> InCursor);
-	
-	inline bool HasValidFocusPath() const { return WeakFocusPath.IsValid(); }
-	inline const FWeakWidgetPath& GetWeakFocusPath() const { return WeakFocusPath; }
-	
+
 	inline TSharedRef<FWidgetPath> GetFocusPath() const
 	{
 		if (!StrongFocusPath.IsValid())
diff --git a/Engine/Source/Runtime/Slate/Public/Framework/Docking/SDockingNode.h b/Engine/Source/Runtime/Slate/Public/Framework/Docking/SDockingNode.h
new file mode 100644
--- /dev/null
+++ b/Engine/Source/Runtime/Slate/Public/Framework/Docking/SDockingNode.h
@@ -0,0 +1,153 @@
+// Copyright Epic Games, Inc. All Rights Reserved.
+
+#pragma once
+
+#include "CoreMinimal.h"
+#include "Input/Reply.h"
+#include "Widgets/SCompoundWidget.h"
+#include "Framework/Docking/TabManager.h"
+#include "Widgets/Layout/SSplitter.h"
+
+class SDockingArea;
+class SDockingSplitter;
+
+/**
+ * A node in the Docking/Tabbing hierarchy.
+ * Any SDockingNode can be either a stack of tabs or a splitter.
+ */
+class SDockingNode : public SCompoundWidget
+{
+public:
+	enum Type
+	{
+		/**
+		 * A tab stack is a collection of tabs with associated content.
+		 * Only one of the tabs is active at a time; its content will be shown.
+		 */
+		DockTabStack,
+		/** Displays multiple child SDockNodes horizontally or vertically. */
+		DockSplitter,
+		/** Top-level dock splitters */
+		DockArea,
+		/** Takes up some room during layout re-arranging. */
+		PlaceholderNode
+	};
+
+	/** Direction relative to some DockNode. For example, we may wish to say "Dock a tab RightOf this node." */
+	enum RelativeDirection
+	{
+		LeftOf,
+		Above,
+		RightOf,
+		Below,
+		Center
+	};
+
+	/** Tracking for the docking cross, content, and background of the TabStack */
+	struct FOverlayManagement
+	{
+		FOverlayManagement()
+			: bShowingCross(false)
+		{
+		}
+
+		/** The overlay widget that shows TabStack's content*/
+		TSharedPtr<SOverlay> ContentAreaOverlay;
+
+		/** true when we're showing the DockCross */
+		bool bShowingCross;
+	};
+
+	/** @return Is this dock node a tab stack, splitter or something else? */
+	virtual Type GetNodeType() const = 0;
+
+	/**
+	 * All DockNodes are aware of their parent DockNode (unless they are a root node; a.k.a DockArea)
+	 * The parent is set whenever a node is inserted into a docking hierarchy.
+	 *
+	 * @param InParent      Set this node's parent.
+	 */
+	virtual void SetParentNode( TSharedRef<SDockingSplitter> InParent )
+	{
+		ParentNodePtr = InParent;
+	}
+
+	/** A tab can be removed from a stack when a uset drags it away or when the user closes it */
+	enum ELayoutModification
+	{
+		/** Tab was dragged out of the node */
+		TabRemoval_DraggedOut,
+		/** Tab is being closed */
+		TabRemoval_Closed,
+		/** Tab is being put in a sidebar */
+		TabRemoval_Sidebar,
+		TabRemoval_None
+	};
+	
+
+	/** Gets the dock area that this resides in */
+	SLATE_API virtual TSharedPtr<SDockingArea> GetDockArea();
+	SLATE_API virtual TSharedPtr<const SDockingArea> GetDockArea() const;
+	
+	/** Recursively searches through all children looking for child tabs */
+	virtual TArray< TSharedRef<SDockTab> > GetAllChildTabs() const {return TArray< TSharedRef<SDockTab> >();}
+
+	/** Gets the number of tabs in all children recursively */
+	virtual int32 GetNumTabs() const { return 0; }
+
+	/**
+	 * Attempt to dock the TabWidget from the DragDropEvent next to this node.
+	 * 
+	 * @param Direction       Where relative to this direction to dock the node: e.g. to the right
+	 * @param DragDropEvent   DragAndDrop event that potentially carries
+	 *
+	 * @return An FReply::Handled() if we successfully docked the content; FReply::Unhandled() otherwise.
+	 */
+	SLATE_API virtual FReply OnUserAttemptingDock( SDockingNode::RelativeDirection Direction, const FDragDropEvent& DragDropEvent );
+
+	/** Should this node auto-size or be a percentage of its parent size; this setting is usually determined by users */
+	virtual SSplitter::ESizeRule GetSizeRule() const { return SSplitter::FractionOfParent; }
+
+	/** @return the numerator for the fraction of available space that this dock node should occupy. */
+	SLATE_API float GetSizeCoefficient() const;
+
+	/** Set the coefficient size */
+	SLATE_API void SetSizeCoefficient( float InSizeCoefficient );
+
+	/** @returns true when this node is showing a live DockableTab; false when this nodes content is only kept around for history's sake */
+	//virtual bool IsShowingLiveTabs() const;
+
+	/**
+	 * Recursively build up an tree of FLayoutNodes that represent the persistent layout state of this DockingNode and all its descendants.
+	 * The result can be easily serialized to disk or used as history for restoring tabs to their previous locations.
+	 *
+	 * @return a pointer to persistent layout node; invalid pointer if this dock node had no tab layouts worth saving.
+	 */
+	virtual TSharedPtr<FTabManager::FLayoutNode> GatherPersistentLayout() const = 0;
+
+	enum ECleanupRetVal
+	{
+		VisibleTabsUnderNode,
+		HistoryTabsUnderNode,
+		NoTabsUnderNode
+	};
+	virtual SDockingNode::ECleanupRetVal CleanUpNodes(){ return NoTabsUnderNode; }
+
+protected:
+
+	SLATE_API SDockingNode();
+
+	/** A live tab was added to this node or one of its descendants */
+	SLATE_API virtual void OnLiveTabAdded();
+
+	virtual void OnResized() {}
+
+	/**
+	 * Weak reference to the parent node. It is NULL until the node
+	 * is inserted into the hierarchy. Also NULL for root nodes (aka SDockingArea)
+	 */
+	TWeakPtr<SDockingSplitter> ParentNodePtr;
+
+	float SizeCoefficient = 1.0f;
+};
+
diff --git a/Engine/Source/Runtime/Slate/Public/Framework/Docking/SDockingTabStack.h b/Engine/Source/Runtime/Slate/Public/Framework/Docking/SDockingTabStack.h
new file mode 100644
--- /dev/null
+++ b/Engine/Source/Runtime/Slate/Public/Framework/Docking/SDockingTabStack.h
@@ -0,0 +1,378 @@
+// Copyright Epic Games, Inc. All Rights Reserved.
+
+#pragma once
+
+#include "CoreMinimal.h"
+#include "Layout/Margin.h"
+#include "Layout/Visibility.h"
+#include "Widgets/DeclarativeSyntaxSupport.h"
+#include "Styling/SlateColor.h"
+#include "Layout/Geometry.h"
+#include "Input/Reply.h"
+#include "Widgets/SWidget.h"
+#include "Animation/CurveSequence.h"
+#include "Framework/Docking/TabManager.h"
+#include "Framework/Docking/SDockingNode.h"
+
+#define DEBUG_TAB_MANAGEMENT 0
+
+class FUICommandList;
+class FWeakWidgetPath;
+class FWidgetPath;
+class SDockingTabWell;
+class SBorder;
+
+template<typename ChildType> class TSlotlessChildren;
+
+struct FDockingStackOptionalContent
+{
+	FDockingStackOptionalContent()
+		: ContentLeft(SNullWidget::NullWidget)
+		, ContentRight(SNullWidget::NullWidget)
+		, TitleBarContentRight(SNullWidget::NullWidget)
+	{}
+
+	// Content that appears on the left side of tabs in this docking stack
+	TSharedRef<SWidget> ContentLeft;
+
+	// Content that appears on the right side of tabs in this docking stack
+	TSharedRef<SWidget> ContentRight;
+
+	// Content that appears on the right side of the title bar in the window this stack is in
+	TSharedRef<SWidget> TitleBarContentRight;
+};
+
+
+/**
+ * A node in the Docking/Tabbing hierarchy.
+ * A DockTabStack shows a row of tabs and the content of one selected tab.
+ * It also supports re-arranging tabs and dragging them out for the stack.
+ */
+class SDockingTabStack : public SDockingNode
+{
+public:
+	SLATE_BEGIN_ARGS(SDockingTabStack)
+			: _IsDocumentArea(false) {}
+		SLATE_ARGUMENT( bool, IsDocumentArea )
+	SLATE_END_ARGS()
+
+	virtual Type GetNodeType() const override
+	{
+		return SDockingNode::DockTabStack;
+	}
+
+	SLATE_API void OnLastTabRemoved();
+
+	SLATE_API void OnTabClosed(const TSharedRef<SDockTab>& ClosedTab, SDockingNode::ELayoutModification RemovalMethod);
+
+	SLATE_API void OnTabRemoved( const FTabId& TabId );
+
+	SLATE_API void Construct( const FArguments& InArgs, const TSharedRef<FTabManager::FStack>& PersistentNode );
+
+	// TabStack methods
+
+	SLATE_API void OpenTab(const TSharedRef<SDockTab>& InTab, int32 InsertAtLocation = INDEX_NONE, bool bKeepInactive = false);
+
+	SLATE_API void AddTabWidget(const TSharedRef<SDockTab>& InTab, int32 AtLocation = INDEX_NONE, bool bKeepInactive = false);
+
+	SLATE_API void AddSidebarTab(const TSharedRef<SDockTab>& InTab);
+
+	SLATE_API float GetTabSidebarSizeCoefficient(const TSharedRef<SDockTab>& InTab);
+	SLATE_API void SetTabSidebarSizeCoefficient(const TSharedRef<SDockTab>& InTab, float InSizeCoefficient);
+
+	SLATE_API bool IsTabPinnedInSidebar(const TSharedRef<SDockTab>& InTab);
+	SLATE_API void SetTabPinnedInSidebar(const TSharedRef<SDockTab>& InTab, bool bPinnedInSidebar);
+
+	SLATE_API bool IsTabLocked(const TSharedRef<const SDockTab>& InTab) const;
+	SLATE_API void SetTabLocked(const TSharedRef<const SDockTab>& InTab, TOptional<bool> bInLocked);
+
+	/** @return true if the specified tab can be moved to a sidebar */
+	SLATE_API bool CanMoveTabToSideBar(TSharedRef<SDockTab> Tab) const;
+
+	/** Moves a specific tab from this stack into a sidebar */
+	SLATE_API void MoveTabToSidebar(TSharedRef<SDockTab> Tab);
+
+	/** Restores a specific tab from this stack to its original position */
+	SLATE_API void RestoreTabFromSidebar(TSharedRef<SDockTab> Tab);
+
+	/** @return All child tabs in this node */
+	SLATE_API const TSlotlessChildren<SDockTab>& GetTabs() const;
+
+	SLATE_API bool HasTab( const struct FTabMatcher& TabMatcher ) const;
+
+	/** @return the last known geometry of this TabStack */
+	SLATE_API FGeometry GetTabStackGeometry() const;
+
+	SLATE_API void RemoveClosedTabsWithName( FName InName );
+
+	SLATE_API bool IsShowingLiveTabs() const;
+
+	SLATE_API void BringToFront( const TSharedRef<SDockTab>& TabToBringToFront );
+
+	/** Set the content that the DockNode is presenting. */
+	SLATE_API void SetNodeContent( const TSharedRef<SWidget>& InContent, const FDockingStackOptionalContent& OptionalContent);
+
+	SLATE_API virtual FReply OnUserAttemptingDock( SDockingNode::RelativeDirection Direction, const FDragDropEvent& DragDropEvent ) override;
+	
+	/** Recursively searches through all children looking for child tabs */
+	SLATE_API virtual TArray< TSharedRef<SDockTab> > GetAllChildTabs() const override;
+
+	/** Gets the number of tabs in all children recursively */
+	SLATE_API virtual int32 GetNumTabs() const override;
+
+	SLATE_API virtual SSplitter::ESizeRule GetSizeRule() const override;
+
+	/**
+	 * Returns a bool indicating whether or not a tab is in a state where it can be hidden. A tab typically shouldn't
+	 * be hidden in certain conditions like if there is more than one tab in the well, or if it's in a floating window.
+	 *
+	 * @return returns true if the tab can be hidden, else it returns false */
+	SLATE_API bool CanHideTabWell() const;
+
+	
+	/**
+	 * Sets the tab well to hidden/not hidden. 
+	 *
+	 * @param bShouldHideTabWell the bool that defines whether or not the tab well should be hidden
+	 */
+	SLATE_API void SetTabWellHidden( bool bShouldHideTabWell );
+	SLATE_API bool IsTabWellHidden() const;
+
+	SLATE_API virtual TSharedPtr<FTabManager::FLayoutNode> GatherPersistentLayout() const override;
+
+	SLATE_API TSharedPtr<SDockTab> GetSelectedTabFromWell();
+
+	/** Elements for which we might want to reserve space. */
+	enum class EChromeElement
+	{
+		Icon,
+		Controls
+	};
+
+	/**
+	 * Remove all the space that might have been reserved for
+	 * various window chrome elements (app icons, minimize, close, etc.)
+	 */
+	SLATE_API void ClearReservedSpace();
+	/**
+	 * Add some extra padding so that the corresponding window chrome element
+	 * does not overlap our tabs.
+	 */
+	SLATE_API void ReserveSpaceForWindowChrome(EChromeElement Element, bool bIncludePaddingForMenuBar, bool bOnlyMinorTabs);
+
+public:
+	/** SWidget interface */
+	SLATE_API virtual FReply OnDragOver( const FGeometry& MyGeometry, const FDragDropEvent& DragDropEvent ) override;
+	SLATE_API virtual void OnDragLeave( const FDragDropEvent& DragDropEvent ) override;
+	SLATE_API virtual FReply OnDrop( const FGeometry& MyGeometry, const FDragDropEvent& DragDropEvent ) override;
+	SLATE_API virtual void OnFocusChanging( const FWeakWidgetPath& PreviousFocusPath, const FWidgetPath& NewWidgetPath, const FFocusEvent& InFocusEvent ) override;
+	SLATE_API virtual FReply OnMouseButtonDown( const FGeometry& MyGeometry, const FPointerEvent& MouseEvent ) override;
+	SLATE_API virtual FReply OnKeyDown(const FGeometry& MyGeometry, const FKeyEvent& InKeyEvent) override;
+	virtual bool SupportsKeyboardFocus() const override { return true; }
+
+protected:
+
+	virtual EWindowZone::Type GetWindowZoneOverride() const override
+	{
+		// Pretend we are a title bar so the user can grab the area to move the window around
+		return EWindowZone::TitleBar;
+	}
+
+	SLATE_API void CloseForegroundTab();
+
+	enum class ETabsToClose : uint8
+	{
+		CloseDocumentTabs,
+		CloseDocumentAndMajorTabs,
+		CloseAllTabs
+	};
+
+	enum class ECloseTabsDirection : uint8
+	{
+		Left,
+		Right
+	};
+
+	SLATE_API void CloseTabsInDirectionFromForegroundTab(ETabsToClose TabsToClose, ECloseTabsDirection Direction);
+
+	/**
+	 * Close all the background tabs.
+	 *
+	 * @param  TabsToClose   Close just document tabs or both document and tool tabs.
+	 */
+	SLATE_API void CloseAllButForegroundTab(ETabsToClose TabsToClose);
+
+	SLATE_API FReply TabWellRightClicked( const FGeometry& TabWellGeometry, const FPointerEvent& MouseEvent );
+
+	SLATE_API virtual SDockingNode::ECleanupRetVal CleanUpNodes() override;
+
+	SLATE_API int32 OpenPersistentTab( const FTabId& TabId, int32 OpenLocationAmongActiveTabs = INDEX_NONE );
+
+	SLATE_API int32 ClosePersistentTab( const FTabId& TabId );
+
+	SLATE_API void RemovePersistentTab( const FTabId& TabId );
+
+	/** Overridden from SDockingNode */
+	SLATE_API virtual void SetParentNode( TSharedRef<class SDockingSplitter> InParent ) override;
+
+	/** @return false if the content is disabled (non-interactable) */
+	bool IsContentEnabled() const;
+
+private:
+	
+	/** Data that persists across sessions and when the widget associated with this node is removed. */	
+	TArray<FTabManager::FTab> Tabs;
+
+
+	/**
+	 * Creates a SDockingTabStack by adding a new split to this stack's parent splitter and attaching the new SDockingTabStack.
+	 * 
+	 * @param	Direction	The relative direction to split the parent splitter
+	 *
+	 * @return	The newly created empty SDockingTabStack, ready for a tab to be added to it
+	 */
+	SLATE_API TSharedRef< SDockingTabStack > CreateNewTabStackBySplitting( const SDockingNode::RelativeDirection Direction );
+
+	/** What should the content area look like for the current tab? */
+	SLATE_API const FSlateBrush* GetContentAreaBrush() const;
+
+	/** How much padding to show around the content currently being presented */
+	SLATE_API FMargin GetContentPadding() const;
+
+	/** Depending on the tabs we put into the tab well, we want a different background brush. */
+	SLATE_API const FSlateBrush* GetTabWellBrush() const;
+
+	/** Show the tab well? */
+	SLATE_API EVisibility GetTabWellVisibility() const;
+
+	/** Show the stuff needed to unhide the tab well? */
+	SLATE_API EVisibility GetUnhideButtonVisibility() const;
+
+	/** Show/Hide the tab well; do it smoothly with an animation */
+	SLATE_API void ToggleTabWellVisibility();
+
+	/** Moves the foreground tab to a sidebar */
+	SLATE_API void MoveForegroundTabToSidebar();
+
+	/** Unhides the tab well, revealing all tab headers */
+	SLATE_API FReply UnhideTabWell();
+
+	/** Only allow closing the tab well when the tab allows it. */
+	SLATE_API bool CanCloseForegroundTab() const;
+
+	static int32 GetIncrementFromCloseTabsDirection(ECloseTabsDirection);
+
+	/** Only allow closing tabs to the right when there is more than one tab open, the tab is of type Document or Major, and the tab is not furthest to the right. */
+	SLATE_API bool CanCloseTabsInDirectionOfForegroundTab(ECloseTabsDirection Direction) const;
+
+	/** Only allow closing all other tabs when there are more then one tab open and the tab is of type Document or Major. */
+	SLATE_API bool CanCloseAllButForegroundTab() const;
+
+	/** @return true if the foreground tab can be moved to a sidebar */
+	SLATE_API bool CanMoveForegroundTabToSidebar() const;
+
+	/** @return true if the specified tab can be moved to a sidebar */
+	SLATE_API bool IsTabAllowedInSidebar(TSharedPtr<SDockTab> Tab) const;
+
+	/** The tab well widget shows all tabs, keeps track of the selected tab, allows tab rearranging, etc. */
+	TSharedPtr<class SDockingTabWell> TabWell;
+
+	/** The borders that hold any potential inline content areas. */
+	SHorizontalBox::FSlot* InlineContentAreaLeft;
+	SHorizontalBox::FSlot* InlineContentAreaRight;
+
+	SVerticalBox::FSlot* TitleBarSlot;
+	TSharedPtr<SWidget> TitleBarContent;
+
+	TSharedPtr<SBorder> ContentSlot;
+
+	FOverlayManagement OverlayManagement;
+
+	SLATE_API TSharedRef<SWidget> MakeContextMenu();
+
+	/** Show the docking cross */
+	SLATE_API void ShowCross();
+
+	/** Hide the docking cross */
+	SLATE_API void HideCross();
+
+	/** Document Areas do not disappear when out of tabs, and instead say 'Document Area' */
+	bool bIsDocumentArea;
+
+	/** Animation that shows/hides the tab well; also used as a state machine to determine whether tab well is shown/hidden */
+	FCurveSequence ShowHideTabWell;
+
+	/** Grabs the scaling factor for the tab well size from the tab well animation. */
+	SLATE_API FVector2D GetTabWellScale() const;
+
+	/** Get the scale for the button that unhides the tab well */
+	SLATE_API FVector2D GetUnhideTabWellButtonScale() const;
+	/** Get the opacity for the button that unhides the tab well */
+	SLATE_API FSlateColor GetUnhideTabWellButtonOpacity() const;
+
+	/** Gets the background behind the tab stack */
+	SLATE_API const FSlateBrush* GetTabStackBorderImage() const;
+
+	/** Visibility of TitleBar spacer based on maximize/restore status of the window.
+	 ** This gives us a little more space to grab the title bar when the window is not maximized
+	*/
+	SLATE_API EVisibility GetMaximizeSpacerVisibility() const;
+
+	/** Bind tab commands into the ActionList */
+	SLATE_API void BindTabCommands();
+
+	/** Attempts to close the foreground tab when the CloseMajorTab command is executed */
+	SLATE_API void ExecuteCloseMajorTabCommand();
+
+	/** Attempts to find the foreground tab that can be closed by the CloseMajorTab command */
+	SLATE_API bool CanExecuteCloseMajorTabCommand();
+
+	/** Attempts to close the active tab when the CloseMinorTab command is executed */
+	SLATE_API void ExecuteCloseMinorTabCommand();
+
+	/** Attempts to find the active tab that can be closed by the CloseMinorTab command */
+	SLATE_API bool CanExecuteCloseMinorTabCommand();
+
+	/** Attempts to close the focused tab when the CloseFocusedTab command is executed */
+	SLATE_API void ExecuteCloseFocusedTabCommand();
+
+	/** Attempts to find the active tab that can be closed by the CloseFocusedTab command */
+	SLATE_API bool CanExecuteCloseFocusedTabCommand();
+
+	/** Called when this stack is resized. */
+	SLATE_API virtual void OnResized() override;
+
+#if DEBUG_TAB_MANAGEMENT
+	SLATE_API FString ShowPersistentTabs() const;
+#endif
+
+	/** Tab command list */
+	TSharedPtr<FUICommandList> ActionList;
+
+	/** Whether or not this tab stack is part of the title bar area */
+	bool bShowingTitleBarArea;
+};
+
+
+struct FTabMatcher
+{
+	FTabMatcher( const FTabId& InTabId, ETabState::Type InTabState = static_cast<ETabState::Type>(ETabState::ClosedTab | ETabState::OpenedTab | ETabState::SidebarTab), const bool InTreatIndexNoneAsWildcard = true )
+		: TabIdToMatch( InTabId )
+		, RequiredTabState( InTabState )
+		, TreatIndexNoneAsWildcard( InTreatIndexNoneAsWildcard )
+	{
+	}
+
+	bool operator()(const FTabManager::FTab& Candidate) const
+	{
+		return
+			((Candidate.TabState & RequiredTabState) != 0) &&
+			(Candidate.TabId.TabType == TabIdToMatch.TabType) &&
+			// INDEX_NONE is treated as a wildcard
+			((TreatIndexNoneAsWildcard && TabIdToMatch.InstanceId == INDEX_NONE) || TabIdToMatch.InstanceId == Candidate.TabId.InstanceId);;
+	}
+
+	FTabId TabIdToMatch;
+	ETabState::Type RequiredTabState;
+	bool TreatIndexNoneAsWildcard;
+};
