Submitting temporary fix for SSplitter to have a starting size, avoiding a visual glitch when first opening

diff --git a/Engine/Source/Runtime/Slate/Private/Framework/Docking/SPanelDrawerArea.cpp b/Engine/Source/Runtime/Slate/Private/Framework/Docking/SPanelDrawerArea.cpp
--- a/Engine/Source/Runtime/Slate/Private/Framework/Docking/SPanelDrawerArea.cpp
+++ b/Engine/Source/Runtime/Slate/Private/Framework/Docking/SPanelDrawerArea.cpp
@@ -237,13 +237,15 @@ void SPanelDrawerArea::SetupClosedLayout()
 
 void SPanelDrawerArea::SetupOpenedLayout()
 {
+	TSharedPtr<SSplitter> Splitter;
+
 	ChildSlot
 	[
-		SNew(SSplitter)
+		SAssignNew(Splitter, SSplitter)
 		.Orientation(Orient_Horizontal)
 		+ SSplitter::Slot()
 		.OnSlotResized(this, &SPanelDrawerArea::OnMainContentSlotResized)
-		.Value(this, &SPanelDrawerArea::GetMainContentCoefficient)
+		.Value(SPanelDrawerArea::GetMainContentCoefficient())
 		[
 			// Add Clipping for content that isn't responsive to dynamic sizes
 			SNew(SBox)
@@ -256,7 +258,7 @@ void SPanelDrawerArea::SetupOpenedLayout()
 		]
 		+ SSplitter::Slot()
 		.OnSlotResized(this, &SPanelDrawerArea::OnPanelDrawerSlotResized)
-		.Value(this, &SPanelDrawerArea::GetPanelDrawerCoefficient)
+		.Value(SPanelDrawerArea::GetPanelDrawerCoefficient())
 		[
 			// Add Clipping for content that isn't responsive to dynamic sizes
 			SNew(SBox)
@@ -270,6 +272,10 @@ void SPanelDrawerArea::SetupOpenedLayout()
 		]
 	];
 
+	// Temporary fix as values are not set by default with TSlateAttributes
+	Splitter->SlotAt(0).SetSizeValue(TAttribute<float>::CreateRaw(this, &SPanelDrawerArea::GetMainContentCoefficient));
+	Splitter->SlotAt(1).SetSizeValue(TAttribute<float>::CreateRaw(this, &SPanelDrawerArea::GetPanelDrawerCoefficient));
+
 	if (CVarNoAnimationOnTabForegroundedEvent.GetValueOnGameThread())
 	{
 		if (PanelDrawerData && PanelDrawerData->HostedTab)
