PanelDrawer now properly center the title of the hosted tab and has proper text collapsing when there is not enough space for it.

diff --git a/Engine/Source/Runtime/Slate/Private/Framework/Docking/STabPanelDrawer.cpp b/Engine/Source/Runtime/Slate/Private/Framework/Docking/STabPanelDrawer.cpp
--- a/Engine/Source/Runtime/Slate/Private/Framework/Docking/STabPanelDrawer.cpp
+++ b/Engine/Source/Runtime/Slate/Private/Framework/Docking/STabPanelDrawer.cpp
@@ -40,7 +40,7 @@ void STabPanelDrawer::Construct(const FArguments& InArgs, const TSharedRef<SDock
 	const FButtonStyle* const OpenInNewWindowButtonStyle = &FAppStyle::Get().GetWidgetStyle<FButtonStyle>("Docking.OpenTabInWindow.Button");
 
 	FMargin GrayLineMargin(2.f, 2.f, 2.f, 0.f);
-	
+
 	ChildSlot
 	[
 		SNew(SOverlay)
@@ -86,31 +86,35 @@ void STabPanelDrawer::Construct(const FArguments& InArgs, const TSharedRef<SDock
 		.VAlign(VAlign_Center)
 		.HAlign(HAlign_Fill)
 		[
-			SNew(SOverlay)
+			SAssignNew(TabLabelOverlay, SOverlay)
 			+ SOverlay::Slot()
-			.Padding(TAttribute<FMargin>::Create(TAttribute<FMargin>::FGetter::CreateSP(DisplayedTab.Get(), &SDockTab::GetTabPadding)))
+			.Padding(TAttribute<FMargin>::Create(TAttribute<FMargin>::FGetter::CreateSP(this, &STabPanelDrawer::GetTabLabelOverlayPadding)))
 			.VAlign(VAlign_Center)
 			.HAlign(HAlign_Center)
 			[
-				// Label sub HBox
-				SNew(SHorizontalBox)
-				.ToolTipText(DisplayedTab.Get(), &SDockTab::GetTabLabel)
-				// Tab Label
-				+ SHorizontalBox::Slot()
-				.FillWidth(1.0f)
-				.Padding(0.0f, 1.0f)
-				.VAlign(VAlign_Center)
+				SNew(SBox)
+				.MaxDesiredWidth(this, &STabPanelDrawer::GetTabLabelOverlayMaxSize)
 				[
-					DisplayedTab->LabelWidget.ToSharedRef()
-				]
+					// Label sub HBox
+					SNew(SHorizontalBox)
+					.ToolTipText(DisplayedTab.Get(), &SDockTab::GetTabLabel)
+					// Tab Label
+					+ SHorizontalBox::Slot()
+					.FillWidth(1.0f)
+					.Padding(0.0f, 1.0f)
+					.VAlign(VAlign_Center)
+					[
+						DisplayedTab->LabelWidget.ToSharedRef()
+					]
 
-				// Tab Label Suffix
-				+ SHorizontalBox::Slot()
-				.AutoWidth()
-				.Padding(0.0f, 1.0f)
-				.VAlign(VAlign_Center)
-				[
-					DisplayedTab->LabelSuffix.ToSharedRef()
+					// Tab Label Suffix
+					+ SHorizontalBox::Slot()
+					.AutoWidth()
+					.Padding(0.0f, 1.0f)
+					.VAlign(VAlign_Center)
+					[
+						DisplayedTab->LabelSuffix.ToSharedRef()
+					]
 				]
 			]
 
@@ -152,7 +156,7 @@ void STabPanelDrawer::Construct(const FArguments& InArgs, const TSharedRef<SDock
 				.HAlign(HAlign_Right)
 				.VAlign(VAlign_Center)
 				[
-					SNew(SButton)
+					SAssignNew(OpenTabInNewWindowButton, SButton)
 						.ButtonStyle(OpenInNewWindowButtonStyle)
 						.OnClicked(this, &STabPanelDrawer::OnButtonOpenTabInNewWindowClicked)
 						.ContentPadding(FMargin(0.0, 1.5, 0.0, 0.0))
@@ -168,7 +172,7 @@ void STabPanelDrawer::Construct(const FArguments& InArgs, const TSharedRef<SDock
 				.HAlign(HAlign_Right)
 				.VAlign(VAlign_Center)
 				[
-					SNew(SSpacer)
+					SAssignNew(ButtonSpacer, SSpacer)
 					.Size(OpenInNewWindowButtonStyle->Normal.ImageSize * 3 / 4)
 				]
 
@@ -177,7 +181,7 @@ void STabPanelDrawer::Construct(const FArguments& InArgs, const TSharedRef<SDock
 				.HAlign(HAlign_Right)
 				.VAlign(VAlign_Center)
 				[
-					SNew(SButton)
+					SAssignNew(DismissButton, SButton)
 					.ButtonStyle(CloseButtonStyle) 
 					.OnClicked(this, &STabPanelDrawer::OnButtonDismissTabClicked)
 					.ContentPadding(FMargin(0.0, 1.5, 0.0, 0.0))
@@ -190,7 +194,6 @@ void STabPanelDrawer::Construct(const FArguments& InArgs, const TSharedRef<SDock
 			]
 		]
 	];
-
 }
 
 void STabPanelDrawer::DismissTab() const
@@ -484,6 +487,55 @@ void STabPanelDrawer::InvokeTabOutsidePanelDrawer() const
 	}
 }
 
+FMargin STabPanelDrawer::GetTabLabelOverlayPadding() const
+{
+	if (DisplayedTab)
+	{
+		FMargin TabLabelMargin = DisplayedTab->GetTabPadding();
+		const float LargestSide = FMath::Max(TabLabelMargin.Left, TabLabelMargin.Right);
+		TabLabelMargin.Left = LargestSide;
+		TabLabelMargin.Right = LargestSide;
+		return TabLabelMargin;
+	}
+
+	return {};
+}
+
+FOptionalSize STabPanelDrawer::GetTabLabelOverlayMaxSize() const
+{
+	TSharedPtr<SWidget> TabLabelOverlayAsShared = TabLabelOverlay.Pin();
+	if (TabLabelOverlayAsShared && DisplayedTab)
+	{
+		const FGeometry& PaintGeo = TabLabelOverlayAsShared->GetPaintSpaceGeometry();
+		float TotalWidth = PaintGeo.GetLocalSize().X;
+		if (TotalWidth > 0.f)
+		{
+			FMargin TabLabelMargin = DisplayedTab->GetTabPadding();
+			float LargestSide = FMath::Max(TabLabelMargin.Left, TabLabelMargin.Right);
+
+			if (TSharedPtr<SWidget> OpenTabInNewWindowButtonAsShared = OpenTabInNewWindowButton.Pin())
+			{
+				LargestSide += OpenTabInNewWindowButtonAsShared->GetPaintSpaceGeometry().GetLocalSize().X;
+			}
+
+			if (TSharedPtr<SWidget> ButtonSpacerAsShared = ButtonSpacer.Pin())
+			{
+				// Double the size for to use the same spacing between the most left button and the label
+				LargestSide += ButtonSpacerAsShared->GetPaintSpaceGeometry().GetLocalSize().X * 2;
+			}
+
+			if (TSharedPtr<SWidget> DismissButtonAsShared = DismissButton.Pin())
+			{
+				LargestSide += DismissButtonAsShared->GetPaintSpaceGeometry().GetLocalSize().X;
+			}
+
+			return FOptionalSize(FMath::Max(0.f, TotalWidth - LargestSide * 2));
+		}
+	}
+
+	return {};
+}
+
 }
 
 #undef LOCTEXT_NAMESPACE
diff --git a/Engine/Source/Runtime/Slate/Private/Framework/Docking/STabPanelDrawer.h b/Engine/Source/Runtime/Slate/Private/Framework/Docking/STabPanelDrawer.h
--- a/Engine/Source/Runtime/Slate/Private/Framework/Docking/STabPanelDrawer.h
+++ b/Engine/Source/Runtime/Slate/Private/Framework/Docking/STabPanelDrawer.h
@@ -27,6 +27,8 @@ private:
 	TSharedRef<SWidget> MakeContextMenu() const;
 	void OpenTabInNewWindow() const;
 	void InvokeTabOutsidePanelDrawer() const;
+	FMargin GetTabLabelOverlayPadding() const;
+	FOptionalSize GetTabLabelOverlayMaxSize() const;
 
 	FReply OnButtonDismissTabClicked() const;
 	FReply OnMouseButtonUp(const FGeometry& MyGeometry, const FPointerEvent& MouseEvent);
@@ -38,6 +40,10 @@ private:
 	FReply OnButtonOpenTabInNewWindowClicked() const;
 
 	TSharedPtr<SDockTab> DisplayedTab;
+	TWeakPtr<SWidget> TabLabelOverlay;
+	TWeakPtr<SWidget> OpenTabInNewWindowButton;
+	TWeakPtr<SWidget> ButtonSpacer;
+	TWeakPtr<SWidget> DismissButton;
 	FSlateBrush DarkerTabBrush;
 };
 
