Fix combo box row triggering a selection change event on left mouse button down when the item is already selected, instead of LMB up like other non-selected items. This behavior would cause widgets underneath the combo box menu to receive and potentially wrongfully react to a LMB up event after the combo box menu is closed.
Turns out STableRow which is the baseclass for SComboRow already has a bAllowPreselectedItemActivation slate argument that does the right thing.
Slate.ComboRowMouseButtonHandlingV2 cvar can be set to 0 to revert back to the previous behavior in case this causes more problems that it solves

#rnx

diff --git a/Engine/Source/Runtime/Slate/Private/Widgets/Input/SComboBox.cpp b/Engine/Source/Runtime/Slate/Private/Widgets/Input/SComboBox.cpp
new file mode 100644
--- /dev/null
+++ b/Engine/Source/Runtime/Slate/Private/Widgets/Input/SComboBox.cpp
@@ -0,0 +1,16 @@
+// Copyright Epic Games, Inc. All Rights Reserved.
+
+#include "Widgets/Input/SComboBox.h"
+
+#include "HAL/IConsoleManager.h"
+
+namespace SlateComboRow
+{
+	bool bMouseButtonHandlingV2 = true;
+}
+
+FAutoConsoleVariableRef CVarComboRowMouseButtonDownHandlingV2(
+	TEXT("Slate.ComboRowMouseButtonHandlingV2"),
+	SlateComboRow::bMouseButtonHandlingV2,
+	TEXT("Whether to enable new mouse button event handling on SComboRow")
+);
diff --git a/Engine/Source/Runtime/Slate/Public/Widgets/Input/SComboBox.h b/Engine/Source/Runtime/Slate/Public/Widgets/Input/SComboBox.h
--- a/Engine/Source/Runtime/Slate/Public/Widgets/Input/SComboBox.h
+++ b/Engine/Source/Runtime/Slate/Public/Widgets/Input/SComboBox.h
@@ -32,6 +32,10 @@
 
 DECLARE_DELEGATE( FOnComboBoxOpening )
 
+namespace SlateComboRow
+{
+	extern SLATE_API bool bMouseButtonHandlingV2;
+}
 
 template<typename OptionType>
 class SComboRow : public STableRow< OptionType >
@@ -59,6 +63,7 @@ public:
 			typename STableRow<OptionType>::FArguments()
 			.Style(InArgs._Style)
 			.Padding(InArgs._Padding)
+			.bAllowPreselectedItemActivation(SlateComboRow::bMouseButtonHandlingV2)
 			.Content()
 			[
 				InArgs._Content.Widget
@@ -70,18 +75,21 @@ public:
 	// handle case where user clicks on an existing selected item
 	virtual FReply OnMouseButtonDown( const FGeometry& MyGeometry, const FPointerEvent& MouseEvent )
 	{
-		if (MouseEvent.GetEffectingButton() == EKeys::LeftMouseButton )
+		if (!STableRow<OptionType>::bAllowPreselectedItemActivation)
 		{
-			TSharedPtr< ITypedTableView<OptionType> > OwnerWidget = this->OwnerTablePtr.Pin();
-
-			const TObjectPtrWrapTypeOf<OptionType>* MyItem = OwnerWidget->Private_ItemFromWidget( this );
-			const bool bIsSelected = OwnerWidget->Private_IsItemSelected( *MyItem );
-				
-			if (bIsSelected)
+			if (MouseEvent.GetEffectingButton() == EKeys::LeftMouseButton )
 			{
-				// Reselect content to ensure selection is taken
-				OwnerWidget->Private_SignalSelectionChanged(ESelectInfo::OnMouseClick);
-				return FReply::Handled();
+				TSharedPtr< ITypedTableView<OptionType> > OwnerWidget = this->OwnerTablePtr.Pin();
+
+				const TObjectPtrWrapTypeOf<OptionType>* MyItem = OwnerWidget->Private_ItemFromWidget( this );
+				const bool bIsSelected = OwnerWidget->Private_IsItemSelected( *MyItem );
+
+				if (bIsSelected)
+				{
+					// Reselect content to ensure selection is taken
+					OwnerWidget->Private_SignalSelectionChanged(ESelectInfo::OnMouseClick);
+					return FReply::Handled();
+				}
 			}
 		}
 		return STableRow<OptionType>::OnMouseButtonDown(MyGeometry, MouseEvent);
