[Slate] Add a Slate.EnableTooltipForceFieldDraw console variable, with which it's possible to draw the effective tooltip force fields.

#Merge 47910358.

diff --git a/Engine/Source/Runtime/Slate/Private/Framework/Application/SlateUser.cpp b/Engine/Source/Runtime/Slate/Private/Framework/Application/SlateUser.cpp
--- a/Engine/Source/Runtime/Slate/Private/Framework/Application/SlateUser.cpp
+++ b/Engine/Source/Runtime/Slate/Private/Framework/Application/SlateUser.cpp
@@ -74,6 +74,14 @@ FAutoConsoleVariableRef CVarAllowTooltipsWhileMouseDown(
 	bAllowTooltipsWhileMouseDown,
 	TEXT(""));
 
+#if WITH_SLATE_DEBUGGING
+static bool bEnableTooltipForceFieldDraw = false;
+FAutoConsoleVariableRef CVarEnableTooltipForceFieldDraw(
+	TEXT("Slate.EnableTooltipForceFieldDraw"),
+	bEnableTooltipForceFieldDraw,
+	TEXT("If turned on, the rectangles of effective tooltip force fields will be visualized."));
+#endif
+
 //////////////////////////////////////////////////////////////////////////
 // FSlateVirtualUserHandle
 //////////////////////////////////////////////////////////////////////////
@@ -154,6 +162,10 @@ FSlateUser::FSlateUser(FPlatformUserId InPlatformUser, TSharedPtr<ICursor> InCur
 
 FSlateUser::~FSlateUser()
 {
+#if WITH_SLATE_DEBUGGING
+	ClearForceFieldDebug();
+#endif
+
 	UE_LOG(LogSlate, Log, TEXT("Slate User Destroyed.  User Index %d, Is Virtual User: %d"), UserIndex, IsVirtualUser());
 }
 
@@ -1045,6 +1057,36 @@ TSharedRef<SWindow> FSlateUser::GetOrCreateTooltipWindow()
 	return NewTooltipWindow;
 }
 
+#if WITH_SLATE_DEBUGGING
+void FSlateUser::ClearForceFieldDebug()
+{
+	if (!DebugForceFieldRects.IsEmpty())
+	{
+		DebugForceFieldRects.Reset();
+		FSlateDebugging::PaintDebugElements.RemoveAll(this);
+	}
+}
+
+void FSlateUser::HandlePaint(const FPaintArgs& InArgs, const FGeometry& InAllottedGeometry,
+	FSlateWindowElementList& InOutDrawElements, int32& InOutLayerId)
+{
+	for (const FSlateRect& Rect : DebugForceFieldRects)
+	{
+		FSlateDrawElement::MakeBox(
+			InOutDrawElements,
+			++InOutLayerId,
+			InAllottedGeometry.ToPaintGeometry(
+				Rect.GetSize2f(),
+				FSlateLayoutTransform(Rect.GetTopLeft2f() - InArgs.GetWindowToDesktopTransform())
+			),
+			FCoreStyle::Get().GetBrush("Debug.Border"),
+			ESlateDrawEffect::None,
+			FLinearColor::Blue
+		);
+	}
+}
+#endif
+
 void FSlateUser::NotifyTouchStarted(const FPointerEvent& TouchEvent)
 {
 	UE_CLOG(PointerPositionsByIndex.Contains(TouchEvent.GetPointerIndex()), LogSlate, Error, TEXT("SlateUser [%d] notified of a touch starting for pointer [%d] without finding out it ever ended."), TouchEvent.GetUserIndex(), TouchEvent.GetPointerIndex());
@@ -1227,6 +1269,10 @@ void FSlateUser::ResetDragDropContent()
 
 void FSlateUser::UpdateTooltip(const FMenuStack& MenuStack, bool bCanSpawnNewTooltip)
 {
+#if WITH_SLATE_DEBUGGING
+	ClearForceFieldDebug();
+#endif
+
 	FSlateApplication& SlateApp = FSlateApplication::Get();
 	if (!SlateApp.GetAllowTooltips())
 	{
@@ -1286,20 +1332,35 @@ void FSlateUser::UpdateTooltip(const FMenuStack& MenuStack, bool bCanSpawnNewToo
 		// Make sure we account for all widgets in the path with a force field, even if we've already found the tooltip we'll be showing
 		if (CurWidget->HasToolTipForceField())
 		{
+			const FSlateRect WidgetRect = ArrangedWidget->Geometry.GetLayoutBoundingRect();
 			if (!ForceFieldRect.IsSet())
 			{
-				ForceFieldRect = ArrangedWidget->Geometry.GetLayoutBoundingRect();
+				ForceFieldRect = WidgetRect;
 			}
 			else
 			{
 				// Grow the rect to encompass this geometry to be super safe
 				// The parent's rect should always be inclusive of its child, so this should usually be overkill.
-				ForceFieldRect = ForceFieldRect->Expand(ArrangedWidget->Geometry.GetLayoutBoundingRect());
+				ForceFieldRect = ForceFieldRect->Expand(WidgetRect);
 			}
 			ForceFieldRect = (1.0f / DPIScaleFactor) * ForceFieldRect.GetValue();
+
+#if WITH_SLATE_DEBUGGING
+			if (bEnableTooltipForceFieldDraw)
+			{
+				DebugForceFieldRects.Add(WidgetRect);
+			}
+#endif
 		}
 	}
 
+#if WITH_SLATE_DEBUGGING
+	if (!DebugForceFieldRects.IsEmpty())
+	{
+		FSlateDebugging::PaintDebugElements.AddRaw(this, &FSlateUser::HandlePaint);
+	}
+#endif
+
 	TSharedPtr<SWidget> NewTooltipVisualizer;
 	TSharedPtr<IToolTip> ActiveTooltip = ActiveTooltipInfo.Tooltip.Pin();
 	const bool bTooltipChanged = NewTooltip != ActiveTooltip;
diff --git a/Engine/Source/Runtime/Slate/Public/Framework/Application/SlateUser.h b/Engine/Source/Runtime/Slate/Public/Framework/Application/SlateUser.h
--- a/Engine/Source/Runtime/Slate/Public/Framework/Application/SlateUser.h
+++ b/Engine/Source/Runtime/Slate/Public/Framework/Application/SlateUser.h
@@ -201,6 +201,13 @@ private:
 	SLATE_API void LockCursorInternal(const FWidgetPath& WidgetPath);
 	SLATE_API TSharedRef<SWindow> GetOrCreateTooltipWindow();
 
+#if WITH_SLATE_DEBUGGING
+	SLATE_API void ClearForceFieldDebug();
+	SLATE_API void HandlePaint(const FPaintArgs& InArgs, const FGeometry& InAllottedGeometry, FSlateWindowElementList& InOutDrawElements, int32& InOutLayerId);
+
+	TArray<FSlateRect> DebugForceFieldRects;
+#endif
+
 	/** The index the user was assigned. */
 	int32 UserIndex = INDEX_NONE;
 
