[Slate]  Fix Status Bar tooltips appeared in another screen (or overlapped the status bar if there was no monitor to the right).

#Merge 47910374 by kristov.morva1

diff --git a/Engine/Source/Runtime/Slate/Private/Framework/Application/SlateUser.cpp b/Engine/Source/Runtime/Slate/Private/Framework/Application/SlateUser.cpp
--- a/Engine/Source/Runtime/Slate/Private/Framework/Application/SlateUser.cpp
+++ b/Engine/Source/Runtime/Slate/Private/Framework/Application/SlateUser.cpp
@@ -1464,6 +1464,14 @@ void FSlateUser::UpdateTooltip(const FMenuStack& MenuStack, bool bCanSpawnNewToo
 		// Make sure the tooltip needs to be offset
 		if (TooltipShift.X != 0.0f || TooltipShift.Y != 0.0f)
 		{
+			// If in the bottom right corner there's no space for the tooltip, let's try to shift it upwards instead
+			const FSlateRect WorkArea = FSlateApplication::Get().GetPreferredWorkArea();
+			if (DesiredLocation.X + TooltipShift.X + TooltipSize.X > WorkArea.Right &&
+				DesiredLocation.Y + TooltipShift.Y + TooltipSize.Y > WorkArea.Bottom)
+			{
+				TooltipShift.Y = ForceFieldRect->Top - SlateDefs::TooltipOffsetFromForceField.Y - TooltipSize.Y - DesiredLocation.Y;
+			}
+
 			// Find the best edge to move the tooltip towards
 			if (ActiveTooltipInfo.OffsetDirection == ETooltipOffsetDirection::Undetermined)
 			{
