Slate Dynamic Invalidation - SHyperlink

diff --git a/Engine/Source/Runtime/Slate/Public/Widgets/Input/SHyperlink.h b/Engine/Source/Runtime/Slate/Public/Widgets/Input/SHyperlink.h
--- a/Engine/Source/Runtime/Slate/Public/Widgets/Input/SHyperlink.h
+++ b/Engine/Source/Runtime/Slate/Public/Widgets/Input/SHyperlink.h
@@ -20,6 +20,12 @@
 enum class ETextFlowDirection : uint8;
 enum class ETextShapingMethod : uint8;
 
+template<>
+struct TWidgetTypeTraits<class SHyperlink>
+{
+	static constexpr bool SupportsInvalidation() { return true; }
+};
+
 /**
  * A hyperlink widget is what you would expect from a browser hyperlink.
  * When a hyperlink is clicked in invokes an OnNavigate() delegate.
