Only create property & text tooltip widgets on demand to speed up editor startups. Saves up to ~400ms when switching between states in state centric view.
- Add SDeferredToolTip to wrap any previous IToolTip into deferred creation.
- Add SDeferredToolTipText specifically for text tooltips generated by slate app.

#tests editor
#Synced-CL 48414911

diff --git a/Engine/Source/Editor/PropertyEditor/Private/PropertyEditorHelpers.cpp b/Engine/Source/Editor/PropertyEditor/Private/PropertyEditorHelpers.cpp
--- a/Engine/Source/Editor/PropertyEditor/Private/PropertyEditorHelpers.cpp
+++ b/Engine/Source/Editor/PropertyEditor/Private/PropertyEditorHelpers.cpp
@@ -3,6 +3,7 @@
 #include "PropertyEditorHelpers.h"
 #include "Widgets/Layout/SBorder.h"
 #include "Widgets/Input/SCheckBox.h"
+#include "Widgets/SDeferredToolTip.h"
 #include "IDocumentation.h"
 
 #include "PropertyHandleImpl.h"
@@ -39,6 +40,19 @@ void SPropertyNameWidget::Construct( const FArguments& InArgs, TSharedPtr<FPrope
 {
 	PropertyEditor = InPropertyEditor;
 
+	TWeakPtr<FPropertyEditor> WeakPropertyEditor = PropertyEditor;
+	auto CreateDocumentationToolTipDeferred = [WeakPropertyEditor]()
+	{
+		if (TSharedPtr<FPropertyEditor> PropertyEditorPinned = WeakPropertyEditor.Pin())
+		{
+			return IDocumentation::Get()->CreateToolTip(PropertyEditorPinned->GetToolTipText(), NULL, PropertyEditorPinned->GetDocumentationLink(), PropertyEditorPinned->GetDocumentationExcerptName());
+		};
+
+		return SNew(SToolTip)
+			.Text(LOCTEXT("EditorDestroyedToolTip", "Invalid ToolTip, source editor closed."));
+	};
+	TSharedRef<SDeferredToolTip> DeferredDocumentationToolTip = MakeShared<SDeferredToolTip>(FOnGetDeferredToolTip::CreateLambda(CreateDocumentationToolTipDeferred));
+
 	TSharedPtr<SHorizontalBox> HorizontalBox;
 	ChildSlot
 	[
@@ -54,7 +68,7 @@ void SPropertyNameWidget::Construct( const FArguments& InArgs, TSharedPtr<FPrope
 			[
 				SNew( SPropertyEditorTitle, PropertyEditor.ToSharedRef() )
 				.OnDoubleClicked( InArgs._OnDoubleClicked )
-                .ToolTip( IDocumentation::Get()->CreateToolTip( PropertyEditor->GetToolTipText(), NULL, PropertyEditor->GetDocumentationLink(), PropertyEditor->GetDocumentationExcerptName() ) )
+                .ToolTip(DeferredDocumentationToolTip)
 			]
 		]
 	
diff --git a/Engine/Source/Runtime/Slate/Private/Framework/Application/SlateApplication.cpp b/Engine/Source/Runtime/Slate/Private/Framework/Application/SlateApplication.cpp
--- a/Engine/Source/Runtime/Slate/Private/Framework/Application/SlateApplication.cpp
+++ b/Engine/Source/Runtime/Slate/Private/Framework/Application/SlateApplication.cpp
@@ -22,6 +22,7 @@
 #include "Widgets/Images/SImage.h"
 #include "Widgets/Layout/SBox.h"
 #include "Widgets/SToolTip.h"
+#include "Widgets/SDeferredToolTip.h"
 #include "Widgets/SViewport.h"
 #include "Framework/Application/SWindowTitleBar.h"
 #include "Input/Events.h"
@@ -4815,15 +4816,13 @@ TSharedRef<SWidget> FSlateApplication::MakeWindowTitleBar(const FWindowTitleBarA
 
 TSharedRef<IToolTip> FSlateApplication::MakeToolTip(const TAttribute<FText>& ToolTipText)
 {
-	return SNew(SToolTip)
-		.Text(ToolTipText);
+	return MakeShared<SDeferredToolTipText>(ToolTipText);
 }
 
 
 TSharedRef<IToolTip> FSlateApplication::MakeToolTip( const FText& ToolTipText )
 {
-	return SNew(SToolTip)
-		.Text(ToolTipText);
+	return MakeShared<SDeferredToolTipText>(ToolTipText);
 }
 
 /* FGenericApplicationMessageHandler interface
diff --git a/Engine/Source/Runtime/Slate/Private/Widgets/SDeferredToolTip.cpp b/Engine/Source/Runtime/Slate/Private/Widgets/SDeferredToolTip.cpp
new file mode 100644
--- /dev/null
+++ b/Engine/Source/Runtime/Slate/Private/Widgets/SDeferredToolTip.cpp
@@ -0,0 +1,214 @@
+// Copyright Epic Games, Inc. All Rights Reserved.
+
+#include "Widgets/SDeferredToolTip.h"
+#include "Widgets/SToolTip.h"
+
+
+//////////////////////////////////////////////////////////////////////////
+// SDeferredToolTip
+
+SDeferredToolTip::SDeferredToolTip(FOnGetDeferredToolTip InOnGetDeferredToolTip)
+	: OnGetDeferredToolTip(InOnGetDeferredToolTip)
+{
+}
+
+TSharedRef<SWidget> SDeferredToolTip::AsWidget()
+{
+	TryCacheToolTip();
+	if (CachedToolTip)
+	{
+		return CachedToolTip->AsWidget();
+	}
+
+	return SNullWidget::NullWidget;
+}
+
+TSharedRef<SWidget> SDeferredToolTip::GetContentWidget() 
+{
+	TryCacheToolTip();
+	if (CachedToolTip)
+	{
+		return CachedToolTip->GetContentWidget();
+	}
+
+	return SNullWidget::NullWidget;
+}
+
+void SDeferredToolTip::SetContentWidget(const TSharedRef<SWidget>& InContentWidget) 
+{
+	TryCacheToolTip();
+	if (CachedToolTip)
+	{
+		CachedToolTip->SetContentWidget(InContentWidget);
+	}
+}
+
+void SDeferredToolTip::ResetContentWidget() 
+{
+	TryCacheToolTip();
+	if (CachedToolTip)
+	{
+		CachedToolTip->ResetContentWidget();
+	}
+}
+
+bool SDeferredToolTip::IsEmpty() const 
+{
+	TryCacheToolTip();
+	if (CachedToolTip)
+	{
+		return CachedToolTip->IsEmpty();
+	}
+
+	return true;
+}
+
+bool SDeferredToolTip::IsInteractive() const 
+{
+	TryCacheToolTip();
+	if (CachedToolTip)
+	{
+		return CachedToolTip->IsInteractive();
+	}
+
+	return false;
+}
+
+void SDeferredToolTip::OnClosed() 
+{
+	TryCacheToolTip();
+	if (CachedToolTip)
+	{
+		return CachedToolTip->OnClosed();
+	}
+}
+
+void SDeferredToolTip::OnOpening() 
+{
+	TryCacheToolTip();
+	if (CachedToolTip)
+	{
+		return CachedToolTip->OnOpening();
+	}
+}
+
+void SDeferredToolTip::OnSetInteractiveWindowLocation(FVector2D& InOutDesiredLocation) const
+{
+	TryCacheToolTip();
+	if (CachedToolTip)
+	{
+		CachedToolTip->OnSetInteractiveWindowLocation(InOutDesiredLocation);
+	}
+}
+
+void SDeferredToolTip::TryCacheToolTip() const
+{
+	if (!CachedToolTip && OnGetDeferredToolTip.IsBound())
+	{
+		CachedToolTip = OnGetDeferredToolTip.Execute();
+	}
+}
+
+
+//////////////////////////////////////////////////////////////////////////
+// SDeferredToolTipText
+
+SDeferredToolTipText::SDeferredToolTipText(const TAttribute<FText>& InToolTipText)
+	: ToolTipText(InToolTipText)
+{
+}
+
+SDeferredToolTipText::SDeferredToolTipText(const FText& InToolTipText)
+	: ToolTipText(InToolTipText)
+{
+}
+
+TSharedRef<SWidget> SDeferredToolTipText::AsWidget()
+{
+	TryCacheToolTip();
+	if (CachedToolTip)
+	{
+		return CachedToolTip->AsWidget();
+	}
+
+	return SNullWidget::NullWidget;
+}
+
+TSharedRef<SWidget> SDeferredToolTipText::GetContentWidget() 
+{
+	TryCacheToolTip();
+	if (CachedToolTip)
+	{
+		return CachedToolTip->GetContentWidget();
+	}
+
+	return SNullWidget::NullWidget;
+}
+
+void SDeferredToolTipText::SetContentWidget(const TSharedRef<SWidget>& InContentWidget) 
+{
+	TryCacheToolTip();
+	if (CachedToolTip)
+	{
+		CachedToolTip->SetContentWidget(InContentWidget);
+	}
+}
+
+void SDeferredToolTipText::ResetContentWidget() 
+{
+	TryCacheToolTip();
+	if (CachedToolTip)
+	{
+		CachedToolTip->ResetContentWidget();
+	}
+}
+
+bool SDeferredToolTipText::IsEmpty() const 
+{
+	TryCacheToolTip();
+	if (CachedToolTip)
+	{
+		return CachedToolTip->IsEmpty();
+	}
+
+	return true;
+}
+
+bool SDeferredToolTipText::IsInteractive() const 
+{
+	TryCacheToolTip();
+	if (CachedToolTip)
+	{
+		return CachedToolTip->IsInteractive();
+	}
+
+	return false;
+}
+
+void SDeferredToolTipText::OnClosed() 
+{
+	// SToolTip does No-op. We choose to not forward no-op to no-op.
+}
+
+void SDeferredToolTipText::OnOpening() 
+{
+	// SToolTip does No-op. We choose to not forward no-op to no-op.
+}
+
+void SDeferredToolTipText::OnSetInteractiveWindowLocation(FVector2D& InOutDesiredLocation) const
+{
+	TryCacheToolTip();
+	if (CachedToolTip)
+	{
+		CachedToolTip->OnSetInteractiveWindowLocation(InOutDesiredLocation);
+	}
+}
+
+void SDeferredToolTipText::TryCacheToolTip() const
+{
+	if (!CachedToolTip)
+	{
+		SAssignNew(CachedToolTip, SToolTip)
+			.Text(ToolTipText);
+	}
+}
\ No newline at end of file
diff --git a/Engine/Source/Runtime/Slate/Public/Widgets/SDeferredToolTip.h b/Engine/Source/Runtime/Slate/Public/Widgets/SDeferredToolTip.h
new file mode 100644
--- /dev/null
+++ b/Engine/Source/Runtime/Slate/Public/Widgets/SDeferredToolTip.h
@@ -0,0 +1,92 @@
+// Copyright Epic Games, Inc. All Rights Reserved.
+
+#pragma once
+
+#include "CoreMinimal.h"
+#include "Widgets/IToolTip.h"
+
+class SWidget;
+class SToolTip;
+
+/** Called deferred to retrieve the tooltip widget when needed. */
+DECLARE_DELEGATE_RetVal(TSharedPtr<IToolTip>, FOnGetDeferredToolTip);
+
+/**
+ * Interface for deferred tooltips. These do not create unless needed, so have lower startup cost
+ */
+class SDeferredToolTip : public IToolTip
+{
+public:
+
+	SLATE_API SDeferredToolTip(FOnGetDeferredToolTip InOnGetDeferredToolTip);
+
+public:
+
+	// ~Begin IToolTip Interface 
+	SLATE_API virtual TSharedRef<SWidget> AsWidget() override;
+	SLATE_API virtual TSharedRef<SWidget> GetContentWidget() override;
+	SLATE_API virtual void SetContentWidget(const TSharedRef<SWidget>& InContentWidget) override;
+	SLATE_API virtual void ResetContentWidget() override;
+	SLATE_API virtual bool IsEmpty() const override;
+	SLATE_API virtual bool IsInteractive() const override;
+	SLATE_API virtual void OnOpening() override;
+	SLATE_API virtual void OnClosed() override;
+	SLATE_API virtual void OnSetInteractiveWindowLocation(FVector2D& InOutDesiredLocation) const override;
+	// ~End IToolTip Interface 
+
+protected:
+
+	/** Delegate to generate tooltip. Will be called deferred. */
+	FOnGetDeferredToolTip OnGetDeferredToolTip = {};
+
+private:
+
+	/** Helper method to cache the tooltip if needed */
+	void TryCacheToolTip() const;
+
+private:
+
+	/** Cached tooltip creation result */
+	mutable TSharedPtr<IToolTip> CachedToolTip = nullptr;
+};
+
+
+/**
+ * Specialized deferred tooltip just for text
+ */
+class SDeferredToolTipText : public IToolTip
+{
+public:
+
+	SLATE_API SDeferredToolTipText(const TAttribute<FText>& InToolTipText);
+	SLATE_API SDeferredToolTipText(const FText& InToolTipText);
+
+public:
+
+	// ~Begin IToolTip Interface 
+	SLATE_API virtual TSharedRef<SWidget> AsWidget() override;
+	SLATE_API virtual TSharedRef<SWidget> GetContentWidget() override;
+	SLATE_API virtual void SetContentWidget(const TSharedRef<SWidget>& InContentWidget) override;
+	SLATE_API virtual void ResetContentWidget() override;
+	SLATE_API virtual bool IsEmpty() const override;
+	SLATE_API virtual bool IsInteractive() const override;
+	SLATE_API virtual void OnOpening() override;
+	SLATE_API virtual void OnClosed() override;
+	SLATE_API virtual void OnSetInteractiveWindowLocation(FVector2D& InOutDesiredLocation) const override;
+	// ~End IToolTip Interface 
+
+public:
+
+	/** Text attribute we will use to make tooltip on demand */
+	TAttribute<FText> ToolTipText = {};
+
+private:
+
+	/** Helper method to cache the tooltip if needed */
+	void TryCacheToolTip() const;
+
+private:
+
+	/** Cached tooltip creation result */
+	mutable TSharedPtr<SToolTip> CachedToolTip = nullptr;
+};
