Silence Clang SA warning about unused return result
#rnx

diff --git a/Engine/Source/Runtime/Slate/Public/Framework/Docking/TabManager.h b/Engine/Source/Runtime/Slate/Public/Framework/Docking/TabManager.h
--- a/Engine/Source/Runtime/Slate/Public/Framework/Docking/TabManager.h
+++ b/Engine/Source/Runtime/Slate/Public/Framework/Docking/TabManager.h
@@ -33,7 +33,7 @@ namespace UE::Editor::HomeScreen
 		{
 			if (const IConsoleVariable* CVar = IConsoleManager::Get().FindConsoleVariable(TEXT("HomeScreen.EnableHomeScreen")))
 			{
-				ensureAlwaysMsgf(!EnumHasAnyFlags(CVar->GetFlags(), ECVF_Default), TEXT("The CVar should have already been set from commandline."));
+				(void)ensureAlwaysMsgf(!EnumHasAnyFlags(CVar->GetFlags(), ECVF_Default), TEXT("The CVar should have already been set from commandline."));
 				bIsHomeScreenEnabled = CVar->GetBool();
 			}
 		});
