Added a visibility attribute to SSegmentedControl slots
#rnx
#Synced-CL 47206406

diff --git a/Engine/Source/Runtime/Slate/Public/Widgets/Input/SSegmentedControl.h b/Engine/Source/Runtime/Slate/Public/Widgets/Input/SSegmentedControl.h
--- a/Engine/Source/Runtime/Slate/Public/Widgets/Input/SSegmentedControl.h
+++ b/Engine/Source/Runtime/Slate/Public/Widgets/Input/SSegmentedControl.h
@@ -38,6 +38,7 @@ public:
 			, _Text()
 			, _Tooltip()
 			, _Icon(nullptr)
+			, _Visibility(EVisibility::Visible)
 			, _Value(InValue)
 			{ }
 
@@ -46,6 +47,7 @@ public:
 			SLATE_ATTRIBUTE(FText, ToolTip)
 			SLATE_ATTRIBUTE(TSharedPtr<IToolTip>, ToolTipWidget)
 			SLATE_ATTRIBUTE(const FSlateBrush*, Icon)
+			SLATE_ATTRIBUTE(EVisibility, Visibility)
 			SLATE_ARGUMENT(TOptional<OptionType>, Value)
 		SLATE_SLOT_END_ARGS()
 
@@ -69,6 +71,10 @@ public:
 			{
 				_Icon = MoveTemp(InArgs._Icon);
 			}
+			if (InArgs._Visibility.IsSet())
+			{
+				_Visibility = MoveTemp(InArgs._Visibility);
+			}
 			if (InArgs._Value.IsSet())
 			{
 				_Value = MoveTemp(InArgs._Value.GetValue());
@@ -105,6 +111,16 @@ public:
 			return _Tooltip.Get();
 		}
 
+		void SetVisibility(EVisibility InVisibility)
+		{
+			_Visibility = MoveTemp(InVisibility);
+		}
+
+		EVisibility GetVisibility() const
+		{
+			return _Visibility.Get();
+		}
+
 	friend SSegmentedControl<OptionType>;
 
 	private:
@@ -112,6 +128,7 @@ public:
 		TAttribute<FText> _Tooltip;
 		TAttribute<TSharedPtr<IToolTip>> _ToolTipWidget;
 		TAttribute<const FSlateBrush*> _Icon;
+		TAttribute<EVisibility> _Visibility;
 
 		OptionType _Value;
 		TWeakPtr<SCheckBox> _CheckBox;
@@ -294,6 +311,7 @@ public:
 				.IsChecked(GetCheckBoxStateAttribute(ChildValue))
 				.OnCheckStateChanged(this, &SSegmentedControl::CommitValue, ChildValue)
 				.Padding(UniformPadding)
+				.Visibility(ChildSlotPtr->_Visibility)
 				[
 					Child
 				]
