Slate Dynamic Invalidation - Layout Widgets

#Synced-CL 48711348

diff --git a/Engine/Source/Runtime/Slate/Private/Widgets/Layout/SUniformWrapPanel.cpp b/Engine/Source/Runtime/Slate/Private/Widgets/Layout/SUniformWrapPanel.cpp
--- a/Engine/Source/Runtime/Slate/Private/Widgets/Layout/SUniformWrapPanel.cpp
+++ b/Engine/Source/Runtime/Slate/Private/Widgets/Layout/SUniformWrapPanel.cpp
@@ -6,6 +6,9 @@
 SLATE_IMPLEMENT_WIDGET(SUniformWrapPanel)
 void SUniformWrapPanel::PrivateRegisterAttributes(FSlateAttributeInitializer& AttributeInitializer)
 {
+	FSlateWidgetSlotAttributeInitializer Initializer = SLATE_ADD_PANELCHILDREN_DEFINITION(AttributeInitializer, Children);
+	FSlot::RegisterAttributes(Initializer);
+
 	SLATE_ADD_MEMBER_ATTRIBUTE_DEFINITION(AttributeInitializer, SlotPadding, EInvalidateWidgetReason::Layout);
 	SLATE_ADD_MEMBER_ATTRIBUTE_DEFINITION(AttributeInitializer, MinDesiredSlotWidth, EInvalidateWidgetReason::Layout);
 	SLATE_ADD_MEMBER_ATTRIBUTE_DEFINITION(AttributeInitializer, MinDesiredSlotHeight, EInvalidateWidgetReason::Layout);
diff --git a/Engine/Source/Runtime/Slate/Public/Widgets/Layout/SUniformWrapPanel.h b/Engine/Source/Runtime/Slate/Public/Widgets/Layout/SUniformWrapPanel.h
--- a/Engine/Source/Runtime/Slate/Public/Widgets/Layout/SUniformWrapPanel.h
+++ b/Engine/Source/Runtime/Slate/Public/Widgets/Layout/SUniformWrapPanel.h
@@ -27,14 +27,14 @@ class SUniformWrapPanel : public SPanel
 
 public:
 	/** Stores the per-child info for this panel type */
-	struct FSlot : public TSlotBase<FSlot>, public TAlignmentWidgetSlotMixin<FSlot>
+	struct FSlot : public TWidgetSlotWithAttributeSupport<FSlot>, public TAlignmentWidgetSlotMixin<FSlot>
 	{
 		FSlot()
-			: TSlotBase<FSlot>()
+			: TWidgetSlotWithAttributeSupport<FSlot>()
 			, TAlignmentWidgetSlotMixin<FSlot>(HAlign_Fill, VAlign_Fill)
 		{ }
 
-		SLATE_SLOT_BEGIN_ARGS_OneMixin(FSlot, TSlotBase<FSlot>, TAlignmentWidgetSlotMixin<FSlot>)
+		SLATE_SLOT_BEGIN_ARGS_OneMixin(FSlot, TWidgetSlotWithAttributeSupport<FSlot>, TAlignmentWidgetSlotMixin<FSlot>)
 		SLATE_SLOT_END_ARGS()
 
 		void Construct(const FChildren& SlotOwner, FSlotArguments&& InArgs)
@@ -42,6 +42,10 @@ public:
 			TSlotBase<FSlot>::Construct(SlotOwner, MoveTemp(InArgs));
 			TAlignmentWidgetSlotMixin<FSlot>::ConstructMixin(SlotOwner, MoveTemp(InArgs));
 		}
+		static void RegisterAttributes(FSlateWidgetSlotAttributeInitializer& AttributeInitializer)
+		{
+			TWidgetSlotWithAttributeSupport::RegisterAttributes(AttributeInitializer);
+		}
 	};
 
 	SLATE_API SUniformWrapPanel();
diff --git a/Engine/Source/Runtime/Slate/Public/Widgets/Layout/SWrapBox.h b/Engine/Source/Runtime/Slate/Public/Widgets/Layout/SWrapBox.h
--- a/Engine/Source/Runtime/Slate/Public/Widgets/Layout/SWrapBox.h
+++ b/Engine/Source/Runtime/Slate/Public/Widgets/Layout/SWrapBox.h
@@ -13,6 +13,12 @@
 
 class FArrangedChildren;
 
+template <>
+struct TWidgetTypeTraits<class SWrapBox>
+{
+	static constexpr bool SupportsInvalidation() { return true; }
+};
+
 /**
  * With EOrientation::Orient_Horizontal
  * Arranges widgets left-to-right.
diff --git a/Engine/Source/Runtime/Slate/Public/Widgets/SCanvas.h b/Engine/Source/Runtime/Slate/Public/Widgets/SCanvas.h
--- a/Engine/Source/Runtime/Slate/Public/Widgets/SCanvas.h
+++ b/Engine/Source/Runtime/Slate/Public/Widgets/SCanvas.h
@@ -53,7 +53,7 @@ public:
 	class FSlot : public TWidgetSlotWithAttributeSupport<FSlot>, public TAlignmentWidgetSlotMixin<FSlot>
 	{
 	public:
-		SLATE_SLOT_BEGIN_ARGS_OneMixin(FSlot, TSlotBase<FSlot>, TAlignmentWidgetSlotMixin<FSlot>)
+		SLATE_SLOT_BEGIN_ARGS_OneMixin(FSlot, TWidgetSlotWithAttributeSupport<FSlot>, TAlignmentWidgetSlotMixin<FSlot>)
 			SLATE_ATTRIBUTE(FVector2D, Position)
 			SLATE_ATTRIBUTE(FVector2D, Size)
 		SLATE_SLOT_END_ARGS()
