Add VisibilityOverride to button block ctor/args

#rnx
#Synced-CL 48891598

diff --git a/Engine/Source/Runtime/Slate/Private/Framework/MultiBox/MultiBoxBuilder.cpp b/Engine/Source/Runtime/Slate/Private/Framework/MultiBox/MultiBoxBuilder.cpp
--- a/Engine/Source/Runtime/Slate/Private/Framework/MultiBox/MultiBoxBuilder.cpp
+++ b/Engine/Source/Runtime/Slate/Private/Framework/MultiBox/MultiBoxBuilder.cpp
@@ -692,6 +692,11 @@ void FToolBarBuilder::AddToolBarButton(const FButtonArgs& Args)
 	NewToolBarButtonBlock->SetOnGetMenuContent(Args.OnGetMenuContent);
 	NewToolBarButtonBlock->SetGetDecoratedButtonDelegate( Args.GetDecoratedButtonDelegate );
 	
+	if (Args.VisibilityOverride.IsSet())
+	{
+		NewToolBarButtonBlock->SetVisibilityOverride(Args.VisibilityOverride);
+	}	
+
 	if (bHasUserInteractionType)
 	{
 		NewToolBarButtonBlock->SetStyleNameOverride(CurrentStyleOverride);
diff --git a/Engine/Source/Runtime/Slate/Private/Framework/MultiBox/SToolBarButtonBlock.cpp b/Engine/Source/Runtime/Slate/Private/Framework/MultiBox/SToolBarButtonBlock.cpp
--- a/Engine/Source/Runtime/Slate/Private/Framework/MultiBox/SToolBarButtonBlock.cpp
+++ b/Engine/Source/Runtime/Slate/Private/Framework/MultiBox/SToolBarButtonBlock.cpp
@@ -28,6 +28,10 @@ FToolBarButtonBlock::FToolBarButtonBlock(FButtonArgs ButtonArgs)
 	, bForceSmallIcons(false)
 	, GetDecoratedButtonDelegate(ButtonArgs.GetDecoratedButtonDelegate)
 {
+	if (ButtonArgs.VisibilityOverride.IsSet())
+	{
+		SetVisibilityOverride(ButtonArgs.VisibilityOverride);
+	}
 }
 
 FToolBarButtonBlock::FToolBarButtonBlock(
diff --git a/Engine/Source/Runtime/Slate/Public/Framework/MultiBox/SToolBarButtonBlock.h b/Engine/Source/Runtime/Slate/Public/Framework/MultiBox/SToolBarButtonBlock.h
--- a/Engine/Source/Runtime/Slate/Public/Framework/MultiBox/SToolBarButtonBlock.h
+++ b/Engine/Source/Runtime/Slate/Public/Framework/MultiBox/SToolBarButtonBlock.h
@@ -39,6 +39,8 @@ struct FButtonArgs : public TSharedFromThis<FButtonArgs>
 	FOnGetContent OnGetMenuContent;
 	FGetDecoratedButtonDelegate GetDecoratedButtonDelegate;
 
+	TAttribute<EVisibility> VisibilityOverride;
+
 	explicit FButtonArgs() {}
 };
 
