Fix Color Picker hitches when editing.
Modify SColor Slider to not throttle the Viewport anymore when created by the color picker.
Modify SColorWheel and SColorSpectrum to have High Precision Input enabled only when requested by the user(Control held down) and not by default.
Modify SColorWheel and SColorSpectrum to gather mouse input and use it on OnFinishedPointerInput(), rather than on each input.

[RN] minor | Rendering Architecture

#Synced-CL 48955713

diff --git a/Engine/Source/Runtime/AppFramework/Private/Widgets/Colors/SColorPicker.cpp b/Engine/Source/Runtime/AppFramework/Private/Widgets/Colors/SColorPicker.cpp
--- a/Engine/Source/Runtime/AppFramework/Private/Widgets/Colors/SColorPicker.cpp
+++ b/Engine/Source/Runtime/AppFramework/Private/Widgets/Colors/SColorPicker.cpp
@@ -168,6 +168,7 @@ void SColorPicker::GenerateDefaultColorPickerContent( bool bAdvancedSectionExpan
 													.OnValueChanged(this, &SColorPicker::HandleColorWheelValueChanged)
 													.OnMouseCaptureBegin(this, &SColorPicker::HandleInteractiveChangeBegin)
 													.OnMouseCaptureEnd(this, &SColorPicker::HandleInteractiveChangeEnd)
+													.PreventThrottling(true)
 											]
 
 										+ SHorizontalBox::Slot()
@@ -200,6 +201,7 @@ void SColorPicker::GenerateDefaultColorPickerContent( bool bAdvancedSectionExpan
 													.OnValueChanged(this, &SColorPicker::HandleColorSpectrumValueChanged)
 													.OnMouseCaptureBegin(this, &SColorPicker::HandleInteractiveChangeBegin)
 													.OnMouseCaptureEnd(this, &SColorPicker::HandleInteractiveChangeEnd)
+													.PreventThrottling(true)
 											]
 									]
 							]
@@ -526,6 +528,7 @@ void SColorPicker::GenerateInlineColorPickerContent()
 					.OnValueChanged(this, &SColorPicker::HandleColorWheelValueChanged)
 					.OnMouseCaptureBegin(this, &SColorPicker::HandleInteractiveChangeBegin)
 					.OnMouseCaptureEnd(this, &SColorPicker::HandleInteractiveChangeEnd)
+					.PreventThrottling(true)
 			]
 
 		+ SHorizontalBox::Slot()
@@ -712,7 +715,8 @@ TSharedRef<SWidget> SColorPicker::MakeColorSlider( EColorPickerChannels Channel
 		.Visibility(this, &SColorPicker::HandleColorPickerModeVisibility, EColorPickerModes::Wheel)
 		.OnBeginSliderMovement(const_cast<SColorPicker*>(this), &SColorPicker::HandleInteractiveChangeBegin)
 		.OnEndSliderMovement(const_cast<SColorPicker*>(this), &SColorPicker::HandleInteractiveChangeEnd)
-		.OnValueChanged(const_cast<SColorPicker*>(this), &SColorPicker::HandleColorSpinBoxValueChanged, Channel);
+		.OnValueChanged(const_cast<SColorPicker*>(this), &SColorPicker::HandleColorSpinBoxValueChanged, Channel)
+		.PreventThrottling(true);
 }
 
 
diff --git a/Engine/Source/Runtime/AppFramework/Private/Widgets/Colors/SColorSlider.cpp b/Engine/Source/Runtime/AppFramework/Private/Widgets/Colors/SColorSlider.cpp
--- a/Engine/Source/Runtime/AppFramework/Private/Widgets/Colors/SColorSlider.cpp
+++ b/Engine/Source/Runtime/AppFramework/Private/Widgets/Colors/SColorSlider.cpp
@@ -43,7 +43,8 @@ void SColorSlider::Construct(const FArguments& InArgs)
 		.Value(InArgs._Value)
 		.OnMouseCaptureBegin(InArgs._OnBeginSliderMovement)
 		.OnMouseCaptureEnd(InArgs._OnEndSliderMovement)
-		.OnValueChanged(InArgs._OnValueChanged);
+		.OnValueChanged(InArgs._OnValueChanged)
+		.PreventThrottling(InArgs._PreventThrottling);
 
 	TSharedPtr<SHorizontalBox> ColorWidget = SNew(SHorizontalBox);
 
diff --git a/Engine/Source/Runtime/AppFramework/Private/Widgets/Colors/SColorSlider.h b/Engine/Source/Runtime/AppFramework/Private/Widgets/Colors/SColorSlider.h
--- a/Engine/Source/Runtime/AppFramework/Private/Widgets/Colors/SColorSlider.h
+++ b/Engine/Source/Runtime/AppFramework/Private/Widgets/Colors/SColorSlider.h
@@ -35,6 +35,7 @@ public:
 		, _GradientColors()
 		, _HasAlphaBackground(false)
 		, _UseSRGB(true)
+		, _PreventThrottling(false)
 		, _OnValueChanged()
 		, _OnBeginSliderMovement()
 		, _OnEndSliderMovement()
@@ -78,6 +79,9 @@ public:
 		/** Whether to display sRGB color */
 		SLATE_ATTRIBUTE(bool, UseSRGB)
 
+		/** Whether to update the viewport with every single change to the value during dragging */
+		SLATE_ARGUMENT(bool, PreventThrottling)
+
 		/** Called when the value is changed by the Slider or SpinBox */
 		SLATE_EVENT(FOnValueChanged, OnValueChanged)
 
diff --git a/Engine/Source/Runtime/Slate/Private/Widgets/Colors/SColorSpectrum.cpp b/Engine/Source/Runtime/Slate/Private/Widgets/Colors/SColorSpectrum.cpp
--- a/Engine/Source/Runtime/Slate/Private/Widgets/Colors/SColorSpectrum.cpp
+++ b/Engine/Source/Runtime/Slate/Private/Widgets/Colors/SColorSpectrum.cpp
@@ -3,6 +3,7 @@
 #include "Widgets/Colors/SColorSpectrum.h"
 #include "Rendering/DrawElements.h"
 #include "Styling/CoreStyle.h"
+#include "Application/SlateApplicationBase.h"
 
 
 SColorSpectrum::SColorSpectrum()
@@ -24,6 +25,7 @@ void SColorSpectrum::Construct( const FArguments& InArgs )
 	OnMouseCaptureBegin = InArgs._OnMouseCaptureBegin;
 	OnMouseCaptureEnd = InArgs._OnMouseCaptureEnd;
 	OnValueChanged = InArgs._OnValueChanged;
+	bPreventThrottling = InArgs._PreventThrottling;
 
 	CtrlMultiplier = InArgs._CtrlMultiplier;
 }
@@ -48,10 +50,23 @@ FReply SColorSpectrum::OnMouseButtonDown( const FGeometry& MyGeometry, const FPo
 {
 	if (MouseEvent.GetEffectingButton() == EKeys::LeftMouseButton)
 	{
+		bHighPrecision = MouseEvent.IsControlDown();
 		ProcessMouseAction(MyGeometry, MouseEvent);
 		OnMouseCaptureBegin.ExecuteIfBound();
 
-		return FReply::Handled().CaptureMouse(SharedThis(this)).UseHighPrecisionMouseMovement(SharedThis(this)).SetUserFocus(SharedThis(this), EFocusCause::Mouse);
+		FReply ReturnReply = FReply::Handled().CaptureMouse(SharedThis(this)).SetUserFocus(SharedThis(this), EFocusCause::Mouse);
+
+		if (bHighPrecision)
+		{
+			ReturnReply = ReturnReply.UseHighPrecisionMouseMovement(SharedThis(this));
+		}
+
+		if (bPreventThrottling)
+		{
+			ReturnReply.PreventThrottling();
+		}
+
+		return ReturnReply;
 	}
 
 	return FReply::Unhandled();
@@ -89,11 +104,27 @@ FReply SColorSpectrum::OnMouseMove( const FGeometry& MyGeometry, const FPointerE
 		LastSpectrumPosition = MyGeometry.AbsoluteToLocal(MouseEvent.GetScreenSpacePosition());
 	}
 
-	ProcessMouseAction(MyGeometry, MouseEvent);
+	bHighPrecision = MouseEvent.IsControlDown();
+	LastGeometry = MyGeometry;
+	LastMouseEvent = MouseEvent;
+
+	if (bHighPrecision && !FSlateApplicationBase::Get().GetPlatformApplication()->IsUsingHighPrecisionMouseMode())
+	{
+		return FReply::Handled().UseHighPrecisionMouseMovement(SharedThis(this));
+	}
 
 	return FReply::Handled();
 }
 
+void SColorSpectrum::OnFinishedPointerInput()
+{
+	if (bDragging && !LastMouseEvent.GetCursorDelta().IsNearlyZero())
+	{
+		ProcessMouseAction(LastGeometry, LastMouseEvent);
+		LastGeometry = FGeometry();
+		LastMouseEvent = FPointerEvent();
+	}
+}
 
 int32 SColorSpectrum::OnPaint( const FPaintArgs& Args, const FGeometry& AllottedGeometry, const FSlateRect& MyCullingRect, FSlateWindowElementList& OutDrawElements, int32 LayerId, const FWidgetStyle& InWidgetStyle, bool bParentEnabled ) const
 {
@@ -153,7 +184,8 @@ void SColorSpectrum::ProcessMouseAction(const FGeometry& MyGeometry, const FPoin
 	{
 		constexpr float SpectrumSensitivity = 0.4f;
 		FVector2f Delta = MouseEvent.GetCursorDelta() * SpectrumSensitivity;
-		if (MouseEvent.IsControlDown())
+
+		if (bHighPrecision)
 		{
 			Delta *= CtrlMultiplier.Get();
 		}
diff --git a/Engine/Source/Runtime/Slate/Private/Widgets/Colors/SColorWheel.cpp b/Engine/Source/Runtime/Slate/Private/Widgets/Colors/SColorWheel.cpp
--- a/Engine/Source/Runtime/Slate/Private/Widgets/Colors/SColorWheel.cpp
+++ b/Engine/Source/Runtime/Slate/Private/Widgets/Colors/SColorWheel.cpp
@@ -2,6 +2,7 @@
 
 #include "Widgets/Colors/SColorWheel.h"
 #include "Rendering/DrawElements.h"
+#include "Application/SlateApplicationBase.h"
 
 SColorWheel::SColorWheel()
 	: SelectedColor(*this, FLinearColor(ForceInit))
@@ -22,6 +23,7 @@ void SColorWheel::Construct(const FArguments& InArgs)
 	OnMouseCaptureBegin = InArgs._OnMouseCaptureBegin;
 	OnMouseCaptureEnd = InArgs._OnMouseCaptureEnd;
 	OnValueChanged = InArgs._OnValueChanged;
+	bPreventThrottling = InArgs._PreventThrottling;
 
 	CtrlMultiplier = InArgs._CtrlMultiplier;
 }
@@ -47,6 +49,7 @@ FReply SColorWheel::OnMouseButtonDown(const FGeometry& MyGeometry, const FPointe
 	if (MouseEvent.GetEffectingButton() == EKeys::LeftMouseButton)
 	{
 		OnMouseCaptureBegin.ExecuteIfBound();
+		bHighPrecision = MouseEvent.IsControlDown();
 
 		if (!ProcessMouseAction(MyGeometry, MouseEvent, false))
 		{
@@ -54,7 +57,19 @@ FReply SColorWheel::OnMouseButtonDown(const FGeometry& MyGeometry, const FPointe
 			return FReply::Unhandled();
 		}
 
-		return FReply::Handled().CaptureMouse(SharedThis(this)).UseHighPrecisionMouseMovement(SharedThis(this)).SetUserFocus(SharedThis(this), EFocusCause::Mouse);
+		FReply ReturnReply = FReply::Handled().CaptureMouse(SharedThis(this)).SetUserFocus(SharedThis(this), EFocusCause::Mouse);
+
+		if (bHighPrecision)
+		{
+			ReturnReply = ReturnReply.UseHighPrecisionMouseMovement(SharedThis(this));
+		}
+
+		if (bPreventThrottling)
+		{
+			ReturnReply.PreventThrottling();
+		}
+
+		return ReturnReply;
 	}
 
 	return FReply::Unhandled();
@@ -66,6 +81,7 @@ FReply SColorWheel::OnMouseButtonUp(const FGeometry& MyGeometry, const FPointerE
 	if (MouseEvent.GetEffectingButton() == EKeys::LeftMouseButton && HasMouseCapture())
 	{
 		bDragging = false;
+		bHighPrecision = false;
 
 		OnMouseCaptureEnd.ExecuteIfBound();
 
@@ -93,11 +109,27 @@ FReply SColorWheel::OnMouseMove(const FGeometry& MyGeometry, const FPointerEvent
 		LastWheelPosition = MyGeometry.AbsoluteToLocal(MouseEvent.GetScreenSpacePosition());
 	}
 
-	ProcessMouseAction(MyGeometry, MouseEvent, true);
+	bHighPrecision = MouseEvent.IsControlDown();
+	LastGeometry = MyGeometry;
+	LastMouseEvent = MouseEvent;
+
+	if (bHighPrecision && !FSlateApplicationBase::Get().GetPlatformApplication()->IsUsingHighPrecisionMouseMode())
+	{
+		return FReply::Handled().UseHighPrecisionMouseMovement(SharedThis(this));
+	}
 
 	return FReply::Handled();
 }
 
+void SColorWheel::OnFinishedPointerInput()
+{
+	if (bDragging && !LastMouseEvent.GetCursorDelta().IsNearlyZero())
+	{
+		ProcessMouseAction(LastGeometry, LastMouseEvent, true);
+		LastGeometry = FGeometry();
+		LastMouseEvent = FPointerEvent();
+	}
+}
 
 int32 SColorWheel::OnPaint(const FPaintArgs& Args, const FGeometry& AllottedGeometry, const FSlateRect& MyCullingRect, FSlateWindowElementList& OutDrawElements, int32 LayerId, const FWidgetStyle& InWidgetStyle, bool bParentEnabled) const
 {
@@ -154,11 +186,13 @@ UE::Slate::FDeprecateVector2DResult SColorWheel::CalcRelativePositionFromCenter(
 bool SColorWheel::ProcessMouseAction(const FGeometry& MyGeometry, const FPointerEvent& MouseEvent, bool bProcessWhenOutsideColorWheel)
 {
 	FVector2f LocalMouseCoordinate;
+
 	if (bDragging)
 	{
 		constexpr float WheelSensitivity = 0.35f;
 		FVector2f Delta = MouseEvent.GetCursorDelta() * WheelSensitivity;
-		if (MouseEvent.IsControlDown())
+
+		if (bHighPrecision)
 		{
 			Delta *= CtrlMultiplier.Get();
 		}
diff --git a/Engine/Source/Runtime/Slate/Public/Widgets/Colors/SColorSpectrum.h b/Engine/Source/Runtime/Slate/Public/Widgets/Colors/SColorSpectrum.h
--- a/Engine/Source/Runtime/Slate/Public/Widgets/Colors/SColorSpectrum.h
+++ b/Engine/Source/Runtime/Slate/Public/Widgets/Colors/SColorSpectrum.h
@@ -27,6 +27,7 @@ public:
 		, _OnMouseCaptureEnd()
 		, _OnValueChanged()
 		, _CtrlMultiplier(0.1f)
+		, _PreventThrottling(false)
 	{ }
 	
 		/** The current color selected by the user */
@@ -44,6 +45,9 @@ public:
 		/** Multiplier to use when ctrl is held down */
 		SLATE_ATTRIBUTE(float, CtrlMultiplier)
 
+		/** If refresh requests for the viewport should happen for all value changes **/
+		SLATE_ARGUMENT(bool, PreventThrottling)
+
 	SLATE_END_ARGS()
 	
 public:
@@ -66,6 +70,7 @@ public:
 	SLATE_API virtual FReply OnMouseButtonDown( const FGeometry& MyGeometry, const FPointerEvent& MouseEvent ) override;
 	SLATE_API virtual FReply OnMouseButtonUp( const FGeometry& MyGeometry, const FPointerEvent& MouseEvent ) override;
 	SLATE_API virtual FReply OnMouseMove( const FGeometry& MyGeometry, const FPointerEvent& MouseEvent ) override;
+	SLATE_API virtual void OnFinishedPointerInput() override;
 	SLATE_API virtual int32 OnPaint( const FPaintArgs& Args, const FGeometry& AllottedGeometry, const FSlateRect& MyCullingRect, FSlateWindowElementList& OutDrawElements, int32 LayerId, const FWidgetStyle& InWidgetStyle, bool bParentEnabled ) const override;
 	SLATE_API virtual FCursorReply OnCursorQuery(const FGeometry& MyGeometry, const FPointerEvent& CursorEvent) const override;
 
@@ -100,9 +105,21 @@ private:
 	// Whether the user is dragging the slider
 	bool bDragging = false;
 
+	/** Whether the user requires high precision mode */
+	bool bHighPrecision = false;
+
 	// Cached mouse position to restore after dragging.
 	FVector2f LastSpectrumPosition;
 
+	/** If refresh requests for the viewport should happen for all value changes **/
+	bool bPreventThrottling = false;
+
+	/** Last Geometry to use after dragging. */
+	FGeometry LastGeometry;
+
+	/** Last mouse event to use after dragging. */
+	FPointerEvent LastMouseEvent;
+
 private:
 
 	// Holds a delegate that is executed when the mouse is pressed and a capture begins.
diff --git a/Engine/Source/Runtime/Slate/Public/Widgets/Colors/SColorWheel.h b/Engine/Source/Runtime/Slate/Public/Widgets/Colors/SColorWheel.h
--- a/Engine/Source/Runtime/Slate/Public/Widgets/Colors/SColorWheel.h
+++ b/Engine/Source/Runtime/Slate/Public/Widgets/Colors/SColorWheel.h
@@ -29,6 +29,7 @@ public:
 		, _OnMouseCaptureEnd()
 		, _OnValueChanged()
 		, _CtrlMultiplier(0.1f)
+		, _PreventThrottling(false)
 	{ }
 	
 		/** The current color selected by the user. */
@@ -49,6 +50,9 @@ public:
 		/** Multiplier to use when ctrl is held down */
 		SLATE_ATTRIBUTE(float, CtrlMultiplier)
 
+		/** If refresh requests for the viewport should happen for all value changes **/
+		SLATE_ARGUMENT(bool, PreventThrottling)
+
 	SLATE_END_ARGS()
 	
 public:
@@ -78,6 +82,7 @@ public:
 	SLATE_API virtual FReply OnMouseButtonDown(const FGeometry& MyGeometry, const FPointerEvent& MouseEvent) override;
 	SLATE_API virtual FReply OnMouseButtonUp(const FGeometry& MyGeometry, const FPointerEvent& MouseEvent) override;
 	SLATE_API virtual FReply OnMouseMove(const FGeometry& MyGeometry, const FPointerEvent& MouseEvent) override;
+	SLATE_API virtual void OnFinishedPointerInput() override;
 	SLATE_API virtual int32 OnPaint(const FPaintArgs& Args, const FGeometry& AllottedGeometry, const FSlateRect& MyCullingRect, FSlateWindowElementList& OutDrawElements, int32 LayerId, const FWidgetStyle& InWidgetStyle, bool bParentEnabled) const override;
 	SLATE_API virtual FCursorReply OnCursorQuery(const FGeometry& MyGeometry, const FPointerEvent& CursorEvent) const override;
 
@@ -97,6 +102,9 @@ protected:
 	 */
 	SLATE_API bool ProcessMouseAction(const FGeometry& MyGeometry, const FPointerEvent& MouseEvent, bool bProcessWhenOutsideColorWheel);
 
+	/** If refresh requests for the viewport should happen for all value changes **/
+	bool bPreventThrottling = false;
+
 private:
 
 	/** The color wheel image to show. */
@@ -117,9 +125,18 @@ private:
 	/** Whether the user is dragging the slider */
 	bool bDragging = false;
 
+	/** Whether the user requires high precision mode */
+	bool bHighPrecision = false;
+
 	/** Cached mouse position to restore after dragging. */
 	FVector2f LastWheelPosition;
 
+	/** Last Geometry to use after dragging. */
+	FGeometry LastGeometry;
+
+	/** Last mouse event to use after dragging. */
+	FPointerEvent LastMouseEvent;
+
 private:
 
 	/** Invoked when the mouse is pressed and a capture begins. */
