Fixed SScrollPanel not properly initialized

Found and fixed by [REDACTED_PERSON], I'm just helping with the commit.

diff --git a/Engine/Source/Runtime/Slate/Public/Widgets/Layout/SScrollBox.h b/Engine/Source/Runtime/Slate/Public/Widgets/Layout/SScrollBox.h
--- a/Engine/Source/Runtime/Slate/Public/Widgets/Layout/SScrollBox.h
+++ b/Engine/Source/Runtime/Slate/Public/Widgets/Layout/SScrollBox.h
@@ -625,6 +625,10 @@ class SScrollPanel : public SPanel
 public:
 
 	SLATE_BEGIN_ARGS(SScrollPanel)
+		: _Orientation(EOrientation::Orient_Horizontal)
+		, _BackPadScrolling(false)
+		, _FrontPadScrolling(false)
+		, _AllowContentToShrink(false)
 	{
 		_Visibility = EVisibility::SelfHitTestInvisible;
 	}
