Provide an optional DefaultState to Notification Info.

#rnx
#Synced-CL 49037932

diff --git a/Engine/Source/Runtime/Slate/Private/Widgets/Notifications/SNotificationList.cpp b/Engine/Source/Runtime/Slate/Private/Widgets/Notifications/SNotificationList.cpp
--- a/Engine/Source/Runtime/Slate/Private/Widgets/Notifications/SNotificationList.cpp
+++ b/Engine/Source/Runtime/Slate/Private/Widgets/Notifications/SNotificationList.cpp
@@ -908,6 +908,11 @@ TSharedRef<SNotificationItem> SNotificationList::AddNotification(const FNotifica
 		{
 			NewItem->ExpireAndFadeout();
 		}
+
+		if (Info.DefaultState != SNotificationItem::ECompletionState::CS_None)
+		{
+			NewItem->SetCompletionState(Info.DefaultState);
+		}
 	}
 	else
 	{
diff --git a/Engine/Source/Runtime/Slate/Public/Widgets/Notifications/SNotificationList.h b/Engine/Source/Runtime/Slate/Public/Widgets/Notifications/SNotificationList.h
--- a/Engine/Source/Runtime/Slate/Public/Widgets/Notifications/SNotificationList.h
+++ b/Engine/Source/Runtime/Slate/Public/Widgets/Notifications/SNotificationList.h
@@ -146,6 +146,7 @@ struct FNotificationInfo
 		bUseCopyToClipboard(false),
 		WidthOverride(320.0f),
 		bFireAndForget(true),
+		DefaultState(SNotificationItem::CS_None),
 		CheckBoxState(ECheckBoxState::Unchecked),
 		CheckBoxStateChanged(),
 		CheckBoxText(),
@@ -174,6 +175,7 @@ struct FNotificationInfo
 		bUseCopyToClipboard(false),
 		WidthOverride(320.0f),
 		bFireAndForget(true),
+		DefaultState(SNotificationItem::CS_None),
 		CheckBoxState(ECheckBoxState::Unchecked),
 		CheckBoxStateChanged(),
 		CheckBoxText(),
@@ -231,6 +233,9 @@ struct FNotificationInfo
 	/** When true the notification will automatically time out after the expire duration. */ 
 	bool bFireAndForget;
 
+	/** Defines what state should this notification have by default */
+	SNotificationItem::ECompletionState DefaultState;
+
 	/** When set this will display a check box on the notification; handles getting the current check box state */
 	TAttribute<ECheckBoxState> CheckBoxState;
 
