[TabManager]
- Add multicast delegate that broadcasts when the F10 hotkey is pressed to sidebar all open tabs
- Add tab exceptions for the same function to disallow Sequencer tab from being docked
#Synced-CL 49128623

diff --git a/Engine/Source/Editor/LevelEditor/Private/SLevelViewport.cpp b/Engine/Source/Editor/LevelEditor/Private/SLevelViewport.cpp
--- a/Engine/Source/Editor/LevelEditor/Private/SLevelViewport.cpp
+++ b/Engine/Source/Editor/LevelEditor/Private/SLevelViewport.cpp
@@ -2690,7 +2690,8 @@ void SLevelViewport::OnToggleImmersive()
 
 void SLevelViewport::OnToggleSidebarTabs()
 {
-	ParentLevelEditor.Pin()->GetTabManager()->ToggleSidebarOpenTabs();
+	static const TSet<FTabId> Exceptions = { FTabId(TEXT("Sequencer")) };
+	ParentLevelEditor.Pin()->GetTabManager()->ToggleSidebarOpenTabs(Exceptions);
 }
 
 bool SLevelViewport::IsImmersive() const
diff --git a/Engine/Source/Runtime/Slate/Private/Framework/Docking/TabManager.cpp b/Engine/Source/Runtime/Slate/Private/Framework/Docking/TabManager.cpp
--- a/Engine/Source/Runtime/Slate/Private/Framework/Docking/TabManager.cpp
+++ b/Engine/Source/Runtime/Slate/Private/Framework/Docking/TabManager.cpp
@@ -2398,9 +2398,35 @@ bool FTabManager::IsTabAllowedInSidebar(const FTabId TabId) const
 	return false;
 }
 
-void FTabManager::ToggleSidebarOpenTabs()
+void FTabManager::ToggleSidebarOpenTabs(const TSet<FTabId>& InExceptions)
 {
-	if(TemporarilySidebaredTabs.Num() == 0)
+	bool bHasSidebarTabs = !TemporarilySidebaredTabs.IsEmpty();
+
+	// Check if TemporarilySidebaredTabs only contains InExceptions
+	if (bHasSidebarTabs)
+	{
+		bool bHasOnlyExceptions = true;
+
+		for (const TWeakPtr<SDockTab>& WeakTab : TemporarilySidebaredTabs)
+		{
+			if (TSharedPtr<SDockTab> Tab = WeakTab.Pin())
+			{
+				const FTabId TabId = Tab->GetLayoutIdentifier();
+				if (!InExceptions.Contains(TabId))
+				{
+					bHasOnlyExceptions = false;
+					break;
+				}
+			}
+		}
+
+		if (bHasOnlyExceptions)
+		{
+			bHasSidebarTabs = false;
+		}
+	}
+
+	if (!bHasSidebarTabs)
 	{
 		// Sidebar opened tabs not in a sidebar already
 		for (int32 AreaIndex = 0; AreaIndex < DockAreas.Num(); ++AreaIndex)
@@ -2412,7 +2438,11 @@ void FTabManager::ToggleSidebarOpenTabs()
 				TArray<TSharedRef<SDockTab>> AllTabs = SomeDockArea->GetAllChildTabs();
 				for (TSharedRef<SDockTab>& Tab : AllTabs)
 				{
-					if (IsTabAllowedInSidebar(Tab->GetLayoutIdentifier()) && !SomeDockArea->IsTabInSidebar(Tab) && Tab->GetParentDockTabStack()->CanMoveTabToSideBar(Tab))
+					const FTabId TabId = Tab->GetLayoutIdentifier();
+					if (!InExceptions.Contains(TabId)
+						&& IsTabAllowedInSidebar(Tab->GetLayoutIdentifier())
+						&& !SomeDockArea->IsTabInSidebar(Tab)
+						&& Tab->GetParentDockTabStack()->CanMoveTabToSideBar(Tab))
 					{
 						Tab->GetParentDockTabStack()->MoveTabToSidebar(Tab);
 						TemporarilySidebaredTabs.Add(Tab);
@@ -2424,16 +2454,25 @@ void FTabManager::ToggleSidebarOpenTabs()
 	}
 	else
 	{
-		for (TWeakPtr<SDockTab>& TabPtr : TemporarilySidebaredTabs)
+		for (auto It = TemporarilySidebaredTabs.CreateIterator(); It; ++It)
 		{
-			if (TSharedPtr<SDockTab> Tab = TabPtr.Pin())
+			if (TSharedPtr<SDockTab> Tab = It->Pin())
 			{
-				Tab->GetParentDockTabStack()->GetDockArea()->RestoreTabFromSidebar(Tab.ToSharedRef());
+				const FTabId TabId = Tab->GetLayoutIdentifier();
+				if (!InExceptions.Contains(TabId))
+				{
+					Tab->GetParentDockTabStack()->GetDockArea()->RestoreTabFromSidebar(Tab.ToSharedRef());
+					It.RemoveCurrent();
+				}
+			}
+			else
+			{
+				It.RemoveCurrent();
 			}
 		}
-
-		TemporarilySidebaredTabs.Empty();
 	}
+
+	OnToggleTemporarilySidebaredTabs.Broadcast(!bHasSidebarTabs, InExceptions);
 }
 
 TSharedPtr<SDockTab> FTabManager::SpawnTab(const FTabId& TabId, const TSharedPtr<SWindow>& ParentWindow, const bool bCanOutputBeNullptr)
diff --git a/Engine/Source/Runtime/Slate/Public/Framework/Docking/TabManager.h b/Engine/Source/Runtime/Slate/Public/Framework/Docking/TabManager.h
--- a/Engine/Source/Runtime/Slate/Public/Framework/Docking/TabManager.h
+++ b/Engine/Source/Runtime/Slate/Public/Framework/Docking/TabManager.h
@@ -182,6 +182,12 @@ DECLARE_DELEGATE_RetVal_OneParam(bool, FCanSpawnTab, const FSpawnTabArgs&);
  */
 DECLARE_DELEGATE_RetVal_OneParam( TSharedPtr<SDockTab>, FOnFindTabToReuse, const FTabId& )
 
+DECLARE_MULTICAST_DELEGATE_TwoParams(
+	FOnToggleTemporarilySidebaredTabs,
+	const bool /*bInSidebared*/,
+	/** Tab exceptions that were requested to not be sidebared */
+	const TSet<FTabId>& /*InExceptions*/)
+
 struct FMinorTabConfig
 {
 public:
@@ -1120,6 +1126,8 @@ class FTabManager : public TSharedFromThis<FTabManager>
 	
 		FOnReadOnlyModeChanged& GetOnReadOnlyModeChangedDelegate() { return OnReadOnlyModeChanged; }
 
+		FOnToggleTemporarilySidebaredTabs& GetOnToggleTemporarilySidebaredTabsDelegate() { return OnToggleTemporarilySidebaredTabs; }
+
 		/* Prevent or allow all tabs to be drag */
 		void SetCanDoDragOperation(bool CanDoDragOperation) { bCanDoDragOperation = CanDoDragOperation; }
 		
@@ -1138,7 +1146,7 @@ class FTabManager : public TSharedFromThis<FTabManager>
 		/**
 		 * Temporarily moves all open tabs in this tab manager to a sidebar or restores them from a temporary state
 		 */
-		SLATE_API void ToggleSidebarOpenTabs();
+		SLATE_API void ToggleSidebarOpenTabs(const TSet<FTabId>& InExceptions = {});
 
 		/** @return The local workspace menu root */
 		SLATE_API const TSharedRef<FWorkspaceItem> GetLocalWorkspaceMenuRoot() const;
@@ -1466,6 +1474,9 @@ PRAGMA_ENABLE_EXPERIMENTAL_WARNINGS
 
 		/** Keep alive the PanelDrawer tabs that have no area */
 		TMap<FTabId, TSharedRef<SDockTab>> PanelDrawerActiveHiddenTabs;
+
+		/** Delegate that broadcasts when the tab manager toggles between temporarily sidebared tabs */
+		FOnToggleTemporarilySidebaredTabs OnToggleTemporarilySidebaredTabs;
 };
 
 
