Enabling Slate Dynamic Invalidation by default in the editor

#Synced-CL 48711348

diff --git a/Engine/Source/Runtime/Slate/Private/Widgets/SInvalidationPanel.cpp b/Engine/Source/Runtime/Slate/Private/Widgets/SInvalidationPanel.cpp
--- a/Engine/Source/Runtime/Slate/Private/Widgets/SInvalidationPanel.cpp
+++ b/Engine/Source/Runtime/Slate/Private/Widgets/SInvalidationPanel.cpp
@@ -40,7 +40,7 @@ namespace DynamicInvalidationCVars
 {
 	void ConsoleVariableEnableDynamicInvalidationChanged(IConsoleVariable*);
 
-	static int32 DynamicInvalidationOptions = 0;
+	static int32 DynamicInvalidationOptions = 2;
 	static FAutoConsoleVariableRef CVarDynamicInvalidationOptions(
 		TEXT("Slate.DynamicInvalidation.Options"),
 		DynamicInvalidationOptions,
