Turning Slate Dynamic Invalidation back off until the CB crash is fixed.

#rnx
#Synced-CL 49280314

diff --git a/Engine/Source/Runtime/Slate/Private/Widgets/SInvalidationPanel.cpp b/Engine/Source/Runtime/Slate/Private/Widgets/SInvalidationPanel.cpp
--- a/Engine/Source/Runtime/Slate/Private/Widgets/SInvalidationPanel.cpp
+++ b/Engine/Source/Runtime/Slate/Private/Widgets/SInvalidationPanel.cpp
@@ -40,7 +40,7 @@ namespace DynamicInvalidationCVars
 {
 	void ConsoleVariableEnableDynamicInvalidationChanged(IConsoleVariable*);
 
-	static int32 DynamicInvalidationOptions = 2;
+	static int32 DynamicInvalidationOptions = 0;
 	static FAutoConsoleVariableRef CVarDynamicInvalidationOptions(
 		TEXT("Slate.DynamicInvalidation.Options"),
 		DynamicInvalidationOptions,
