EditorUsability : Tabs
- Extend Browse to Asset capability for Editor Utility Widget
- Allow context menu actions without activating the tab (web browser style)
- Flash tab to attract user attention when switching or opening context menu

#Synced-CL 49388435

diff --git a/Engine/Source/Editor/Blutility/Private/EditorUtilitySubsystem.cpp b/Engine/Source/Editor/Blutility/Private/EditorUtilitySubsystem.cpp
--- a/Engine/Source/Editor/Blutility/Private/EditorUtilitySubsystem.cpp
+++ b/Engine/Source/Editor/Blutility/Private/EditorUtilitySubsystem.cpp
@@ -2,6 +2,7 @@
 
 #include "EditorUtilitySubsystem.h"
 
+#include "ContentBrowserModule.h"
 #include "AssetRegistry/ARFilter.h"
 #include "AssetRegistry/AssetData.h"
 #include "AssetRegistry/IAssetRegistry.h"
@@ -15,11 +16,13 @@
 #include "Engine/BlueprintCore.h"
 #include "Framework/Application/SlateApplication.h"
 #include "Framework/Docking/TabManager.h"
+#include "Framework/MultiBox/MultiBoxBuilder.h"
 #include "GameFramework/Actor.h"
 #include "HAL/IConsoleManager.h"
 #include "HAL/Platform.h"
 #include "HAL/PlatformCrt.h"
 #include "IBlutilityModule.h"
+#include "IContentBrowserSingleton.h"
 #include "Interfaces/IMainFrameModule.h"
 #include "Internationalization/Internationalization.h"
 #include "Internationalization/Text.h"
@@ -466,6 +469,44 @@ UEditorUtilityWidget* UEditorUtilitySubsystem::SpawnAndRegisterTabWithIdGenerate
 	return SpawnAndRegisterTabAndGetIDGeneratedClass(InGeneratedWidgetBlueprint, InTabID);
 }
 
+void UEditorUtilitySubsystem::ExtendTabContextMenu(FMenuBuilder& InMenuBuilder, FName InTabId)
+{
+	InMenuBuilder.BeginSection("DockingTabAssetOptions", LOCTEXT("AssetTabSection", "Assets"));
+	{
+		InMenuBuilder.AddMenuEntry(
+			LOCTEXT("BrowseAssetInContentBrowser", "Browse to Asset"),
+			LOCTEXT("BrowseAssetInContentBrowserTooltip", "Browse to asset in content browser"),
+			FSlateIcon(FAppStyle::GetAppStyleSetName(), "Icons.BrowseContent"),
+			FUIAction(
+				FExecuteAction::CreateUObject(this, &UEditorUtilitySubsystem::OnBrowseAssetInContentBrowser, InTabId),
+				FCanExecuteAction::CreateUObject(this, &UEditorUtilitySubsystem::OnCanBrowseAssetInContentBrowser, InTabId)
+			)
+		);
+	}
+	InMenuBuilder.EndSection();
+}
+
+void UEditorUtilitySubsystem::OnBrowseAssetInContentBrowser(FName InTabId)
+{
+	UEditorUtilityWidgetBlueprint** UtilityWidget = RegisteredTabs.Find(InTabId);
+	if (UtilityWidget == nullptr || (*UtilityWidget) == nullptr)
+	{
+		return;
+	}
+
+	if (const FContentBrowserModule* ContentBrowserModule = FModuleManager::GetModulePtr<FContentBrowserModule>("ContentBrowser"))
+	{
+		const TArray<UObject*> Assets {*UtilityWidget};
+		ContentBrowserModule->Get().SyncBrowserToAssets(Assets);
+	}
+}
+
+bool UEditorUtilitySubsystem::OnCanBrowseAssetInContentBrowser(FName InTabId)
+{
+	UEditorUtilityWidgetBlueprint** UtilityWidget = RegisteredTabs.Find(InTabId);
+	return UtilityWidget != nullptr && (*UtilityWidget) != nullptr;
+}
+
 void UEditorUtilitySubsystem::RegisterTabAndGetIDGeneratedClass(UWidgetBlueprintGeneratedClass* InGeneratedWidgetBlueprint, FName& NewTabID)
 {
 	if (InGeneratedWidgetBlueprint && !IsRunningCommandlet())
@@ -517,6 +558,7 @@ bool UEditorUtilitySubsystem::SpawnRegisteredTabByID(FName NewTabID)
 				if (WidgetToSpawn)
 				{
 					check(*WidgetToSpawn);
+					NewDockTab->SetOnExtendContextMenu(SDockTab::FExtendContextMenu::CreateUObject(this, &UEditorUtilitySubsystem::ExtendTabContextMenu, NewTabID));
 					BlutilityModule->AddLoadedScriptUI(*WidgetToSpawn);
 					return true;
 				}
diff --git a/Engine/Source/Editor/Blutility/Public/EditorUtilitySubsystem.h b/Engine/Source/Editor/Blutility/Public/EditorUtilitySubsystem.h
--- a/Engine/Source/Editor/Blutility/Public/EditorUtilitySubsystem.h
+++ b/Engine/Source/Editor/Blutility/Public/EditorUtilitySubsystem.h
@@ -23,6 +23,7 @@
 
 #define UE_API BLUTILITY_API
 
+class FMenuBuilder;
 class FOutputDevice;
 class FSpawnTabArgs;
 class FSubsystemCollectionBase;
@@ -185,9 +186,12 @@ protected:
 
 	TMap<TSharedRef<SDockTab>, UEditorUtilityWidget*> SpawnedFromGeneratedClassTabs;
 
-
 private:
 	FName GenerateNewTabID(const UObject* WidgetObject, const FName& OptionalPostfix);
+	
+	void ExtendTabContextMenu(FMenuBuilder& InMenuBuilder, FName InTabId);
+	void OnBrowseAssetInContentBrowser(FName InTabId);
+	bool OnCanBrowseAssetInContentBrowser(FName InTabId);
 
 	IConsoleObject* RunTaskCommandObject = nullptr;
 	IConsoleObject* CancelAllTasksCommandObject = nullptr;
diff --git a/Engine/Source/Editor/UnrealEd/Private/Toolkits/SGlobalTabSwitchingDialog.cpp b/Engine/Source/Editor/UnrealEd/Private/Toolkits/SGlobalTabSwitchingDialog.cpp
--- a/Engine/Source/Editor/UnrealEd/Private/Toolkits/SGlobalTabSwitchingDialog.cpp
+++ b/Engine/Source/Editor/UnrealEd/Private/Toolkits/SGlobalTabSwitchingDialog.cpp
@@ -151,6 +151,7 @@ public:
 			{
 				// Don't use DrawAttention since it steals focus
 				DockTab->ActivateInParent(ETabActivationCause::SetDirectly);
+				DockTab->FlashTab();
 			}
 		}
 	}
@@ -205,6 +206,7 @@ public:
 			{
 				// Don't use DrawAttention since it steals focus
 				DockTab->ActivateInParent(ETabActivationCause::SetDirectly);
+				DockTab->FlashTab();
 			}
 		}
 	}
@@ -286,6 +288,7 @@ public:
 		{
 			// Don't use DrawAttention since it steals focus
 			DockTab->ActivateInParent(ETabActivationCause::SetDirectly);
+			DockTab->FlashTab();
 		}
 	}
 
diff --git a/Engine/Source/Runtime/Slate/Private/Framework/Docking/SDockingTabStack.cpp b/Engine/Source/Runtime/Slate/Private/Framework/Docking/SDockingTabStack.cpp
--- a/Engine/Source/Runtime/Slate/Private/Framework/Docking/SDockingTabStack.cpp
+++ b/Engine/Source/Runtime/Slate/Private/Framework/Docking/SDockingTabStack.cpp
@@ -554,6 +554,19 @@ FReply SDockingTabStack::OnKeyDown(const FGeometry& MyGeometry, const FKeyEvent&
 	return SDockingNode::OnKeyDown(MyGeometry, InKeyEvent);
 }
 
+void SDockingTabStack::ShowContextMenu(const TSharedRef<SDockTab>& InContextTab, const FWidgetPath& InWidgetPath)
+{
+	FSlateApplication::Get().PushMenu(
+		AsShared(), 
+		InWidgetPath, 
+		MakeContextMenu(InContextTab), 
+		FSlateApplication::Get().GetCursorPos(), 
+		FPopupTransitionEffect::ContextMenu
+	);
+	
+	InContextTab->FlashTab();
+}
+
 FReply SDockingTabStack::OnUserAttemptingDock( SDockingNode::RelativeDirection Direction, const FDragDropEvent& DragDropEvent )
 {
 	TSharedPtr<FDockingDragOperation> DragDropOperation = DragDropEvent.GetOperationAs<FDockingDragOperation>();
@@ -579,76 +592,98 @@ TArray< TSharedRef<SDockTab> > SDockingTabStack::GetAllChildTabs() const
 
 void SDockingTabStack::CloseForegroundTab()
 {
-	TSharedPtr<SDockTab> ForegroundTab = TabWell->GetForegroundTab();
-	if (ForegroundTab.IsValid())
+	CloseTab(TabWell->GetForegroundTab());
+}
+
+void SDockingTabStack::CloseTab(TSharedPtr<SDockTab> InTab)
+{
+	if (InTab.IsValid())
 	{
-		ForegroundTab->RequestCloseTab();
+		InTab->RequestCloseTab();
 	}
 }
 
 void SDockingTabStack::CloseTabsInDirectionFromForegroundTab(ETabsToClose TabsToClose, ECloseTabsDirection Direction)
 {
-	TSharedPtr<SDockTab> ForegroundTab = TabWell->GetForegroundTab();
-	if (ForegroundTab.IsValid())
+	CloseTabsInDirectionFromTab(TabWell->GetForegroundTab(), TabsToClose, Direction);
+}
+
+void SDockingTabStack::CloseTabsInDirectionFromTab(TSharedPtr<SDockTab> InTab, ETabsToClose InTabsToClose, ECloseTabsDirection InDirection)
+{
+	if (!InTab.IsValid())
+	{
+		return;
+	}
+
+	const int32 IndexOfTab = TabWell->GetTabIndex(InTab);
+	if (IndexOfTab == INDEX_NONE)
+	{
+		return;
+	}
+
+	const int32 Increment = GetIncrementFromCloseTabsDirection(InDirection);
+
+	int32 DestroyIndex = IndexOfTab + Increment;
+	while (DestroyIndex >= 0 && DestroyIndex < TabWell->GetNumTabs())
 	{
-		const int32 ForegroundTabIndex = TabWell->GetForegroundTabIndex();
-		const int32 Increment = GetIncrementFromCloseTabsDirection(Direction);
+		const TSharedRef<SDockTab>& Tab = TabWell->GetTabs()[DestroyIndex];
+		const ETabRole VisualTabRole = Tab->GetVisualTabRole();
+		const bool bCanClose =
+			(InTabsToClose == ETabsToClose::CloseAllTabs) ||
+			(InTabsToClose == ETabsToClose::CloseDocumentTabs && VisualTabRole == ETabRole::DocumentTab) ||
+			(InTabsToClose == ETabsToClose::CloseDocumentAndMajorTabs && (VisualTabRole == ETabRole::DocumentTab || VisualTabRole == ETabRole::MajorTab));
 
-		int32 DestroyIndex = ForegroundTabIndex + Increment;
-		while (DestroyIndex >= 0 && DestroyIndex < TabWell->GetNumTabs())
+		// Need to increment index if we failed to close a tab to the right or if we're closing tabs to the left regardless of if we succeeded
+		if (!bCanClose || !Tab->RequestCloseTab() || InDirection == ECloseTabsDirection::Left)
 		{
-			const TSharedRef<SDockTab>& Tab = TabWell->GetTabs()[DestroyIndex];
-			const ETabRole VisualTabRole = Tab->GetVisualTabRole();
-			const bool bCanClose =
-				(TabsToClose == ETabsToClose::CloseAllTabs) ||
-				(TabsToClose == ETabsToClose::CloseDocumentTabs && VisualTabRole == ETabRole::DocumentTab) ||
-				(TabsToClose == ETabsToClose::CloseDocumentAndMajorTabs && (VisualTabRole == ETabRole::DocumentTab || VisualTabRole == ETabRole::MajorTab));
-
-			// Need to increment index if we failed to close a tab to the right or if we're closing tabs to the left regardless of if we succeeded
-			if (!bCanClose || !Tab->RequestCloseTab() || Direction == ECloseTabsDirection::Left)
-			{
-				DestroyIndex += Increment;
-			}
+			DestroyIndex += Increment;
 		}
 	}
 }
 
 void SDockingTabStack::CloseAllButForegroundTab(ETabsToClose TabsToClose)
 {
-	TSharedPtr<SDockTab> ForegroundTab = TabWell->GetForegroundTab();
-	if (ForegroundTab.IsValid())
+	CloseAllButTab(TabWell->GetForegroundTab(), TabsToClose);
+}
+
+void SDockingTabStack::CloseAllButTab(TSharedPtr<SDockTab> InTab, ETabsToClose InTabsToClose)
+{
+	if (!InTab.IsValid())
 	{
-		int32 DestroyIndex = 0;
-		while ((TabWell->GetNumTabs() > 1) && (DestroyIndex < TabWell->GetNumTabs()))
-		{
-			const TSharedRef<SDockTab>& Tab = TabWell->GetTabs()[DestroyIndex];
+		return;
+	}
+	
+	int32 DestroyIndex = 0;
+	while ((TabWell->GetNumTabs() > 1) && (DestroyIndex < TabWell->GetNumTabs()))
+	{
+		const TSharedRef<SDockTab>& Tab = TabWell->GetTabs()[DestroyIndex];
 
-			const ETabRole VisualTabRole = Tab->GetVisualTabRole();
-			const bool bCanClose = 
-				(TabsToClose == ETabsToClose::CloseAllTabs) ||
-				(TabsToClose == ETabsToClose::CloseDocumentTabs && VisualTabRole == ETabRole::DocumentTab) ||
-				(TabsToClose == ETabsToClose::CloseDocumentAndMajorTabs && (VisualTabRole == ETabRole::DocumentTab || VisualTabRole == ETabRole::MajorTab));
+		const ETabRole VisualTabRole = Tab->GetVisualTabRole();
+		const bool bCanClose = 
+			(InTabsToClose == ETabsToClose::CloseAllTabs) ||
+			(InTabsToClose == ETabsToClose::CloseDocumentTabs && VisualTabRole == ETabRole::DocumentTab) ||
+			(InTabsToClose == ETabsToClose::CloseDocumentAndMajorTabs && (VisualTabRole == ETabRole::DocumentTab || VisualTabRole == ETabRole::MajorTab));
 
-			if ((Tab == ForegroundTab) || !bCanClose || !Tab->RequestCloseTab())
-			{
-				++DestroyIndex;
-			}
+		if ((Tab == InTab) || !bCanClose || !Tab->RequestCloseTab())
+		{
+			++DestroyIndex;
 		}
 	}
 }
 
 FReply SDockingTabStack::TabWellRightClicked( const FGeometry& TabWellGeometry, const FPointerEvent& MouseEvent )
 {
-	if (MouseEvent.GetEffectingButton() == EKeys::RightMouseButton)
-	{
-		FWidgetPath WidgetPath = MouseEvent.GetEventPath() != nullptr ? *MouseEvent.GetEventPath() : FWidgetPath();
-		FSlateApplication::Get().PushMenu(AsShared(), WidgetPath, MakeContextMenu(), FSlateApplication::Get().GetCursorPos(), FPopupTransitionEffect::ContextMenu);
-		return FReply::Handled();
-	}
-	else
+	if (TabWell.IsValid() && MouseEvent.GetEffectingButton() == EKeys::RightMouseButton)
 	{
-		return FReply::Unhandled();
+		if (const TSharedPtr<SDockTab> ForegroundTab = TabWell->GetForegroundTab())
+		{
+			const FWidgetPath WidgetPath = MouseEvent.GetEventPath() != nullptr ? *MouseEvent.GetEventPath() : FWidgetPath();
+			ShowContextMenu(ForegroundTab.ToSharedRef(), WidgetPath);
+			return FReply::Handled();
+		}
 	}
+
+	return FReply::Unhandled();
 }
 
 SDockingNode::ECleanupRetVal SDockingTabStack::CleanUpNodes()
@@ -668,14 +703,19 @@ SDockingNode::ECleanupRetVal SDockingTabStack::CleanUpNodes()
 	}
 }
 
-TSharedRef<SWidget> SDockingTabStack::MakeContextMenu()
+TSharedRef<SWidget> SDockingTabStack::MakeContextMenu(const TSharedPtr<SDockTab>& InTab)
 {
+	if (!InTab.IsValid())
+	{
+		return SNullWidget::NullWidget;
+	}
+
 	// Show a menu that allows users to toggle whether
 	// a specific tab should hide if it is the sole tab
 	// in its tab well.
-	const bool bCloseAfterSelection = true;
-	const bool bCloseSelfOnly = false;
-	FMenuBuilder MenuBuilder( bCloseAfterSelection, NULL, TSharedPtr<FExtender>(), bCloseSelfOnly, &FCoreStyle::Get() );
+	constexpr bool bCloseAfterSelection = true;
+	constexpr bool bCloseSelfOnly = false;
+	FMenuBuilder MenuBuilder( bCloseAfterSelection, nullptr, TSharedPtr<FExtender>(), bCloseSelfOnly, &FCoreStyle::Get() );
 	{
 		MenuBuilder.BeginSection("DockingTabStackOptions", LOCTEXT("TabOptionsHeading", "Options") );
 		{
@@ -689,15 +729,15 @@ TSharedRef<SWidget> SDockingTabStack::MakeContextMenu()
 				) 
 			);
 
-			if(IsTabAllowedInSidebar(TabWell->GetForegroundTab()))
+			if (IsTabAllowedInSidebar(InTab))
 			{
 				MenuBuilder.AddMenuEntry(
 					LOCTEXT("MoveToSidebar", "Dock to Sidebar"),
 					LOCTEXT("MoveToSidebarTooltip", "Moves this tab to a sidebar drawer on the side of the window closest to the tab.\nThe tab can be opened from the drawer and will automatically close again when clicking off it."),
 					FSlateIcon(),
 					FUIAction(
-						FExecuteAction::CreateSP(this, &SDockingTabStack::MoveForegroundTabToSidebar),
-						FCanExecuteAction::CreateSP(this, &SDockingTabStack::CanMoveForegroundTabToSidebar)
+						FExecuteAction::CreateSP(this, &SDockingTabStack::MoveTabToSidebar, InTab),
+						FCanExecuteAction::CreateSP(this, &SDockingTabStack::CanMoveTabToSidebar, InTab)
 					)
 				);
 			}
@@ -708,11 +748,11 @@ TSharedRef<SWidget> SDockingTabStack::MakeContextMenu()
 		{
 			MenuBuilder.AddMenuEntry(
 				LOCTEXT("CloseTab", "Close"),
-				LOCTEXT("CloseTabTooltil", "Close this tab."),
+				LOCTEXT("CloseTabTooltip", "Close this tab."),
 				FSlateIcon(),
 				FUIAction(
-					FExecuteAction::CreateSP( this, &SDockingTabStack::CloseForegroundTab ),
-					FCanExecuteAction::CreateSP(this, &SDockingTabStack::CanCloseForegroundTab)
+					FExecuteAction::CreateSP(this, &SDockingTabStack::CloseTab, InTab),
+					FCanExecuteAction::CreateSP(this, &SDockingTabStack::CanCloseTab, InTab)
 				)
 			);
 
@@ -723,8 +763,8 @@ TSharedRef<SWidget> SDockingTabStack::MakeContextMenu()
 				LOCTEXT("CloseTabsToTheLeftTooltip", "Closes all tabs to the left of the active tab."),
 				FSlateIcon(),
 				FUIAction(
-					FExecuteAction::CreateSP( this, &SDockingTabStack::CloseTabsInDirectionFromForegroundTab, TabsToClose, ECloseTabsDirection::Left),
-					FCanExecuteAction::CreateSP( this, &SDockingTabStack::CanCloseTabsInDirectionOfForegroundTab, ECloseTabsDirection::Left)
+					FExecuteAction::CreateSP( this, &SDockingTabStack::CloseTabsInDirectionFromTab, InTab, TabsToClose, ECloseTabsDirection::Left),
+					FCanExecuteAction::CreateSP( this, &SDockingTabStack::CanCloseTabsInDirectionOfTab, InTab, ECloseTabsDirection::Left)
 				)
 			);
 
@@ -733,28 +773,25 @@ TSharedRef<SWidget> SDockingTabStack::MakeContextMenu()
 				LOCTEXT("CloseTabsToTheRightTooltip", "Closes all tabs to the right of the active tab."),
 				FSlateIcon(),
 				FUIAction(
-					FExecuteAction::CreateSP( this, &SDockingTabStack::CloseTabsInDirectionFromForegroundTab, TabsToClose, ECloseTabsDirection::Right),
-					FCanExecuteAction::CreateSP( this, &SDockingTabStack::CanCloseTabsInDirectionOfForegroundTab, ECloseTabsDirection::Right)
+					FExecuteAction::CreateSP( this, &SDockingTabStack::CloseTabsInDirectionFromTab, InTab, TabsToClose, ECloseTabsDirection::Right),
+					FCanExecuteAction::CreateSP( this, &SDockingTabStack::CanCloseTabsInDirectionOfTab, InTab, ECloseTabsDirection::Right)
 				)
 			);
 
 			MenuBuilder.AddMenuEntry(
 				LOCTEXT("CloseOtherTabs", "Close Other Tabs"),
-				LOCTEXT("CloseOtherTabsTooltil", "Closes all tabs except for the active tab."),
+				LOCTEXT("CloseOtherTabsTooltip", "Closes all tabs except for the active tab."),
 				FSlateIcon(),
 				FUIAction(
-					FExecuteAction::CreateSP( this, &SDockingTabStack::CloseAllButForegroundTab, TabsToClose ),
-					FCanExecuteAction::CreateSP(this, &SDockingTabStack::CanCloseAllButForegroundTab)
+					FExecuteAction::CreateSP( this, &SDockingTabStack::CloseAllButTab, InTab, TabsToClose ),
+					FCanExecuteAction::CreateSP(this, &SDockingTabStack::CanCloseAllButTab, InTab)
 				)
 			);
 		}
 		MenuBuilder.EndSection();
 
-		TSharedPtr<SDockTab> ForegroundTab = TabWell->GetForegroundTab();
-		if (ForegroundTab.IsValid())
-		{
-			ForegroundTab->ExtendContextMenu(MenuBuilder);
-		}
+		// Let the tab extend its own context menu
+		InTab->ExtendContextMenu(MenuBuilder);
 	}
 
 	return MenuBuilder.MakeWidget();
@@ -1062,6 +1099,14 @@ void SDockingTabStack::MoveForegroundTabToSidebar()
 	}
 }
 
+void SDockingTabStack::MoveTabToSidebar(TSharedPtr<SDockTab> InTab)
+{
+	if (InTab.IsValid())
+	{
+		MoveTabToSidebar(InTab.ToSharedRef());
+	}
+}
+
 void SDockingTabStack::MoveTabToSidebar(TSharedRef<SDockTab> Tab)
 {
 	const int32 TabIndex = Tabs.IndexOfByPredicate(FTabMatcher(Tab->GetLayoutIdentifier(), ETabState::OpenedTab));
@@ -1122,8 +1167,12 @@ bool SDockingTabStack::CanHideTabWell() const
 
 bool SDockingTabStack::CanCloseForegroundTab() const
 {
-	TSharedPtr<SDockTab> ForegroundTabPtr = TabWell->GetForegroundTab();
-	return ForegroundTabPtr.IsValid() && ForegroundTabPtr->CanCloseTab();
+	return CanCloseTab(TabWell->GetForegroundTab());
+}
+
+bool SDockingTabStack::CanCloseTab(TSharedPtr<SDockTab> InTab) const
+{
+	return InTab.IsValid() && InTab->CanCloseTab();
 }
 
 int32 SDockingTabStack::GetIncrementFromCloseTabsDirection(ECloseTabsDirection Direction)
@@ -1138,20 +1187,30 @@ int32 SDockingTabStack::GetIncrementFromCloseTabsDirection(ECloseTabsDirection D
 
 bool SDockingTabStack::CanCloseTabsInDirectionOfForegroundTab(ECloseTabsDirection Direction) const
 {
-	TSharedPtr<SDockTab> ForegroundTabPtr = TabWell->GetForegroundTab();
-	if (!ForegroundTabPtr.IsValid())
+	return CanCloseTabsInDirectionOfTab(TabWell->GetForegroundTab(), Direction);
+}
+
+bool SDockingTabStack::CanCloseTabsInDirectionOfTab(TSharedPtr<SDockTab> InTab, ECloseTabsDirection InDirection) const
+{
+	if (!InTab.IsValid())
+	{
+		return false;
+	}
+
+	const int32 TabIndex = TabWell->GetTabIndex(InTab);
+	if (TabIndex == INDEX_NONE)
 	{
 		return false;
 	}
 
-	const ETabRole VisualTabRole = ForegroundTabPtr->GetVisualTabRole();
+	const ETabRole VisualTabRole = InTab->GetVisualTabRole();
 	if (VisualTabRole == ETabRole::DocumentTab || VisualTabRole == ETabRole::MajorTab)
 	{
 		const TArray< TSharedRef<SDockTab> > MyTabs = this->GetTabs().AsArrayCopy();
-		const int32 Increment = GetIncrementFromCloseTabsDirection(Direction);
-		for (int32 TabIndex = TabWell->GetForegroundTabIndex() + Increment; TabIndex >= 0 && TabIndex < MyTabs.Num(); TabIndex += Increment)
+		const int32 Increment = GetIncrementFromCloseTabsDirection(InDirection);
+		for (int32 Index = TabIndex + Increment; Index >= 0 && Index < MyTabs.Num(); Index += Increment)
 		{
-			TSharedRef<SDockTab> Tab = MyTabs[TabIndex];
+			TSharedRef<SDockTab> Tab = MyTabs[Index];
 			if (Tab->CanCloseTab())
 			{
 				return true;
@@ -1164,25 +1223,30 @@ bool SDockingTabStack::CanCloseTabsInDirectionOfForegroundTab(ECloseTabsDirectio
 bool SDockingTabStack::CanCloseAllButForegroundTab() const
 {
 	// If the active tab is a document tab or major tab and there is at least 1 other closeable tab, offer to close the others
-	TSharedPtr<SDockTab> ForegroundTabPtr = TabWell->GetForegroundTab();
-	if (!ForegroundTabPtr.IsValid())
+	return CanCloseAllButTab(TabWell->GetForegroundTab());
+}
+
+bool SDockingTabStack::CanCloseAllButTab(TSharedPtr<SDockTab> InTab) const
+{
+	if (!InTab.IsValid())
 	{
 		return false;
 	}
-
-	const ETabRole VisualTabRole = ForegroundTabPtr->GetVisualTabRole();
+	
+	const ETabRole VisualTabRole = InTab->GetVisualTabRole();
 	if ((VisualTabRole == ETabRole::DocumentTab || VisualTabRole == ETabRole::MajorTab) && (TabWell->GetNumTabs() > 1))
 	{
 		const TArray< TSharedRef<SDockTab> > MyTabs = this->GetTabs().AsArrayCopy();
 		for (int32 TabIndex = 0; TabIndex < MyTabs.Num(); ++TabIndex)
 		{
 			TSharedRef<SDockTab> Tab = MyTabs[TabIndex];
-			if (Tab != ForegroundTabPtr && Tab->CanCloseTab())
+			if (Tab != InTab && Tab->CanCloseTab())
 			{
 				return true;
 			}
 		}
 	}
+
 	return false;
 }
 
@@ -1196,6 +1260,16 @@ bool SDockingTabStack::CanMoveForegroundTabToSidebar() const
 	return false;
 }
 
+bool SDockingTabStack::CanMoveTabToSidebar(TSharedPtr<SDockTab> InTab) const
+{
+	if (InTab.IsValid())
+	{
+		return CanMoveTabToSideBar(InTab.ToSharedRef());
+	}
+
+	return false;
+}
+
 bool SDockingTabStack::CanMoveTabToSideBar(TSharedRef<SDockTab> Tab) const
 {
 	const FTabId TabIdBeingClosed = Tab->GetLayoutIdentifier();
diff --git a/Engine/Source/Runtime/Slate/Private/Framework/Docking/SDockingTabWell.cpp b/Engine/Source/Runtime/Slate/Private/Framework/Docking/SDockingTabWell.cpp
--- a/Engine/Source/Runtime/Slate/Private/Framework/Docking/SDockingTabWell.cpp
+++ b/Engine/Source/Runtime/Slate/Private/Framework/Docking/SDockingTabWell.cpp
@@ -42,6 +42,16 @@ const TSlotlessChildren<SDockTab>& SDockingTabWell::GetTabs() const
 	return Tabs;
 }
 
+int32 SDockingTabWell::GetTabIndex(const TSharedPtr<SDockTab>& InTab) const
+{
+	if (!InTab.IsValid())
+	{
+		return INDEX_NONE;
+	}
+
+	return Tabs.Find(InTab.ToSharedRef());
+}
+
 int32 SDockingTabWell::GetNumTabs() const
 {
 	return Tabs.Num();
diff --git a/Engine/Source/Runtime/Slate/Private/Framework/Docking/SDockingTabWell.h b/Engine/Source/Runtime/Slate/Private/Framework/Docking/SDockingTabWell.h
--- a/Engine/Source/Runtime/Slate/Private/Framework/Docking/SDockingTabWell.h
+++ b/Engine/Source/Runtime/Slate/Private/Framework/Docking/SDockingTabWell.h
@@ -47,6 +47,9 @@ public:
 
 	/** @return All child tabs in this node */
 	const TSlotlessChildren<SDockTab>& GetTabs() const;
+	
+	/** @return Index of the tab in this dock or INDEX_NONE */
+	int32 GetTabIndex(const TSharedPtr<SDockTab>& InTab) const;
 
 	/**
 	 * Add a new tab (InTab) to the TabWell at location (AtIndex).
diff --git a/Engine/Source/Runtime/Slate/Private/Widgets/Docking/SDockTab.cpp b/Engine/Source/Runtime/Slate/Private/Widgets/Docking/SDockTab.cpp
--- a/Engine/Source/Runtime/Slate/Private/Widgets/Docking/SDockTab.cpp
+++ b/Engine/Source/Runtime/Slate/Private/Widgets/Docking/SDockTab.cpp
@@ -51,9 +51,15 @@ FReply SDockTab::OnMouseButtonDown( const FGeometry& MyGeometry, const FPointerE
 		}
 		else if ( MouseEvent.GetEffectingButton() == EKeys::RightMouseButton )
 		{
-			// We clicked on the tab, so it should be active
-			ActivateInParent(ETabActivationCause::UserClickedOnTab);
-			// ... but let the tab well bring up the context menu or whatever it wants to do with the right click.
+			// Show context menu for this tab otherwise let tab well handle it
+			if (const TSharedPtr<SDockingTabStack> DockingTabStack = GetParentDockTabStack())
+			{
+				const FWidgetPath WidgetPath = MouseEvent.GetEventPath() != nullptr ? *MouseEvent.GetEventPath() : FWidgetPath();
+				DockingTabStack->ShowContextMenu(SharedThis(this), WidgetPath);
+				return FReply::Handled();
+			}
+
+			// let the tab well bring up the context menu or whatever it wants to do with the right click.
 			return FReply::Unhandled();
 		}
 	}
diff --git a/Engine/Source/Runtime/Slate/Public/Framework/Docking/SDockingTabStack.h b/Engine/Source/Runtime/Slate/Public/Framework/Docking/SDockingTabStack.h
--- a/Engine/Source/Runtime/Slate/Public/Framework/Docking/SDockingTabStack.h
+++ b/Engine/Source/Runtime/Slate/Public/Framework/Docking/SDockingTabStack.h
@@ -160,8 +160,7 @@ public:
 	 */
 	SLATE_API void ReserveSpaceForWindowChrome(EChromeElement Element, bool bIncludePaddingForMenuBar, bool bOnlyMinorTabs);
 
-public:
-	/** SWidget interface */
+	//~ Begin SWidget
 	SLATE_API virtual FReply OnDragOver( const FGeometry& MyGeometry, const FDragDropEvent& DragDropEvent ) override;
 	SLATE_API virtual void OnDragLeave( const FDragDropEvent& DragDropEvent ) override;
 	SLATE_API virtual FReply OnDrop( const FGeometry& MyGeometry, const FDragDropEvent& DragDropEvent ) override;
@@ -169,6 +168,10 @@ public:
 	SLATE_API virtual FReply OnMouseButtonDown( const FGeometry& MyGeometry, const FPointerEvent& MouseEvent ) override;
 	SLATE_API virtual FReply OnKeyDown(const FGeometry& MyGeometry, const FKeyEvent& InKeyEvent) override;
 	virtual bool SupportsKeyboardFocus() const override { return true; }
+	//~ End SWidget
+
+	/** Create and show context menu for a specific tab */
+	void ShowContextMenu(const TSharedRef<SDockTab>& InContextTab, const FWidgetPath& InWidgetPath);
 
 protected:
 
@@ -179,6 +182,8 @@ protected:
 	}
 
 	SLATE_API void CloseForegroundTab();
+	
+	void CloseTab(TSharedPtr<SDockTab> InTab);
 
 	enum class ETabsToClose : uint8
 	{
@@ -194,6 +199,8 @@ protected:
 	};
 
 	SLATE_API void CloseTabsInDirectionFromForegroundTab(ETabsToClose TabsToClose, ECloseTabsDirection Direction);
+	
+	void CloseTabsInDirectionFromTab(TSharedPtr<SDockTab> InTab, ETabsToClose InTabsToClose, ECloseTabsDirection InDirection);
 
 	/**
 	 * Close all the background tabs.
@@ -201,6 +208,8 @@ protected:
 	 * @param  TabsToClose   Close just document tabs or both document and tool tabs.
 	 */
 	SLATE_API void CloseAllButForegroundTab(ETabsToClose TabsToClose);
+	
+	void CloseAllButTab(TSharedPtr<SDockTab> InTab, ETabsToClose InTabsToClose);
 
 	SLATE_API FReply TabWellRightClicked( const FGeometry& TabWellGeometry, const FPointerEvent& MouseEvent );
 
@@ -253,23 +262,35 @@ private:
 
 	/** Moves the foreground tab to a sidebar */
 	SLATE_API void MoveForegroundTabToSidebar();
+	
+	/** Moves the foreground tab to a sidebar */
+	void MoveTabToSidebar(TSharedPtr<SDockTab> InTab);
 
 	/** Unhides the tab well, revealing all tab headers */
 	SLATE_API FReply UnhideTabWell();
 
 	/** Only allow closing the tab well when the tab allows it. */
 	SLATE_API bool CanCloseForegroundTab() const;
+	
+	bool CanCloseTab(TSharedPtr<SDockTab> InTab) const;
 
 	static int32 GetIncrementFromCloseTabsDirection(ECloseTabsDirection);
 
 	/** Only allow closing tabs to the right when there is more than one tab open, the tab is of type Document or Major, and the tab is not furthest to the right. */
 	SLATE_API bool CanCloseTabsInDirectionOfForegroundTab(ECloseTabsDirection Direction) const;
+	
+	bool CanCloseTabsInDirectionOfTab(TSharedPtr<SDockTab> InTab, ECloseTabsDirection InDirection) const;
 
 	/** Only allow closing all other tabs when there are more then one tab open and the tab is of type Document or Major. */
 	SLATE_API bool CanCloseAllButForegroundTab() const;
+	
+	bool CanCloseAllButTab(TSharedPtr<SDockTab> InTab) const;
 
 	/** @return true if the foreground tab can be moved to a sidebar */
 	SLATE_API bool CanMoveForegroundTabToSidebar() const;
+	
+	/** @return true if the tab can be moved to a sidebar */
+	bool CanMoveTabToSidebar(TSharedPtr<SDockTab> InTab) const;
 
 	/** @return true if the specified tab can be moved to a sidebar */
 	SLATE_API bool IsTabAllowedInSidebar(TSharedPtr<SDockTab> Tab) const;
@@ -288,7 +309,7 @@ private:
 
 	FOverlayManagement OverlayManagement;
 
-	SLATE_API TSharedRef<SWidget> MakeContextMenu();
+	SLATE_API TSharedRef<SWidget> MakeContextMenu(const TSharedPtr<SDockTab>& InTab);
 
 	/** Show the docking cross */
 	SLATE_API void ShowCross();
