EditorUsability : SSpinBox
- Added more level of precision to editor spinboxes
Shift = Deca (+10/-10)
Ctrl = Deci (+0.1/-0.1)
Shift + Alt = Hecto (+100/-100)
Ctrl + Alt = Centi (+0.01/-0.01)

#Synced-CL 49549735

diff --git a/Engine/Source/Runtime/Slate/Private/Widgets/Input/SSpinBox.cpp b/Engine/Source/Runtime/Slate/Private/Widgets/Input/SSpinBox.cpp
--- a/Engine/Source/Runtime/Slate/Private/Widgets/Input/SSpinBox.cpp
+++ b/Engine/Source/Runtime/Slate/Private/Widgets/Input/SSpinBox.cpp
@@ -127,6 +127,8 @@ void SSpinBox<NumericType>::Construct(const FArguments& InArgs)
 	Delta = InArgs._Delta;
 	ShiftMultiplier = InArgs._ShiftMultiplier;
 	CtrlMultiplier = InArgs._CtrlMultiplier;
+	ShiftAltMultiplier = InArgs._ShiftAltMultiplier;
+	CtrlAltMultiplier = InArgs._CtrlAltMultiplier;
 	LinearDeltaSensitivity = InArgs._LinearDeltaSensitivity;
 
 	BackgroundHoveredBrush = &InArgs._Style->HoveredBackgroundBrush;
@@ -578,16 +580,8 @@ FReply SSpinBox<NumericType>::OnMouseWheel(const FGeometry& MyGeometry, const FP
 		const bool bIsSmallStep = !bIsIntegral && (GetMaxSliderValue() - GetMinSliderValue()) <= SmallStepSizeMax;
 		double Step = WheelStep.IsSet() && WheelStep.Get().IsSet() ? WheelStep.Get().GetValue() : (bIsSmallStep ? SmallStepSize : StepSize);
 
-		if (MouseEvent.IsControlDown())
-		{
-			// If no value is set for WheelSmallStep, we use the DefaultStep multiplied by the CtrlMultiplier
-			Step *= CtrlMultiplier.Get();
-		}
-		else if (MouseEvent.IsShiftDown())
-		{
-			// If no value is set for WheelBigStep, we use the DefaultStep multiplied by the ShiftMultiplier
-			Step *= ShiftMultiplier.Get();
-		}
+		// Apply multiplier based on mouse event
+		Step *= GetInputEventMultiplier(MouseEvent);
 
 		const double Sign = (MouseEvent.GetWheelDelta() > 0) ? 1.0 : -1.0;
 		const double NewValue = InternalValue + (Sign * Step);
@@ -1082,6 +1076,37 @@ float SSpinBox<NumericType>::Fraction(double InValue, double InMinValue, double
 	return (float)FMath::Clamp((HalfVal - HalfMin) / (HalfMax - HalfMin), 0.0, 1.0);
 }
 
+template <typename NumericType>
+float SSpinBox<NumericType>::GetInputEventMultiplier(const FInputEvent& InEvent) const
+{
+	float Multiplier = 1.f;
+
+	if (InEvent.IsShiftDown())
+	{
+		if (InEvent.IsAltDown())
+		{
+			Multiplier = ShiftAltMultiplier.Get();
+		}
+		else
+		{
+			Multiplier = ShiftMultiplier.Get();
+		}
+	}
+	else if (InEvent.IsControlDown())
+	{
+		if (InEvent.IsAltDown())
+		{
+			Multiplier = CtrlAltMultiplier.Get();
+		}
+		else
+		{
+			Multiplier = CtrlMultiplier.Get();
+		}
+	}
+
+	return Multiplier;
+}
+
 template <typename NumericType>
 void SSpinBox<NumericType>::UpdateIsSpinRangeUnlimited()
 {
@@ -1135,15 +1160,9 @@ double SSpinBox<NumericType>::GetDefaultStepSize(const FInputEvent& InputEvent)
 	const bool bIsSmallStep = (GetMaxSliderValue() - GetMinSliderValue()) <= SmallStepSizeMax;
 	double Step = bIsSmallStep ? SmallStepSize : StepSize;
 
-	if (InputEvent.IsControlDown())
-	{
-		Step *= CtrlMultiplier.Get();
-	}
-	else if (InputEvent.IsShiftDown())
-	{
-		Step *= ShiftMultiplier.Get();
-	}
-	
+	// Apply multiplier based on mouse event
+	Step *= GetInputEventMultiplier(InputEvent);
+
 	return Step;
 }
 
diff --git a/Engine/Source/Runtime/Slate/Public/Widgets/Input/SSpinBox.h b/Engine/Source/Runtime/Slate/Public/Widgets/Input/SSpinBox.h
--- a/Engine/Source/Runtime/Slate/Public/Widgets/Input/SSpinBox.h
+++ b/Engine/Source/Runtime/Slate/Public/Widgets/Input/SSpinBox.h
@@ -86,6 +86,8 @@ public:
 		, _Delta(0)
 		, _ShiftMultiplier(10.f)
 		, _CtrlMultiplier(0.1f)
+		, _ShiftAltMultiplier(100.f)
+		, _CtrlAltMultiplier(0.01f)
 		, _SupportDynamicSliderMaxValue(false)
 		, _SupportDynamicSliderMinValue(false)
 		, _SliderExponent(1.f)
@@ -133,6 +135,10 @@ public:
 		SLATE_ATTRIBUTE(float, ShiftMultiplier)
 		/** Multiplier to use when ctrl is held down */
 		SLATE_ATTRIBUTE(float, CtrlMultiplier)
+		/** Multiplier to use when shift+alt is held down */
+		SLATE_ATTRIBUTE(float, ShiftAltMultiplier)
+		/** Multiplier to use when ctrl+alt is held down */
+		SLATE_ATTRIBUTE(float, CtrlAltMultiplier)
 		/** If we're an unbounded spinbox, what value do we divide mouse movement by before multiplying by Delta. Requires Delta to be set. */
 		SLATE_ATTRIBUTE(int32, LinearDeltaSensitivity)
 		/** Tell us if we want to support dynamically changing of the max value using alt */
@@ -393,7 +399,9 @@ protected:
 	SLATE_API static float Fraction(double InValue, double InMinValue, double InMaxValue);
 
 private:
-	
+	/** Returns the specific step multiplier based on the modifier keys pressed */
+	float GetInputEventMultiplier(const FInputEvent& InEvent) const;
+
 	// New value to be Committed on Tick for MouseMove events
 	// This exists to insulate high-frequency mouse move events which can fire many times during input processing
 	// from the side-effects of committing the spinbox value
@@ -435,8 +443,13 @@ private:
 
 	float DistanceDragged;
 	TAttribute<NumericType> Delta;
+
+	// Step multipliers
 	TAttribute<float> ShiftMultiplier;
 	TAttribute<float> CtrlMultiplier;
+	TAttribute<float> ShiftAltMultiplier;
+	TAttribute<float> CtrlAltMultiplier;
+
 	TAttribute<int32> LinearDeltaSensitivity;
 	TAttribute<float> SliderExponent;
 	TAttribute<NumericType> SliderExponentNeutralValue;
