fix comment in InputChord

#Synced-CL 49661615

diff --git a/Engine/Source/Runtime/Slate/Public/Framework/Commands/InputChord.h b/Engine/Source/Runtime/Slate/Public/Framework/Commands/InputChord.h
--- a/Engine/Source/Runtime/Slate/Public/Framework/Commands/InputChord.h
+++ b/Engine/Source/Runtime/Slate/Public/Framework/Commands/InputChord.h
@@ -176,8 +176,6 @@ struct FInputChord
 
 	/**
 	 * Gets the corresponding EModifierKey::Type
-	 *
-	 * @return A localized string.
 	 */
 	SLATE_API EModifierKey::Type GetModifierKey() const;
 
