SlateIM: Sidebar tabs now render their contents.

#Synced-CL 49593614

diff --git a/Engine/Plugins/Experimental/SlateIM/Source/SlateIM/Private/Containers/SImTab.h b/Engine/Plugins/Experimental/SlateIM/Source/SlateIM/Private/Containers/SImTab.h
--- a/Engine/Plugins/Experimental/SlateIM/Source/SlateIM/Private/Containers/SImTab.h
+++ b/Engine/Plugins/Experimental/SlateIM/Source/SlateIM/Private/Containers/SImTab.h
@@ -26,6 +26,21 @@ public:
 	}
 
 	UE_API void ForceCloseTab();
+	
+	bool IsDrawerOpen() const
+	{
+		return bDrawerOpen;
+	}
+
+	void OnDrawerOpened()
+	{
+		bDrawerOpen = true;
+	}
+
+	void OnDrawerClosed()
+	{
+		bDrawerOpen = false;
+	}
 
 	FName TabId = NAME_None;
 	FSlateIcon TabIcon;
@@ -33,6 +48,7 @@ public:
 
 private:
 	TSharedPtr<SImStackBox> Container;
+	bool bDrawerOpen = false;
 };
 
 #undef UE_API
diff --git a/Engine/Plugins/Experimental/SlateIM/Source/SlateIM/Private/Containers/SImTabGroup.cpp b/Engine/Plugins/Experimental/SlateIM/Source/SlateIM/Private/Containers/SImTabGroup.cpp
--- a/Engine/Plugins/Experimental/SlateIM/Source/SlateIM/Private/Containers/SImTabGroup.cpp
+++ b/Engine/Plugins/Experimental/SlateIM/Source/SlateIM/Private/Containers/SImTabGroup.cpp
@@ -281,7 +281,12 @@ bool SImTabGroup::IsTabForegrounded(const FName& TabId) const
 		{
 			if (Tab.TabWidget.IsValid() && Tab.TabId == TabId)
 			{
-				return Tab.TabWidget->IsForeground();
+				if (Tab.TabWidget->IsForeground())
+				{
+					return true;
+				}
+				
+				return (Tab.TabWidget->IsInSidebar() && Tab.TabWidget->IsDrawerOpen());
 			}
 		}
 	}
diff --git a/Engine/Plugins/Experimental/SlateIM/Source/SlateIM/Private/Implementation/SlateIM_Tabs.cpp b/Engine/Plugins/Experimental/SlateIM/Source/SlateIM/Private/Implementation/SlateIM_Tabs.cpp
--- a/Engine/Plugins/Experimental/SlateIM/Source/SlateIM/Private/Implementation/SlateIM_Tabs.cpp
+++ b/Engine/Plugins/Experimental/SlateIM/Source/SlateIM/Private/Implementation/SlateIM_Tabs.cpp
@@ -1,4 +1,4 @@
-﻿// Copyright Epic Games, Inc. All Rights Reserved.
+// Copyright Epic Games, Inc. All Rights Reserved.
 
 #include "SlateIM.h"
 
@@ -106,6 +106,8 @@ namespace SlateIM
 				ContainerWidget = SNew(SImTab)
 					.CanEverClose(false);
 				ContainerWidget->TabId = TabId;
+				ContainerWidget->SetOnTabDrawerOpened(FSimpleDelegate::CreateSP(ContainerWidget.Get(), &SImTab::OnDrawerOpened));
+				ContainerWidget->SetOnTabDrawerClosed(FSimpleDelegate::CreateSP(ContainerWidget.Get(), &SImTab::OnDrawerClosed));
 				Scope.UpdateWidget(ContainerWidget);
 
 				// TabIcon and TabTitle are not updated after construction, this could be supported in the future but is not currently part of the Tab Manager API
diff --git a/Engine/Source/Runtime/Slate/Private/Framework/Docking/SDockingTabStack.cpp b/Engine/Source/Runtime/Slate/Private/Framework/Docking/SDockingTabStack.cpp
--- a/Engine/Source/Runtime/Slate/Private/Framework/Docking/SDockingTabStack.cpp
+++ b/Engine/Source/Runtime/Slate/Private/Framework/Docking/SDockingTabStack.cpp
@@ -15,6 +15,7 @@
 #include "Framework/Docking/SDockingCross.h"
 #include "Framework/Docking/FDockingDragOperation.h"
 #include "Framework/Docking/TabCommands.h"
+#include "Framework/Docking/STabDrawer.h"
 #include "Brushes/SlateColorBrush.h"
 
 #define LOCTEXT_NAMESPACE "DockTabStack"
@@ -504,6 +505,19 @@ FReply SDockingTabStack::OnDrop( const FGeometry& MyGeometry, const FDragDropEve
 
 void SDockingTabStack::OnFocusChanging( const FWeakWidgetPath& PreviousFocusPath, const FWidgetPath& NewWidgetPath, const FFocusEvent& InFocusEvent)
 {
+	// If the new widget is a drawer within the tab, the tab does not need to be activated.
+	// This prevents the focus being reset back to the containing tab and instantly closing the drawer.
+	if (NewWidgetPath.IsValid())
+	{
+		if (TSharedPtr<SWidget> LastWidget = NewWidgetPath.GetLastWidget())
+		{
+			if (LastWidget->GetType() == STabDrawer::StaticWidgetClass().GetWidgetType())
+			{
+				return;
+			}
+		}
+	}
+
 	const TSharedPtr<SDockTab> ForegroundTab = TabWell->GetForegroundTab();
 	if ( ForegroundTab.IsValid() )
 	{
diff --git a/Engine/Source/Runtime/Slate/Private/Framework/Docking/STabDrawer.cpp b/Engine/Source/Runtime/Slate/Private/Framework/Docking/STabDrawer.cpp
--- a/Engine/Source/Runtime/Slate/Private/Framework/Docking/STabDrawer.cpp
+++ b/Engine/Source/Runtime/Slate/Private/Framework/Docking/STabDrawer.cpp
@@ -8,6 +8,10 @@
 #include "Framework/Application/SlateApplication.h"
 #include "Widgets/Layout/SBox.h"
 
+void STabDrawer::PrivateRegisterAttributes(FSlateAttributeInitializer& AttributeInitializer)
+{
+}
+
 STabDrawer::~STabDrawer()
 {
 	FSlateThrottleManager::Get().LeaveResponsiveMode(AnimationThrottle);
diff --git a/Engine/Source/Runtime/Slate/Private/Widgets/Docking/SDockTab.cpp b/Engine/Source/Runtime/Slate/Private/Widgets/Docking/SDockTab.cpp
--- a/Engine/Source/Runtime/Slate/Private/Widgets/Docking/SDockTab.cpp
+++ b/Engine/Source/Runtime/Slate/Private/Widgets/Docking/SDockTab.cpp
@@ -395,6 +395,16 @@ TSharedPtr<SDockingArea> SDockTab::GetDockArea() const
 	return {};
 }
 
+bool SDockTab::IsInSidebar() const
+{
+	if (TSharedPtr<SDockingArea> DockingArea = GetDockArea())
+	{
+		return DockingArea->IsTabInSidebar(SharedThis(const_cast<SDockTab*>(this)));
+	}
+
+	return false;
+}
+
 TSharedPtr<SWindow> SDockTab::GetParentWindow() const
 {
 	TSharedPtr<SDockingArea> DockingAreaPtr = this->GetDockArea();
diff --git a/Engine/Source/Runtime/Slate/Public/Framework/Docking/STabDrawer.h b/Engine/Source/Runtime/Slate/Public/Framework/Docking/STabDrawer.h
--- a/Engine/Source/Runtime/Slate/Public/Framework/Docking/STabDrawer.h
+++ b/Engine/Source/Runtime/Slate/Public/Framework/Docking/STabDrawer.h
@@ -34,6 +34,8 @@ DECLARE_DELEGATE_OneParam(FOnDrawerClosed, TSharedRef<STabDrawer>);
  */
 class STabDrawer : public SCompoundWidget
 {
+	SLATE_DECLARE_WIDGET(STabDrawer, SCompoundWidget)
+
 	SLATE_BEGIN_ARGS(STabDrawer)
 	{
 		_Clipping = EWidgetClipping::ClipToBounds;
diff --git a/Engine/Source/Runtime/Slate/Public/Widgets/Docking/SDockTab.h b/Engine/Source/Runtime/Slate/Public/Widgets/Docking/SDockTab.h
--- a/Engine/Source/Runtime/Slate/Public/Widgets/Docking/SDockTab.h
+++ b/Engine/Source/Runtime/Slate/Public/Widgets/Docking/SDockTab.h
@@ -213,6 +213,9 @@ public:
 	/** Gets the dock area that this resides in */
 	SLATE_API TSharedPtr<class SDockingArea> GetDockArea() const;
 
+	/** Returns true if this docktab is in a sidebar. */
+	SLATE_API bool IsInSidebar() const;
+
 	/** Get the window in which this tab's tabmanager has placed it */
 	SLATE_API TSharedPtr<SWindow> GetParentWindow() const;
 
