PR #12947: Adapt to some widely used IME on windows
#Synced-CL 49703149

diff --git a/Engine/Source/Runtime/ApplicationCore/Private/Windows/TextStoreACP.cpp b/Engine/Source/Runtime/ApplicationCore/Private/Windows/TextStoreACP.cpp
--- a/Engine/Source/Runtime/ApplicationCore/Private/Windows/TextStoreACP.cpp
+++ b/Engine/Source/Runtime/ApplicationCore/Private/Windows/TextStoreACP.cpp
@@ -458,12 +458,6 @@ STDAPI FTextStoreACP::GetTextExt(TsViewCookie vcView, LONG acpStart, LONG acpEnd
 
 	const LONG StringLength = TextContext->GetTextLength();
 
-	// Begin and end indices must not be equal.
-	if(acpStart == acpEnd)
-	{
-		return E_INVALIDARG; // Documentation states TS_E_INVALIDARG, but there seems to be no such thing.
-	}
-
 	// Validate range.
 	if( acpStart < 0 || acpStart > StringLength || ( acpEnd != -1 && ( acpEnd < 0 || acpEnd > StringLength) ) )
 	{
@@ -524,15 +518,19 @@ STDAPI FTextStoreACP::GetText(LONG acpStart, LONG acpEnd, __RPC__out_ecount_part
 	}
 
 	const LONG StringLength = TextContext->GetTextLength();
+	if (acpEnd == -1)
+	{
+		acpEnd = StringLength;
+	}
 
 	// Validate range.
-	if( acpStart < 0 || acpStart > StringLength || ( acpEnd != -1 && ( acpEnd < 0 || acpEnd > StringLength) ) )
+	if( acpStart < 0 || acpStart > StringLength || acpEnd < 0 || acpEnd > StringLength )
 	{
 		return TF_E_INVALIDPOS;
 	}
 
 	const uint32 BeginIndex = acpStart;
-	const uint32 Length = (acpEnd == -1 ? StringLength : acpEnd) - BeginIndex;
+	const uint32 Length = FMath::Min(cchPlainReq, acpEnd - BeginIndex);
 
 	*pcchPlainOut = 0; // No characters yet copied to buffer.
 	// Write characters to buffer only if there is a buffer with allocated size.
@@ -541,7 +539,7 @@ STDAPI FTextStoreACP::GetText(LONG acpStart, LONG acpEnd, __RPC__out_ecount_part
 		FString StringInRange;
 		TextContext->GetTextInRange(BeginIndex, Length, StringInRange);
 
-		for(uint32 i = 0; i < Length && *pcchPlainOut < cchPlainReq; ++i)
+		for(uint32 i = 0; i < Length; ++i)
 		{
 			pchPlain[i] = StringInRange[i];
 			++(*pcchPlainOut);
@@ -550,7 +548,7 @@ STDAPI FTextStoreACP::GetText(LONG acpStart, LONG acpEnd, __RPC__out_ecount_part
 
 	*pulRunInfoOut = 0; // No runs yet copied to buffer.
 	// Write runs to buffer only if there is a buffer with allocated size.
-	if(prgRunInfo && ulRunInfoReq > 0)
+	if(prgRunInfo && ulRunInfoReq > 0 && Length > 0)
 	{
 		// We only edit text that's been stripped of any markup, so we have no need to provide multiple runs.
 		prgRunInfo[0].uCount = FMath::Min(static_cast<uint32>(ulRunInfoReq), Length);
diff --git a/Engine/Source/Runtime/ApplicationCore/Private/Windows/WindowsApplication.cpp b/Engine/Source/Runtime/ApplicationCore/Private/Windows/WindowsApplication.cpp
--- a/Engine/Source/Runtime/ApplicationCore/Private/Windows/WindowsApplication.cpp
+++ b/Engine/Source/Runtime/ApplicationCore/Private/Windows/WindowsApplication.cpp
@@ -1234,16 +1234,6 @@ FWindowsApplication::FWindowsApplication( const HINSTANCE HInstance, const HICON
 	CA_SUPPRESS(6031);
 	OleInitialize( NULL );
 
-#if !USING_ADDRESS_SANITISER
-	// Disabled under AddressSanitizer as otherwise AddresSanitizer catches an issue in the cleanup code
-	// of one of the input method Microsoft DLLs during shutdown
-	TextInputMethodSystem = MakeShareable( new FWindowsTextInputMethodSystem );
-	if(!TextInputMethodSystem->Initialize())
-	{
-		TextInputMethodSystem.Reset();
-	}
-#endif
-
 	TaskbarList = FTaskbarList::Create();
 
 	// Get initial display metrics. (display information for existing desktop, before we start changing resolutions)
@@ -1511,6 +1501,19 @@ void FWindowsApplication::InitializeWindow( const TSharedRef< FGenericWindow >&
 
 	Windows.Add( Window );
 	Window->Initialize( this, InDefinition, InstanceHandle, ParentWindow, bShowImmediately );
+
+#if !USING_ADDRESS_SANITISER
+	// Disabled under AddressSanitizer as otherwise AddresSanitizer catches an issue in the cleanup code
+	// of one of the input method Microsoft DLLs during shutdown
+	static const bool InitInputMethodSystem = [this] {
+		TextInputMethodSystem = MakeShareable(new FWindowsTextInputMethodSystem);
+		if (!TextInputMethodSystem->Initialize())
+		{
+			TextInputMethodSystem.Reset();
+		}
+		return true;
+	}();
+#endif
 }
 
 void FWindowsApplication::SetMessageHandler( const TSharedRef< FGenericApplicationMessageHandler >& InMessageHandler )
diff --git a/Engine/Source/Runtime/Slate/Private/Widgets/Text/SlateEditableTextLayout.cpp b/Engine/Source/Runtime/Slate/Private/Widgets/Text/SlateEditableTextLayout.cpp
--- a/Engine/Source/Runtime/Slate/Private/Widgets/Text/SlateEditableTextLayout.cpp
+++ b/Engine/Source/Runtime/Slate/Private/Widgets/Text/SlateEditableTextLayout.cpp
@@ -902,6 +902,10 @@ bool FSlateEditableTextLayout::HandleFocusLost(const FFocusEvent& InFocusEvent)
 		ITextInputMethodSystem* const TextInputMethodSystem = FSlateApplication::Get().GetTextInputMethodSystem();
 		if (TextInputMethodSystem && bHasRegisteredTextInputMethodContext)
 		{
+			if (TextInputMethodContext->IsComposing() && TextInputMethodChangeNotifier.IsValid())
+			{
+				TextInputMethodChangeNotifier->CancelComposition();
+			}
 			TextInputMethodSystem->DeactivateContext(TextInputMethodContext.ToSharedRef());
 		}
 	}
@@ -2750,7 +2754,8 @@ void FSlateEditableTextLayout::UpdateCursorHighlight()
 			}
 		}
 	}
-	else if (bHasSelection)
+	
+	if (bHasSelection)
 	{
 		const FTextSelection Selection(SelectionLocation, CursorInteractionPosition);
 
@@ -3214,6 +3219,13 @@ void FSlateEditableTextLayout::EndEditTransaction()
 		// Save text state
 		SaveText(EditedText);
 
+		if (TextInputMethodChangeNotifier.IsValid() && TextInputMethodContext.IsValid())
+		{
+			uint32 OldLen = StateBeforeChangingText.GetValue().Text.ToString().Len();
+			uint32 NewLen = EditedText.ToString().Len();
+			TextInputMethodChangeNotifier->NotifyTextChanged(0, OldLen, NewLen);
+		}
+
 		// Text was actually changed, so push the undo state we previously saved off
 		PushUndoState(StateBeforeChangingText.GetValue());
 
@@ -4088,17 +4100,25 @@ bool FSlateEditableTextLayout::FTextInputMethodContext::GetTextBounds(const uint
 	const FVector2D BeginPosition = OwnerLayout->TextLayout->GetLocationAt(BeginLocation, false);
 	const FVector2D EndPosition = OwnerLayout->TextLayout->GetLocationAt(EndLocation, false);
 
+	const TArray< FTextLayout::FLineView >& LineViews = OwnerLayout->TextLayout->GetLineViews();
+	int32 LineViewIndex = OwnerLayout->TextLayout->GetLineViewIndexForTextLocation(LineViews, BeginLocation, false);
+	double BeginLineHeight{};
+	if (LineViewIndex != INDEX_NONE)
+	{
+		BeginLineHeight = LineViews[LineViewIndex].Size.Y;
+	}
+
 	if (BeginPosition.Y == EndPosition.Y)
 	{
 		// The text range is contained within a single line
-		Position = BeginPosition;
-		Size = EndPosition - BeginPosition;
+		Position = FVector2D(BeginPosition.X, BeginPosition.Y - BeginLineHeight);
+		Size = EndPosition - Position;
 	}
 	else
 	{
 		// If the two positions aren't on the same line, then we assume the worst case scenario, and make the size as wide as the text area itself
-		Position = FVector2D(0.0f, BeginPosition.Y);
-		Size = FVector2D(OwnerLayout->TextLayout->GetDrawSize().X, EndPosition.Y - BeginPosition.Y);
+		Position = FVector2D(0.0f, BeginPosition.Y - BeginLineHeight);
+		Size = FVector2D(OwnerLayout->TextLayout->GetDrawSize().X, EndPosition.Y - Position.Y);
 	}
 
 	// Translate the position (which is in local space) into screen (absolute) space
