EditorUsability : Widgets
- Improving slate widget UI following feedback (part 1)
- SAssetSearchBox height can be customized (needed for a following changelist)
- SEditableText HintText opacity can be customized
- Decreased SChordEditor hint text opacity
- Replaced SChordEditBox remove icon

#Synced-CL 49549735

diff --git a/Engine/Source/Editor/EditorWidgets/Private/SAssetSearchBox.cpp b/Engine/Source/Editor/EditorWidgets/Private/SAssetSearchBox.cpp
--- a/Engine/Source/Editor/EditorWidgets/Private/SAssetSearchBox.cpp
+++ b/Engine/Source/Editor/EditorWidgets/Private/SAssetSearchBox.cpp
@@ -106,7 +106,7 @@ void SAssetSearchBox::Construct( const FArguments& InArgs )
 				[
 					SNew(SBox)
 					.WidthOverride(this, &SAssetSearchBox::GetSuggestionListWidth)
-					.HeightOverride(250.f)	// avoids flickering, ideally this would be adaptive to the content without flickering
+					.HeightOverride(InArgs._SuggestionListHeight) // avoids flickering, ideally this would be adaptive to the content without flickering
 					[
 						SAssignNew(SuggestionListView, SListView< TSharedPtr<FSuggestionListEntry> >)
 						.ListItemsSource(&Suggestions)
diff --git a/Engine/Source/Editor/EditorWidgets/Public/SAssetSearchBox.h b/Engine/Source/Editor/EditorWidgets/Public/SAssetSearchBox.h
--- a/Engine/Source/Editor/EditorWidgets/Public/SAssetSearchBox.h
+++ b/Engine/Source/Editor/EditorWidgets/Public/SAssetSearchBox.h
@@ -82,6 +82,7 @@ public:
 		, _DelayChangeNotificationsWhileTyping( false )
 		, _MustMatchPossibleSuggestions( false )
 		, _ShowSearchHistory(false)
+		, _SuggestionListHeight(250.f)
 	{}
 
 		/** Where to place the suggestion list */
@@ -123,6 +124,9 @@ public:
 		/** Handler for when the + Button next to a search is clicked */
 		SLATE_EVENT(SFilterSearchBox::FOnSaveSearchClicked, OnSaveSearchClicked)
 
+		/** Height of the suggestion list dropdown */
+		SLATE_ARGUMENT(float, SuggestionListHeight)
+
 	SLATE_END_ARGS()
 
 	/** Constructs this widget with InArgs */
diff --git a/Engine/Source/Editor/InputBindingEditor/Private/Widgets/SChordEditBox.cpp b/Engine/Source/Editor/InputBindingEditor/Private/Widgets/SChordEditBox.cpp
--- a/Engine/Source/Editor/InputBindingEditor/Private/Widgets/SChordEditBox.cpp
+++ b/Engine/Source/Editor/InputBindingEditor/Private/Widgets/SChordEditBox.cpp
@@ -92,8 +92,8 @@ void SChordEditBox::Construct(const FArguments& InArgs, TSharedPtr<FUICommandInf
 					.ToolTipText(LOCTEXT("ChordEditButtonRemove_ToolTip", "Remove this binding"))
 					[
 						SNew(SImage)
-						.Image(FAppStyle::GetBrush("Symbols.X"))
-						.ColorAndOpacity(FLinearColor(.7f,0,0,.75f))
+						.Image(FAppStyle::Get().GetBrush("Icons.X"))
+						.ColorAndOpacity(FSlateColor::UseForeground())
 					]
 				]
 			]
diff --git a/Engine/Source/Editor/InputBindingEditor/Private/Widgets/SChordEditor.cpp b/Engine/Source/Editor/InputBindingEditor/Private/Widgets/SChordEditor.cpp
--- a/Engine/Source/Editor/InputBindingEditor/Private/Widgets/SChordEditor.cpp
+++ b/Engine/Source/Editor/InputBindingEditor/Private/Widgets/SChordEditor.cpp
@@ -40,6 +40,7 @@ void SChordEditor::Construct( const FArguments& InArgs, TSharedPtr<FUICommandInf
 		.Text( this, &SChordEditor::OnGetChordInputText )
 		.HintText( this, &SChordEditor::OnGetChordInputHintText )
 		.Font( FCoreStyle::GetDefaultFontStyle("Regular", 9) )
+		.HintTextOpacity(0.25f)
 	);
 
 	EditableTextLayout->LoadText();
diff --git a/Engine/Source/Runtime/Slate/Private/Widgets/Input/SEditableText.cpp b/Engine/Source/Runtime/Slate/Private/Widgets/Input/SEditableText.cpp
--- a/Engine/Source/Runtime/Slate/Private/Widgets/Input/SEditableText.cpp
+++ b/Engine/Source/Runtime/Slate/Private/Widgets/Input/SEditableText.cpp
@@ -71,6 +71,7 @@ void SEditableText::Construct( const FArguments& InArgs )
 
 	EditableTextLayout = MakeUnique<FSlateEditableTextLayout>(*this, InArgs._Text, TextStyle, InArgs._TextShapingMethod, InArgs._TextFlowDirection, FCreateSlateTextLayout(), PlainTextMarshaller.ToSharedRef(), HintTextMarshaller);
 	EditableTextLayout->SetHintText(InArgs._HintText);
+	EditableTextLayout->SetHintTextOpacity(InArgs._HintTextOpacity);
 	EditableTextLayout->SetSearchText(InArgs._SearchText);
 	EditableTextLayout->SetCursorBrush(InArgs._CaretImage.IsSet() ? InArgs._CaretImage : &InArgs._Style->CaretImage);
 	EditableTextLayout->SetCompositionBrush(InArgs._BackgroundImageComposing.IsSet() ? InArgs._BackgroundImageComposing : &InArgs._Style->BackgroundImageComposing);
@@ -326,6 +327,16 @@ FText SEditableText::GetHintText() const
 	return EditableTextLayout->GetHintText();
 }
 
+void SEditableText::SetHintTextOpacity(const TAttribute<float>& InHintTextOpacity)
+{
+	EditableTextLayout->SetHintTextOpacity(InHintTextOpacity);
+}
+
+float SEditableText::GetHintTextOpacity() const
+{
+	return EditableTextLayout->GetHintTextOpacity();
+}
+
 void SEditableText::SetSearchText(const TAttribute<FText>& InSearchText)
 {
 	EditableTextLayout->SetSearchText(InSearchText);
diff --git a/Engine/Source/Runtime/Slate/Private/Widgets/Text/SlateEditableTextLayout.cpp b/Engine/Source/Runtime/Slate/Private/Widgets/Text/SlateEditableTextLayout.cpp
--- a/Engine/Source/Runtime/Slate/Private/Widgets/Text/SlateEditableTextLayout.cpp
+++ b/Engine/Source/Runtime/Slate/Private/Widgets/Text/SlateEditableTextLayout.cpp
@@ -332,6 +332,16 @@ const FTextBlockStyle& FSlateEditableTextLayout::GetTextStyle() const
 	return TextStyle;
 }
 
+void FSlateEditableTextLayout::SetHintTextOpacity(const TAttribute<float>& InOpacity)
+{
+	HintTextOpacity = InOpacity;
+}
+
+float FSlateEditableTextLayout::GetHintTextOpacity() const
+{
+	return HintTextOpacity.Get(0.35f);
+}
+
 void FSlateEditableTextLayout::SetCursorBrush(const TAttribute<const FSlateBrush*>& InCursorBrush)
 {
 	CursorLineHighlighter->SetCursorBrush(InCursorBrush);
@@ -3597,10 +3607,10 @@ int32 FSlateEditableTextLayout::OnPaint(const FPaintArgs& Args, const FGeometry&
 	// Only paint the hint text layout if we don't have any text set
 	if (TextLayout->IsEmpty() && HintTextLayout.IsValid())
 	{
-		const FLinearColor ThisColorAndOpacity = TextStyle.ColorAndOpacity.GetColor(InWidgetStyle);
-
 		// Make sure the hint text is the correct color before we paint it
-		HintTextStyle.ColorAndOpacity = FLinearColor(ThisColorAndOpacity.R, ThisColorAndOpacity.G, ThisColorAndOpacity.B, 0.35f);
+		FLinearColor ThisColorAndOpacity = TextStyle.ColorAndOpacity.GetColor(InWidgetStyle);
+		ThisColorAndOpacity.A = GetHintTextOpacity();
+		HintTextStyle.ColorAndOpacity = ThisColorAndOpacity;
 		HintTextLayout->OverrideTextStyle(HintTextStyle);
 
 		LayerId = HintTextLayout->OnPaint(Args, AllottedGeometry, MyCullingRect, OutDrawElements, LayerId, InWidgetStyle, bParentEnabled);
diff --git a/Engine/Source/Runtime/Slate/Public/Widgets/Input/SEditableText.h b/Engine/Source/Runtime/Slate/Public/Widgets/Input/SEditableText.h
--- a/Engine/Source/Runtime/Slate/Public/Widgets/Input/SEditableText.h
+++ b/Engine/Source/Runtime/Slate/Public/Widgets/Input/SEditableText.h
@@ -47,6 +47,7 @@ public:
 		, _Style(&FCoreStyle::Get().GetWidgetStyle< FEditableTextStyle >("NormalEditableText"))
 		, _Font()
 		, _ColorAndOpacity()
+		, _HintTextOpacity()
 		, _BackgroundImageSelected()
 		, _BackgroundImageComposing()
 		, _CaretImage()
@@ -91,6 +92,9 @@ public:
 		/** Text color and opacity (overrides Style) */
 		SLATE_ATTRIBUTE( FSlateColor, ColorAndOpacity )
 
+		/** Hint text opacity in the editable box */
+		SLATE_ATTRIBUTE(float, HintTextOpacity)
+
 		/** Background image for the selected text (overrides Style) */
 		SLATE_ATTRIBUTE( const FSlateBrush*, BackgroundImageSelected )
 
@@ -238,6 +242,12 @@ public:
 	
 	/** Get the text that appears when there is no text in the text box */
 	SLATE_API FText GetHintText() const;
+	
+	/** Sets the hint text opacity */
+	SLATE_API void SetHintTextOpacity(const TAttribute<float>& InHintTextOpacity);
+	
+	/** Get the hint text opacity */
+	SLATE_API float GetHintTextOpacity() const;
 
 	/** Set the text that is currently being searched for (if any) */
 	SLATE_API void SetSearchText(const TAttribute<FText>& InSearchText);
diff --git a/Engine/Source/Runtime/Slate/Public/Widgets/Text/SlateEditableTextLayout.h b/Engine/Source/Runtime/Slate/Public/Widgets/Text/SlateEditableTextLayout.h
--- a/Engine/Source/Runtime/Slate/Public/Widgets/Text/SlateEditableTextLayout.h
+++ b/Engine/Source/Runtime/Slate/Public/Widgets/Text/SlateEditableTextLayout.h
@@ -56,6 +56,9 @@ public:
 
 	SLATE_API void SetTextStyle(const FTextBlockStyle& InTextStyle);
 	SLATE_API const FTextBlockStyle& GetTextStyle() const;
+	
+	SLATE_API void SetHintTextOpacity(const TAttribute<float>& InOpacity);
+	SLATE_API float GetHintTextOpacity() const;
 
 	/** Get the number of Text Lines */
 	SLATE_API int32 GetTextLineCount();
@@ -566,6 +569,9 @@ private:
 
 	/** The text that appears when there is no text in the text box */
 	TAttribute<FText> HintText;
+	
+	/** Override opacity of the hint text */
+	TAttribute<float> HintTextOpacity;
 
 	/** The text to be searched for */
 	TAttribute<FText> BoundSearchText;
