Fix Blueprint color wheel not updating when used as Modals.
Modals don't receive OnFinishedPointerInput, so OnMouseMove needs to be handled immediately.

[RN] minor | Rendering Architecture

#Synced-CL 49844456

diff --git a/Engine/Source/Runtime/AppFramework/Private/Widgets/Colors/SColorPicker.cpp b/Engine/Source/Runtime/AppFramework/Private/Widgets/Colors/SColorPicker.cpp
--- a/Engine/Source/Runtime/AppFramework/Private/Widgets/Colors/SColorPicker.cpp
+++ b/Engine/Source/Runtime/AppFramework/Private/Widgets/Colors/SColorPicker.cpp
@@ -169,6 +169,7 @@ void SColorPicker::GenerateDefaultColorPickerContent( bool bAdvancedSectionExpan
 													.OnMouseCaptureBegin(this, &SColorPicker::HandleInteractiveChangeBegin)
 													.OnMouseCaptureEnd(this, &SColorPicker::HandleInteractiveChangeEnd)
 													.PreventThrottling(true)
+													.ProcessMouseMoveImmediate(bOnlyRefreshOnMouseUp || bOnlyRefreshOnOk)
 											]
 
 										+ SHorizontalBox::Slot()
@@ -202,6 +203,7 @@ void SColorPicker::GenerateDefaultColorPickerContent( bool bAdvancedSectionExpan
 													.OnMouseCaptureBegin(this, &SColorPicker::HandleInteractiveChangeBegin)
 													.OnMouseCaptureEnd(this, &SColorPicker::HandleInteractiveChangeEnd)
 													.PreventThrottling(true)
+													.ProcessMouseMoveImmediate(bOnlyRefreshOnMouseUp || bOnlyRefreshOnOk)
 											]
 									]
 							]
@@ -529,6 +531,7 @@ void SColorPicker::GenerateInlineColorPickerContent()
 					.OnMouseCaptureBegin(this, &SColorPicker::HandleInteractiveChangeBegin)
 					.OnMouseCaptureEnd(this, &SColorPicker::HandleInteractiveChangeEnd)
 					.PreventThrottling(true)
+					.ProcessMouseMoveImmediate(bOnlyRefreshOnMouseUp || bOnlyRefreshOnOk)
 			]
 
 		+ SHorizontalBox::Slot()
diff --git a/Engine/Source/Runtime/Slate/Private/Widgets/Colors/SColorSpectrum.cpp b/Engine/Source/Runtime/Slate/Private/Widgets/Colors/SColorSpectrum.cpp
--- a/Engine/Source/Runtime/Slate/Private/Widgets/Colors/SColorSpectrum.cpp
+++ b/Engine/Source/Runtime/Slate/Private/Widgets/Colors/SColorSpectrum.cpp
@@ -26,6 +26,7 @@ void SColorSpectrum::Construct( const FArguments& InArgs )
 	OnMouseCaptureEnd = InArgs._OnMouseCaptureEnd;
 	OnValueChanged = InArgs._OnValueChanged;
 	bPreventThrottling = InArgs._PreventThrottling;
+	bProcessMouseMoveImmediate = InArgs._ProcessMouseMoveImmediate;
 
 	CtrlMultiplier = InArgs._CtrlMultiplier;
 }
@@ -104,10 +105,17 @@ FReply SColorSpectrum::OnMouseMove( const FGeometry& MyGeometry, const FPointerE
 		LastSpectrumPosition = MyGeometry.AbsoluteToLocal(MouseEvent.GetScreenSpacePosition());
 	}
 
-	bHighPrecision = MouseEvent.IsControlDown();
-	LastGeometry = MyGeometry;
-	LastMouseEvent = MouseEvent;
+	if (bProcessMouseMoveImmediate)
+	{
+		ProcessMouseAction(MyGeometry, MouseEvent);
+	}
+	else
+	{
+		LastGeometry = MyGeometry;
+		LastMouseEvent = MouseEvent;
+	}
 
+	bHighPrecision = MouseEvent.IsControlDown();
 	if (bHighPrecision && !FSlateApplicationBase::Get().GetPlatformApplication()->IsUsingHighPrecisionMouseMode())
 	{
 		return FReply::Handled().UseHighPrecisionMouseMovement(SharedThis(this));
@@ -118,7 +126,7 @@ FReply SColorSpectrum::OnMouseMove( const FGeometry& MyGeometry, const FPointerE
 
 void SColorSpectrum::OnFinishedPointerInput()
 {
-	if (bDragging && !LastMouseEvent.GetCursorDelta().IsNearlyZero())
+	if (bDragging && !LastMouseEvent.GetCursorDelta().IsNearlyZero() && !bProcessMouseMoveImmediate)
 	{
 		ProcessMouseAction(LastGeometry, LastMouseEvent);
 		LastGeometry = FGeometry();
diff --git a/Engine/Source/Runtime/Slate/Private/Widgets/Colors/SColorWheel.cpp b/Engine/Source/Runtime/Slate/Private/Widgets/Colors/SColorWheel.cpp
--- a/Engine/Source/Runtime/Slate/Private/Widgets/Colors/SColorWheel.cpp
+++ b/Engine/Source/Runtime/Slate/Private/Widgets/Colors/SColorWheel.cpp
@@ -24,6 +24,7 @@ void SColorWheel::Construct(const FArguments& InArgs)
 	OnMouseCaptureEnd = InArgs._OnMouseCaptureEnd;
 	OnValueChanged = InArgs._OnValueChanged;
 	bPreventThrottling = InArgs._PreventThrottling;
+	bProcessMouseMoveImmediate = InArgs._ProcessMouseMoveImmediate;
 
 	CtrlMultiplier = InArgs._CtrlMultiplier;
 }
@@ -109,10 +110,17 @@ FReply SColorWheel::OnMouseMove(const FGeometry& MyGeometry, const FPointerEvent
 		LastWheelPosition = MyGeometry.AbsoluteToLocal(MouseEvent.GetScreenSpacePosition());
 	}
 
-	bHighPrecision = MouseEvent.IsControlDown();
-	LastGeometry = MyGeometry;
-	LastMouseEvent = MouseEvent;
+	if (bProcessMouseMoveImmediate)
+	{
+		ProcessMouseAction(MyGeometry, MouseEvent, true);
+	}
+	else
+	{
+		LastGeometry = MyGeometry;
+		LastMouseEvent = MouseEvent;
+	}
 
+	bHighPrecision = MouseEvent.IsControlDown();
 	if (bHighPrecision && !FSlateApplicationBase::Get().GetPlatformApplication()->IsUsingHighPrecisionMouseMode())
 	{
 		return FReply::Handled().UseHighPrecisionMouseMovement(SharedThis(this));
@@ -123,7 +131,7 @@ FReply SColorWheel::OnMouseMove(const FGeometry& MyGeometry, const FPointerEvent
 
 void SColorWheel::OnFinishedPointerInput()
 {
-	if (bDragging && !LastMouseEvent.GetCursorDelta().IsNearlyZero())
+	if (bDragging && !LastMouseEvent.GetCursorDelta().IsNearlyZero() && !bProcessMouseMoveImmediate)
 	{
 		ProcessMouseAction(LastGeometry, LastMouseEvent, true);
 		LastGeometry = FGeometry();
diff --git a/Engine/Source/Runtime/Slate/Public/Widgets/Colors/SColorSpectrum.h b/Engine/Source/Runtime/Slate/Public/Widgets/Colors/SColorSpectrum.h
--- a/Engine/Source/Runtime/Slate/Public/Widgets/Colors/SColorSpectrum.h
+++ b/Engine/Source/Runtime/Slate/Public/Widgets/Colors/SColorSpectrum.h
@@ -28,7 +28,8 @@ public:
 		, _OnValueChanged()
 		, _CtrlMultiplier(0.1f)
 		, _PreventThrottling(false)
-	{ }
+		, _ProcessMouseMoveImmediate(false)
+		{ }
 	
 		/** The current color selected by the user */
 		SLATE_ATTRIBUTE(FLinearColor, SelectedColor)
@@ -48,6 +49,9 @@ public:
 		/** If refresh requests for the viewport should happen for all value changes **/
 		SLATE_ARGUMENT(bool, PreventThrottling)
 
+		/** If mouse inputs should be processed immediately. This should be used when spawned in a Modal. */
+		SLATE_ARGUMENT(bool, ProcessMouseMoveImmediate)
+
 	SLATE_END_ARGS()
 	
 public:
@@ -114,6 +118,9 @@ private:
 	/** If refresh requests for the viewport should happen for all value changes **/
 	bool bPreventThrottling = false;
 
+	/** If mouse inputs should be processed immediately. This should be used when spawned in a Modal. */
+	bool bProcessMouseMoveImmediate = false;
+
 	/** Last Geometry to use after dragging. */
 	FGeometry LastGeometry;
 
diff --git a/Engine/Source/Runtime/Slate/Public/Widgets/Colors/SColorWheel.h b/Engine/Source/Runtime/Slate/Public/Widgets/Colors/SColorWheel.h
--- a/Engine/Source/Runtime/Slate/Public/Widgets/Colors/SColorWheel.h
+++ b/Engine/Source/Runtime/Slate/Public/Widgets/Colors/SColorWheel.h
@@ -30,6 +30,7 @@ public:
 		, _OnValueChanged()
 		, _CtrlMultiplier(0.1f)
 		, _PreventThrottling(false)
+		, _ProcessMouseMoveImmediate(false)
 	{ }
 	
 		/** The current color selected by the user. */
@@ -50,9 +51,12 @@ public:
 		/** Multiplier to use when ctrl is held down */
 		SLATE_ATTRIBUTE(float, CtrlMultiplier)
 
-		/** If refresh requests for the viewport should happen for all value changes **/
+		/** If refresh requests for the viewport should happen for all value changes. */
 		SLATE_ARGUMENT(bool, PreventThrottling)
 
+		/** If mouse inputs should be processed immediately. This should be used when spawned in a Modal. */
+		SLATE_ARGUMENT(bool, ProcessMouseMoveImmediate)
+
 	SLATE_END_ARGS()
 	
 public:
@@ -105,6 +109,9 @@ protected:
 	/** If refresh requests for the viewport should happen for all value changes **/
 	bool bPreventThrottling = false;
 
+	/** If mouse inputs should be processed immediately. This should be used when spawned in a Modal. */
+	bool bProcessMouseMoveImmediate = false;
+
 private:
 
 	/** The color wheel image to show. */
