PR #13779: Linux Classic Editor Layout Fix

Consider line terminators from any supported platform, as the ini may have been saved on a different platform.
Engine\Config\Layouts\UE4ClassicLayout.ini is currently saved with "\\\r\n" values.

#Synced-CL 49845947

diff --git a/Engine/Source/Runtime/Slate/Private/Framework/Docking/LayoutService.cpp b/Engine/Source/Runtime/Slate/Private/Framework/Docking/LayoutService.cpp
--- a/Engine/Source/Runtime/Slate/Private/Framework/Docking/LayoutService.cpp
+++ b/Engine/Source/Runtime/Slate/Private/Framework/Docking/LayoutService.cpp
@@ -120,7 +120,9 @@ static FString GetLayoutStringFromIni(const FString& LayoutString, bool& bOutIsJ
 	return LayoutString
 		.Replace(TEXT("("), TEXT("{"))
 		.Replace(TEXT(")"), TEXT("}"))
-		.Replace(TEXT("\\") LINE_TERMINATOR, LINE_TERMINATOR);
+		// Consider line terminators from any supported platform, as the ini may have been saved on a different platform.
+		.Replace(TEXT("\\\n"), LINE_TERMINATOR)
+		.Replace(TEXT("\\\r\n"), LINE_TERMINATOR);
 }
 
 
