Slate Dynamic Invalidation - SEditableText and SlateEditableTextLayout

This change introduces a new Slate attribute type: `TSlateNestedAttribute`.
This allows Slate attributes to be stored inside a container type that itself lives inside a widget (e.g `SWidget -> FAttributeHolder -> TSlateNestedAttribute`).
For compatibility with the existing attribute system, the container must be inlined within the widget. This guarantees that nested attributes live within the widget's memory footprint.

`TSlateNestedAttribute` is implemented within the helper struct `FSlateNestedAttributeContainer`, which any container needing attributes must inherit from.
A new macro `SLATE_ADD_NESTED_MEMBER_ATTRIBUTE_DEFINITION` has been added to declare nested attributes. It requires a pointer to the container member inside the widget so the system can compute the correct memory offset.

This change is required for `FSlateEditableTextLayout` to own attributes that expose an `OnValueChanged` delegate, ensuring that the right callbacks are triggered.

#Synced-CL 49551883

diff --git a/Engine/Plugins/Animation/DeformerGraph/Source/OptimusEditor/Private/Widgets/SOptimusShaderTextDocumentTextBox.cpp b/Engine/Plugins/Animation/DeformerGraph/Source/OptimusEditor/Private/Widgets/SOptimusShaderTextDocumentTextBox.cpp
--- a/Engine/Plugins/Animation/DeformerGraph/Source/OptimusEditor/Private/Widgets/SOptimusShaderTextDocumentTextBox.cpp
+++ b/Engine/Plugins/Animation/DeformerGraph/Source/OptimusEditor/Private/Widgets/SOptimusShaderTextDocumentTextBox.cpp
@@ -440,7 +440,7 @@ void SOptimusShaderTextDocumentTextBox::OnToggleComment() const
 	{
 		// Automatically go to the next line if there is no selection
 		int32 LineIndex = CursorLocation.GetLineIndex() + 1;
-		if (LineIndex < Text->GetTextLineCount())
+		if (LineIndex < const_cast<const SMultiLineEditableText*>(Text.Get())->GetTextLineCount())
 		{
 			FString Line;
 			Text->GetTextLine(LineIndex, Line);
diff --git a/Engine/Plugins/PCG/Source/PCGEditor/Private/Widgets/SPCGNodeSourceTextBox.cpp b/Engine/Plugins/PCG/Source/PCGEditor/Private/Widgets/SPCGNodeSourceTextBox.cpp
--- a/Engine/Plugins/PCG/Source/PCGEditor/Private/Widgets/SPCGNodeSourceTextBox.cpp
+++ b/Engine/Plugins/PCG/Source/PCGEditor/Private/Widgets/SPCGNodeSourceTextBox.cpp
@@ -436,7 +436,7 @@ void SPCGNodeSourceTextBox::OnToggleComment() const
 	{
 		// Automatically go to the next line if there is no selection
 		int32 LineIndex = CursorLocation.GetLineIndex() + 1;
-		if (LineIndex < Text->GetTextLineCount())
+		if (LineIndex < const_cast<const SMultiLineEditableText*>(Text.Get())->GetTextLineCount())
 		{
 			FString Line;
 			Text->GetTextLine(LineIndex, Line);
diff --git a/Engine/Source/Editor/InputBindingEditor/Private/Widgets/SChordEditor.cpp b/Engine/Source/Editor/InputBindingEditor/Private/Widgets/SChordEditor.cpp
--- a/Engine/Source/Editor/InputBindingEditor/Private/Widgets/SChordEditor.cpp
+++ b/Engine/Source/Editor/InputBindingEditor/Private/Widgets/SChordEditor.cpp
@@ -43,7 +43,7 @@ void SChordEditor::Construct( const FArguments& InArgs, TSharedPtr<FUICommandInf
 		.HintTextOpacity(0.25f)
 	);
 
-	EditableTextLayout->LoadText();
+	SlateEditableTextLayout.LoadText();
 }
 
 
@@ -64,17 +64,17 @@ FReply SChordEditor::OnKeyDown( const FGeometry& MyGeometry, const FKeyEvent& In
 			EditingInputChord.Key = Key;
 		}
 	
-		EditableTextLayout->BeginEditTransation();
+		SlateEditableTextLayout.BeginEditTransaction();
 
 		EditingInputChord.bCtrl = InKeyEvent.IsControlDown();
 		EditingInputChord.bAlt = InKeyEvent.IsAltDown();
 		EditingInputChord.bShift = InKeyEvent.IsShiftDown();
 		EditingInputChord.bCmd = InKeyEvent.IsCommandDown();
 
-		EditableTextLayout->LoadText();
-		EditableTextLayout->GoTo(ETextLocation::EndOfDocument);
+		SlateEditableTextLayout.LoadText();
+		SlateEditableTextLayout.GoTo(ETextLocation::EndOfDocument);
 
-		EditableTextLayout->EndEditTransaction();
+		SlateEditableTextLayout.EndEditTransaction();
 		
 		OnChordTyped( EditingInputChord );
 		
diff --git a/Engine/Source/Runtime/Slate/Private/Widgets/Input/SEditableText.cpp b/Engine/Source/Runtime/Slate/Private/Widgets/Input/SEditableText.cpp
--- a/Engine/Source/Runtime/Slate/Private/Widgets/Input/SEditableText.cpp
+++ b/Engine/Source/Runtime/Slate/Private/Widgets/Input/SEditableText.cpp
@@ -4,6 +4,7 @@
 #include "Framework/Text/TextEditHelper.h"
 #include "Framework/Text/PlainTextLayoutMarshaller.h"
 #include "Widgets/Text/SlateEditableTextLayout.h"
+#include "Widgets/Text/SlateTextBlockLayout.h"
 #include "Types/ReflectionMetadata.h"
 #include "Types/TrackedMetaData.h"
 
@@ -11,7 +12,44 @@
 #include "Widgets/Accessibility/SlateAccessibleWidgets.h"
 #endif
 
+SLATE_IMPLEMENT_WIDGET(SEditableText)
+void SEditableText::PrivateRegisterAttributes(FSlateAttributeInitializer& AttributeInitializer)
+{
+	struct FInvalidation
+	{
+		static void SynchronizeTextStyle(SWidget& OwningWidget)
+		{
+			SEditableText& EditableText = static_cast<SEditableText&>(OwningWidget);
+			EditableText.SynchronizeTextStyle();
+		}
+		static void UpdateMarshaller(SWidget& OwningWidget)
+		{
+			SEditableText& EditableText = static_cast<SEditableText&>(OwningWidget);
+			EditableText.PlainTextMarshaller->MakeDirty();
+		}
+	};
+
+	FSlateEditableTextLayout::RegisterNestedAttributes<&SEditableText::SlateEditableTextLayout>(AttributeInitializer);
+
+	SLATE_ADD_MEMBER_ATTRIBUTE_DEFINITION(AttributeInitializer, MinDesiredWidthAttribute, EInvalidateWidgetReason::Layout);
+	SLATE_ADD_MEMBER_ATTRIBUTE_DEFINITION(AttributeInitializer, bIsPasswordAttribute, EInvalidateWidgetReason::Layout)
+		.OnValueChanged(FSlateAttributeDescriptor::FAttributeValueChangedDelegate::CreateStatic(FInvalidation::UpdateMarshaller));
+	SLATE_ADD_MEMBER_ATTRIBUTE_DEFINITION(AttributeInitializer, FontOverrideAttribute, EInvalidateWidgetReason::Layout)
+		.OnValueChanged(FSlateAttributeDescriptor::FAttributeValueChangedDelegate::CreateStatic(FInvalidation::SynchronizeTextStyle));
+
+	SLATE_ADD_MEMBER_ATTRIBUTE_DEFINITION(AttributeInitializer, ColorAndOpacityAttribute, EInvalidateWidgetReason::Paint)
+		.OnValueChanged(FSlateAttributeDescriptor::FAttributeValueChangedDelegate::CreateStatic(FInvalidation::SynchronizeTextStyle));
+	SLATE_ADD_MEMBER_ATTRIBUTE_DEFINITION(AttributeInitializer, BackgroundImageSelectedAttribute, EInvalidateWidgetReason::Layout)
+		.OnValueChanged(FSlateAttributeDescriptor::FAttributeValueChangedDelegate::CreateStatic(FInvalidation::SynchronizeTextStyle));
+}
+
 SEditableText::SEditableText()
+	: SlateEditableTextLayout(*this)
+	, FontOverrideAttribute(*this)
+	, ColorAndOpacityAttribute(*this)
+	, BackgroundImageSelectedAttribute(*this)
+	, bIsPasswordAttribute(*this, false)
+	, MinDesiredWidthAttribute(*this, 0.0f)
 {
 #if WITH_ACCESSIBILITY
 	AccessibleBehavior = EAccessibleBehavior::Auto;
@@ -19,15 +57,9 @@ SEditableText::SEditableText()
 #endif
 }
 
-SEditableText::~SEditableText()
-{
-	// Needed to avoid "deletion of pointer to incomplete type 'FSlateEditableTextLayout'; no destructor called" error when using TUniquePtr
-}
-
 void SEditableText::Construct( const FArguments& InArgs )
 {
-	bIsReadOnly = InArgs._IsReadOnly;
-	bIsPassword = InArgs._IsPassword;
+	SetIsReadOnly(InArgs._IsReadOnly);
 
 	bIsCaretMovedWhenGainFocus = InArgs._IsCaretMovedWhenGainFocus;
 	bSelectAllTextWhenFocused = InArgs._SelectAllTextWhenFocused;
@@ -40,7 +72,7 @@ void SEditableText::Construct( const FArguments& InArgs )
 	OnTextChangedCallback = InArgs._OnTextChanged;
 	OnTextCommittedCallback = InArgs._OnTextCommitted;
 	OnCursorMovedWithSelectionCallback = InArgs._OnCursorMovedWithSelection;
-	MinDesiredWidth = InArgs._MinDesiredWidth;
+	SetMinDesiredWidth(InArgs._MinDesiredWidth);
 	bSelectAllTextOnCommit = InArgs._SelectAllTextOnCommit;
 	bSelectWordOnMouseDoubleClick = InArgs._SelectWordOnMouseDoubleClick;
 	VirtualKeyboardType = InArgs._VirtualKeyboardType;
@@ -51,10 +83,6 @@ void SEditableText::Construct( const FArguments& InArgs )
 	OnKeyDownHandler = InArgs._OnKeyDownHandler;
 	bEnableIntegratedKeyboard = InArgs._EnableIntegratedKeyboard;
 
-	Font = InArgs._Font;
-	ColorAndOpacity = InArgs._ColorAndOpacity;
-	BackgroundImageSelected = InArgs._BackgroundImageSelected;
-
 	// We use the given style when creating the text layout as it may not be safe to call the override delegates until we've finished being constructed
 	// The first call to SynchronizeTextStyle will apply the correct overrides, and that will happen before the first paint
 	check(InArgs._Style);
@@ -64,26 +92,31 @@ void SEditableText::Construct( const FArguments& InArgs )
 	TextStyle.HighlightShape = InArgs._Style->BackgroundImageSelected;
 
 	PlainTextMarshaller = FPlainTextLayoutMarshaller::Create();
-	PlainTextMarshaller->SetIsPassword(bIsPassword);
 
 	// We use a separate marshaller for the hint text, as that should never be displayed as a password
 	TSharedRef<FPlainTextLayoutMarshaller> HintTextMarshaller = FPlainTextLayoutMarshaller::Create();
 
-	EditableTextLayout = MakeUnique<FSlateEditableTextLayout>(*this, InArgs._Text, TextStyle, InArgs._TextShapingMethod, InArgs._TextFlowDirection, FCreateSlateTextLayout(), PlainTextMarshaller.ToSharedRef(), HintTextMarshaller);
-	EditableTextLayout->SetHintText(InArgs._HintText);
-	EditableTextLayout->SetHintTextOpacity(InArgs._HintTextOpacity);
-	EditableTextLayout->SetSearchText(InArgs._SearchText);
-	EditableTextLayout->SetCursorBrush(InArgs._CaretImage.IsSet() ? InArgs._CaretImage : &InArgs._Style->CaretImage);
-	EditableTextLayout->SetCompositionBrush(InArgs._BackgroundImageComposing.IsSet() ? InArgs._BackgroundImageComposing : &InArgs._Style->BackgroundImageComposing);
-	EditableTextLayout->SetDebugSourceInfo(TAttribute<FString>::Create(TAttribute<FString>::FGetter::CreateLambda([this]{ return FReflectionMetaData::GetWidgetDebugInfo(this); })));
-	EditableTextLayout->SetJustification(InArgs._Justification);
-	EditableTextLayout->SetOverflowPolicy(InArgs._OverflowPolicy);
-	EditableTextLayout->SetFontFacesLoadingPaintPolicy(InArgs._FontFacesLoadingPaintPolicy);
+	SlateEditableTextLayout.Construct(InArgs._Text, TextStyle, InArgs._TextShapingMethod, InArgs._TextFlowDirection, FCreateSlateTextLayout(), PlainTextMarshaller.ToSharedRef(), MoveTemp(HintTextMarshaller));
+	SlateEditableTextLayout.SetCursorBrush(InArgs._CaretImage.IsSet() ? InArgs._CaretImage : &InArgs._Style->CaretImage);
+	SlateEditableTextLayout.SetCompositionBrush(InArgs._BackgroundImageComposing.IsSet() ? InArgs._BackgroundImageComposing : &InArgs._Style->BackgroundImageComposing);
+	SlateEditableTextLayout.SetDebugSourceInfo(TAttribute<FString>::Create(TAttribute<FString>::FGetter::CreateLambda([this]{ return FReflectionMetaData::GetWidgetDebugInfo(this); })));
+	SlateEditableTextLayout.SetOverflowPolicy(InArgs._OverflowPolicy);
+	SlateEditableTextLayout.SetFontFacesLoadingPaintPolicy(InArgs._FontFacesLoadingPaintPolicy);
 	if (InArgs._OnAllFontFacesFinishLoading.IsBound())
 	{
-		EditableTextLayout->SetOnAllFontFacesFinishLoadingDelegate(InArgs._OnAllFontFacesFinishLoading);
+		SlateEditableTextLayout.SetOnAllFontFacesFinishLoadingDelegate(InArgs._OnAllFontFacesFinishLoading);
 	}
 
+	SetText(InArgs._Text);
+	SetIsPassword(InArgs._IsPassword);
+	SetHintText(InArgs._HintText);
+	SetHintTextOpacity(InArgs._HintTextOpacity);
+	SetSearchText(InArgs._SearchText);
+	SetFont(InArgs._Font);
+	SetBackgroundImageSelected(InArgs._BackgroundImageSelected);
+	SetColorAndOpacity(InArgs._ColorAndOpacity);
+	SetJustification(InArgs._Justification);
+
 	// build context menu extender
 	MenuExtender = MakeShareable(new FExtender());
 	MenuExtender->AddMenuExtension("EditText", EExtensionHook::Before, TSharedPtr<FUICommandList>(), InArgs._ContextMenuExtender);
@@ -91,65 +124,64 @@ void SEditableText::Construct( const FArguments& InArgs )
 	AddMetadata(MakeShared<FTrackedMetaData>(this, FName("EditableText")));
 }
 
-void SEditableText::SetText( const TAttribute< FText >& InNewText )
+void SEditableText::SetText(TAttribute<FText> InNewText)
 {
-	EditableTextLayout->SetText(InNewText);
+	SlateEditableTextLayout.SetText(MoveTemp(InNewText));
 }
 
 FText SEditableText::GetText() const
 {
-	return EditableTextLayout->GetText();
+	return SlateEditableTextLayout.GetText();
 }
 
 bool SEditableText::SetEditableText(const FText& InNewText)
 {
-	return EditableTextLayout->SetEditableText(InNewText);
+	return SlateEditableTextLayout.SetEditableText(InNewText);
 }
 
-void SEditableText::SetFont( const TAttribute< FSlateFontInfo >& InNewFont )
+void SEditableText::SetFont(TAttribute<FSlateFontInfo> InNewFont)
 {
-	Font = InNewFont;
-
-	Invalidate(EInvalidateWidgetReason::Layout);
+	bIsFontOverrideAttrSet = InNewFont.IsSet();
+	FontOverrideAttribute.Assign(*this, MoveTemp(InNewFont));
 }
 
 FSlateFontInfo SEditableText::GetFont() const
 {
-	return Font.Get();
+	return bIsFontOverrideAttrSet ? FontOverrideAttribute.Get() : SlateEditableTextLayout.GetTextStyle().Font;
 }
 
 void SEditableText::SetTextStyle( const FEditableTextStyle& InNewTextStyle )
 {
-	Font = InNewTextStyle.Font;
-	ColorAndOpacity = InNewTextStyle.ColorAndOpacity;
-	BackgroundImageSelected = &InNewTextStyle.BackgroundImageSelected;
-
-	Invalidate(EInvalidateWidgetReason::Layout);
+	bIsFontOverrideAttrSet = true;
+	FontOverrideAttribute.Assign(*this, InNewTextStyle.Font);
+	bIsColorAndOpacityAttrSet = true;
+	ColorAndOpacityAttribute.Assign(*this, InNewTextStyle.ColorAndOpacity);
+	bIsBackgroundImageSelectedAttrSet = true;
+	BackgroundImageSelectedAttribute.Assign(*this, &InNewTextStyle.BackgroundImageSelected);
 }
 
 void SEditableText::SetTextBlockStyle(const FTextBlockStyle* InTextStyle)
 {
 	if (InTextStyle)
 	{
-		EditableTextLayout->SetTextStyle(*InTextStyle);
-		Invalidate(EInvalidateWidgetReason::Layout); //Using Layout as changing text block size can affect the size.
+		SlateEditableTextLayout.SetTextStyle(*InTextStyle);
 	}
 }
 
 void SEditableText::Tick( const FGeometry& AllottedGeometry, const double InCurrentTime, const float InDeltaTime )
 {
-	EditableTextLayout->Tick(AllottedGeometry, InCurrentTime, InDeltaTime);
+	SlateEditableTextLayout.Tick(AllottedGeometry, InCurrentTime, InDeltaTime);
 }
 
 int32 SEditableText::OnPaint( const FPaintArgs& Args, const FGeometry& AllottedGeometry, const FSlateRect& MyCullingRect, FSlateWindowElementList& OutDrawElements, int32 LayerId, const FWidgetStyle& InWidgetStyle, bool bParentEnabled ) const
 {
-	const FTextBlockStyle& EditableTextStyle = EditableTextLayout->GetTextStyle();
+	const FTextBlockStyle& EditableTextStyle = SlateEditableTextLayout.GetTextStyle();
 	const FLinearColor ForegroundColor = EditableTextStyle.ColorAndOpacity.GetColor(InWidgetStyle);
 
 	FWidgetStyle TextWidgetStyle = FWidgetStyle(InWidgetStyle)
 		.SetForegroundColor(ForegroundColor);
 
-	LayerId = EditableTextLayout->OnPaint(Args, AllottedGeometry, MyCullingRect, OutDrawElements, LayerId, TextWidgetStyle, ShouldBeEnabled(bParentEnabled));
+	LayerId = SlateEditableTextLayout.OnPaint(Args, AllottedGeometry, MyCullingRect, OutDrawElements, LayerId, TextWidgetStyle, ShouldBeEnabled(bParentEnabled));
 
 	return LayerId;
 }
@@ -157,25 +189,25 @@ int32 SEditableText::OnPaint( const FPaintArgs& Args, const FGeometry& AllottedG
 void SEditableText::CacheDesiredSize(float LayoutScaleMultiplier)
 {
 	SynchronizeTextStyle();
-	EditableTextLayout->CacheDesiredSize(LayoutScaleMultiplier);
+	SlateEditableTextLayout.CacheDesiredSize(LayoutScaleMultiplier);
 	SWidget::CacheDesiredSize(LayoutScaleMultiplier);
 }
 
 FVector2D SEditableText::ComputeDesiredSize(float LayoutScaleMultiplier) const
 {
-	FVector2D TextLayoutSize = EditableTextLayout->ComputeDesiredSize(LayoutScaleMultiplier);
-	TextLayoutSize.X = FMath::Max(TextLayoutSize.X, MinDesiredWidth.Get());
+	FVector2D TextLayoutSize = SlateEditableTextLayout.ComputeDesiredSize(LayoutScaleMultiplier);
+	TextLayoutSize.X = FMath::Max(TextLayoutSize.X, MinDesiredWidthAttribute.Get());
 	return TextLayoutSize;
 }
 
 FChildren* SEditableText::GetChildren()
 {
-	return EditableTextLayout->GetChildren();
+	return SlateEditableTextLayout.GetChildren();
 }
 
 void SEditableText::OnArrangeChildren(const FGeometry& AllottedGeometry, FArrangedChildren& ArrangedChildren) const
 {
-	EditableTextLayout->OnArrangeChildren(AllottedGeometry, ArrangedChildren);
+	SlateEditableTextLayout.OnArrangeChildren(AllottedGeometry, ArrangedChildren);
 }
 
 FReply SEditableText::OnDragOver( const FGeometry& MyGeometry, const FDragDropEvent& DragDropEvent )
@@ -188,7 +220,7 @@ FReply SEditableText::OnDragOver( const FGeometry& MyGeometry, const FDragDropEv
 			return FReply::Handled();
 		}
 	}
-	
+
 	return FReply::Unhandled();
 }
 
@@ -199,11 +231,11 @@ FReply SEditableText::OnDrop( const FGeometry& MyGeometry, const FDragDropEvent&
 	{
 		if ( DragDropOp->HasText() )
 		{
-			EditableTextLayout->SetText(FText::FromString(DragDropOp->GetText()));
+			SlateEditableTextLayout.SetText(FText::FromString(DragDropOp->GetText()));
 			return FReply::Handled();
 		}
 	}
-	
+
 	return FReply::Unhandled();
 }
 
@@ -214,13 +246,13 @@ bool SEditableText::SupportsKeyboardFocus() const
 
 FReply SEditableText::OnFocusReceived( const FGeometry& MyGeometry, const FFocusEvent& InFocusEvent )
 {
-	EditableTextLayout->HandleFocusReceived(InFocusEvent);
+	SlateEditableTextLayout.HandleFocusReceived(InFocusEvent);
 	return FReply::Handled();
 }
 
 void SEditableText::OnFocusLost( const FFocusEvent& InFocusEvent )
 {
-	EditableTextLayout->HandleFocusLost(InFocusEvent);
+	SlateEditableTextLayout.HandleFocusLost(InFocusEvent);
 }
 
 FReply SEditableText::OnKeyChar( const FGeometry& MyGeometry, const FCharacterEvent& InCharacterEvent )
@@ -232,10 +264,10 @@ FReply SEditableText::OnKeyChar( const FGeometry& MyGeometry, const FCharacterEv
 	{
 		Reply = OnKeyCharHandler.Execute(MyGeometry, InCharacterEvent);
 	}
-	
+
 	if (!Reply.IsEventHandled())
 	{
-		Reply = EditableTextLayout->HandleKeyChar(InCharacterEvent);
+		Reply = SlateEditableTextLayout.HandleKeyChar(InCharacterEvent);
 	}
 
 	return Reply;
@@ -253,7 +285,7 @@ FReply SEditableText::OnKeyDown( const FGeometry& MyGeometry, const FKeyEvent& I
 
 	if (!Reply.IsEventHandled())
 	{
-		Reply = EditableTextLayout->HandleKeyDown(InKeyEvent);
+		Reply = SlateEditableTextLayout.HandleKeyDown(InKeyEvent);
 
 		if (!Reply.IsEventHandled())
 		{
@@ -266,27 +298,27 @@ FReply SEditableText::OnKeyDown( const FGeometry& MyGeometry, const FKeyEvent& I
 
 FReply SEditableText::OnKeyUp( const FGeometry& MyGeometry, const FKeyEvent& InKeyEvent )
 {
-	return EditableTextLayout->HandleKeyUp(InKeyEvent);
+	return SlateEditableTextLayout.HandleKeyUp(InKeyEvent);
 }
 
 FReply SEditableText::OnMouseButtonDown( const FGeometry& InMyGeometry, const FPointerEvent& InMouseEvent )
 {
-	return EditableTextLayout->HandleMouseButtonDown(InMyGeometry, InMouseEvent);
+	return SlateEditableTextLayout.HandleMouseButtonDown(InMyGeometry, InMouseEvent);
 }
 
 FReply SEditableText::OnMouseButtonUp( const FGeometry& InMyGeometry, const FPointerEvent& InMouseEvent )
 {
-	return EditableTextLayout->HandleMouseButtonUp(InMyGeometry, InMouseEvent);
+	return SlateEditableTextLayout.HandleMouseButtonUp(InMyGeometry, InMouseEvent);
 }
 
 FReply SEditableText::OnMouseMove( const FGeometry& InMyGeometry, const FPointerEvent& InMouseEvent )
 {
-	return EditableTextLayout->HandleMouseMove(InMyGeometry, InMouseEvent);
+	return SlateEditableTextLayout.HandleMouseMove(InMyGeometry, InMouseEvent);
 }
 
 FReply SEditableText::OnMouseButtonDoubleClick( const FGeometry& InMyGeometry, const FPointerEvent& InMouseEvent )
 {
-	return EditableTextLayout->HandleMouseButtonDoubleClick(InMyGeometry, InMouseEvent);
+	return SlateEditableTextLayout.HandleMouseButtonDoubleClick(InMyGeometry, InMouseEvent);
 }
 
 FCursorReply SEditableText::OnCursorQuery( const FGeometry& MyGeometry, const FPointerEvent& CursorEvent ) const
@@ -307,215 +339,212 @@ bool SEditableText::IsInteractable() const
 bool SEditableText::ComputeVolatility() const
 {
 	return SWidget::ComputeVolatility()
-		|| HasKeyboardFocus()
-		|| EditableTextLayout->ComputeVolatility()
-		|| Font.IsBound()
-		|| ColorAndOpacity.IsBound()
-		|| BackgroundImageSelected.IsBound()
-		|| bIsReadOnly.IsBound()
-		|| bIsPassword.IsBound()
-		|| MinDesiredWidth.IsBound();
+		|| HasKeyboardFocus();
 }
 
-void SEditableText::SetHintText( const TAttribute< FText >& InHintText )
+void SEditableText::SetHintText(TAttribute<FText> InHintText )
 {
-	EditableTextLayout->SetHintText(InHintText);
+	SlateEditableTextLayout.SetHintText(MoveTemp(InHintText));
 }
 
 FText SEditableText::GetHintText() const
 {
-	return EditableTextLayout->GetHintText();
+	return SlateEditableTextLayout.GetHintText();
 }
 
-void SEditableText::SetHintTextOpacity(const TAttribute<float>& InHintTextOpacity)
+void SEditableText::SetHintTextOpacity(TAttribute<float> InHintTextOpacity)
 {
-	EditableTextLayout->SetHintTextOpacity(InHintTextOpacity);
+	SlateEditableTextLayout.SetHintTextOpacity(MoveTemp(InHintTextOpacity));
 }
 
 float SEditableText::GetHintTextOpacity() const
 {
-	return EditableTextLayout->GetHintTextOpacity();
+	return SlateEditableTextLayout.GetHintTextOpacity();
 }
 
-void SEditableText::SetSearchText(const TAttribute<FText>& InSearchText)
+void SEditableText::SetSearchText(TAttribute<FText> InSearchText)
 {
-	EditableTextLayout->SetSearchText(InSearchText);
+	SlateEditableTextLayout.SetSearchText(MoveTemp(InSearchText));
 }
 
 FText SEditableText::GetSearchText() const
 {
-	return EditableTextLayout->GetSearchText();
+	return SlateEditableTextLayout.GetSearchText();
 }
 
-void SEditableText::SetIsReadOnly( TAttribute< bool > InIsReadOnly )
+void SEditableText::SetIsReadOnly(TAttribute<bool> InIsReadOnly)
 {
-	bIsReadOnly = InIsReadOnly;
+	bIsReadOnly = MoveTemp(InIsReadOnly);
 }
 
-void SEditableText::SetIsPassword( TAttribute< bool > InIsPassword )
+void SEditableText::SetIsPassword(TAttribute<bool> InIsPassword)
 {
-	bIsPassword = InIsPassword;
-	PlainTextMarshaller->SetIsPassword(bIsPassword);
+	PlainTextMarshaller->SetIsPassword(InIsPassword);
+	bIsPasswordAttribute.Assign(*this, MoveTemp(InIsPassword));
 }
 
 void SEditableText::SetColorAndOpacity(TAttribute<FSlateColor> Color)
 {
-	ColorAndOpacity = Color;
+	bIsColorAndOpacityAttrSet = Color.IsSet();
+	ColorAndOpacityAttribute.Assign(*this, MoveTemp(Color));
+}
+
+void SEditableText::SetBackgroundImageSelected(TAttribute<const FSlateBrush*> InBackgroundImageSelected)
+{
+	bIsBackgroundImageSelectedAttrSet = InBackgroundImageSelected.IsSet();
+	BackgroundImageSelectedAttribute.Assign(*this, MoveTemp(InBackgroundImageSelected));
 }
 
-void SEditableText::SetMinDesiredWidth(const TAttribute<float>& InMinDesiredWidth)
+void SEditableText::SetMinDesiredWidth(TAttribute<float> InMinDesiredWidth)
 {
-	MinDesiredWidth = InMinDesiredWidth;
+	MinDesiredWidthAttribute.Assign(*this, MoveTemp(InMinDesiredWidth));
 }
 
-void SEditableText::SetIsCaretMovedWhenGainFocus(const TAttribute<bool>& InIsCaretMovedWhenGainFocus)
+void SEditableText::SetIsCaretMovedWhenGainFocus(TAttribute<bool> InIsCaretMovedWhenGainFocus)
 {
-	bIsCaretMovedWhenGainFocus = InIsCaretMovedWhenGainFocus;
+	bIsCaretMovedWhenGainFocus = MoveTemp(InIsCaretMovedWhenGainFocus);
 }
 
-void SEditableText::SetSelectAllTextWhenFocused(const TAttribute<bool>& InSelectAllTextWhenFocused)
+void SEditableText::SetSelectAllTextWhenFocused(TAttribute<bool> InSelectAllTextWhenFocused)
 {
-	bSelectAllTextWhenFocused = InSelectAllTextWhenFocused;
+	bSelectAllTextWhenFocused = MoveTemp(InSelectAllTextWhenFocused);
 }
 
-void SEditableText::SetRevertTextOnEscape(const TAttribute<bool>& InRevertTextOnEscape)
+void SEditableText::SetRevertTextOnEscape(TAttribute<bool> InRevertTextOnEscape)
 {
-	bRevertTextOnEscape = InRevertTextOnEscape;
+	bRevertTextOnEscape = MoveTemp(InRevertTextOnEscape);
 }
 
-void SEditableText::SetClearKeyboardFocusOnCommit(const TAttribute<bool>& InClearKeyboardFocusOnCommit)
+void SEditableText::SetClearKeyboardFocusOnCommit(TAttribute<bool> InClearKeyboardFocusOnCommit)
 {
-	bClearKeyboardFocusOnCommit = InClearKeyboardFocusOnCommit;
+	bClearKeyboardFocusOnCommit = MoveTemp(InClearKeyboardFocusOnCommit);
 }
 
-void SEditableText::SetSelectAllTextOnCommit(const TAttribute<bool>& InSelectAllTextOnCommit)
+void SEditableText::SetSelectAllTextOnCommit(TAttribute<bool> InSelectAllTextOnCommit)
 {
-	bSelectAllTextOnCommit = InSelectAllTextOnCommit;
+	bSelectAllTextOnCommit = MoveTemp(InSelectAllTextOnCommit);
 }
 
-void SEditableText::SetSelectWordOnMouseDoubleClick(const TAttribute<bool>& InSelectWordOnMouseDoubleClick)
+void SEditableText::SetSelectWordOnMouseDoubleClick(TAttribute<bool> InSelectWordOnMouseDoubleClick)
 {
-	bSelectWordOnMouseDoubleClick = InSelectWordOnMouseDoubleClick;
+	bSelectWordOnMouseDoubleClick = MoveTemp(InSelectWordOnMouseDoubleClick);
 }
 
-void SEditableText::SetJustification(const TAttribute<ETextJustify::Type>& InJustification)
+void SEditableText::SetJustification(TAttribute<ETextJustify::Type> InJustification)
 {
-	EditableTextLayout->SetJustification(InJustification);
+	SlateEditableTextLayout.SetJustification(MoveTemp(InJustification));
 }
 
-void SEditableText::SetAllowContextMenu(const TAttribute< bool >& InAllowContextMenu)
+void SEditableText::SetAllowContextMenu(TAttribute<bool> InAllowContextMenu)
 {
-	bAllowContextMenu = InAllowContextMenu;
+	bAllowContextMenu = MoveTemp(InAllowContextMenu);
 }
 
-void SEditableText::SetEnableIntegratedKeyboard(const TAttribute<bool>& InEnableIntegratedKeyboard)
+void SEditableText::SetEnableIntegratedKeyboard(TAttribute<bool> InEnableIntegratedKeyboard)
 {
-	bEnableIntegratedKeyboard = InEnableIntegratedKeyboard;
+	bEnableIntegratedKeyboard = MoveTemp(InEnableIntegratedKeyboard);
 }
 
-void SEditableText::SetVirtualKeyboardDismissAction(TAttribute< EVirtualKeyboardDismissAction > InVirtualKeyboardDismissAction)
+void SEditableText::SetVirtualKeyboardDismissAction(TAttribute<EVirtualKeyboardDismissAction> InVirtualKeyboardDismissAction)
 {
-	VirtualKeyboardDismissAction = InVirtualKeyboardDismissAction;
+	VirtualKeyboardDismissAction = MoveTemp(InVirtualKeyboardDismissAction);
 }
 
 void SEditableText::SetTextShapingMethod(const TOptional<ETextShapingMethod>& InTextShapingMethod)
 {
-	EditableTextLayout->SetTextShapingMethod(InTextShapingMethod);
+	SlateEditableTextLayout.SetTextShapingMethod(InTextShapingMethod);
 }
 
 void SEditableText::SetTextFlowDirection(const TOptional<ETextFlowDirection>& InTextFlowDirection)
 {
-	EditableTextLayout->SetTextFlowDirection(InTextFlowDirection);
+	SlateEditableTextLayout.SetTextFlowDirection(InTextFlowDirection);
 }
 
-void SEditableText::SetOverflowPolicy(TOptional<ETextOverflowPolicy> InOverflowPolicy)
+void SEditableText::SetOverflowPolicy(const TOptional<ETextOverflowPolicy>& InOverflowPolicy)
 {
-	EditableTextLayout->SetOverflowPolicy(InOverflowPolicy);
+	SlateEditableTextLayout.SetOverflowPolicy(InOverflowPolicy);
 }
 
 void SEditableText::SetFontFacesLoadingPaintPolicy(EFontFacesLoadingPaintPolicy InFontFacesLoadingPaintPolicy)
 {
-	if (EditableTextLayout)
-	{
-		EditableTextLayout->SetFontFacesLoadingPaintPolicy(InFontFacesLoadingPaintPolicy);
-	}
+	SlateEditableTextLayout.SetFontFacesLoadingPaintPolicy(InFontFacesLoadingPaintPolicy);
 }
 
 bool SEditableText::AnyTextSelected() const
 {
-	return EditableTextLayout->AnyTextSelected();
+	return SlateEditableTextLayout.AnyTextSelected();
 }
 
 void SEditableText::SelectAllText()
 {
-	EditableTextLayout->SelectAllText();
+	SlateEditableTextLayout.SelectAllText();
 }
 
 void SEditableText::ClearSelection()
 {
-	EditableTextLayout->ClearSelection();
+	SlateEditableTextLayout.ClearSelection();
 }
 
 FText SEditableText::GetSelectedText() const
 {
-	return EditableTextLayout->GetSelectedText();
+	return SlateEditableTextLayout.GetSelectedText();
 }
 
 void SEditableText::GoTo(const FTextLocation& NewLocation)
 {
-	EditableTextLayout->GoTo(NewLocation);
+	SlateEditableTextLayout.GoTo(NewLocation);
 }
 
 void SEditableText::GoTo(const ETextLocation NewLocation)
 {
-	EditableTextLayout->GoTo(NewLocation);
+	SlateEditableTextLayout.GoTo(NewLocation);
 }
 
 void SEditableText::ScrollTo(const FTextLocation& NewLocation)
 {
-	EditableTextLayout->ScrollTo(NewLocation);
+	SlateEditableTextLayout.ScrollTo(NewLocation);
 }
 
 void SEditableText::ScrollTo(const ETextLocation NewLocation)
 {
-	EditableTextLayout->ScrollTo(NewLocation);
+	SlateEditableTextLayout.ScrollTo(NewLocation);
 }
 
 void SEditableText::BeginSearch(const FText& InSearchText, const ESearchCase::Type InSearchCase, const bool InReverse)
 {
-	EditableTextLayout->BeginSearch(InSearchText, InSearchCase, InReverse);
+	SlateEditableTextLayout.BeginSearch(InSearchText, InSearchCase, InReverse);
 }
 
 void SEditableText::AdvanceSearch(const bool InReverse)
 {
-	EditableTextLayout->AdvanceSearch(InReverse);
+	SlateEditableTextLayout.AdvanceSearch(InReverse);
 }
 
 void SEditableText::EnableTextInputMethodContext()
 {
-	EditableTextLayout->EnableTextInputMethodContext();
+	SlateEditableTextLayout.EnableTextInputMethodContext();
 }
 
 FTextSelection SEditableText::GetSelection() const
 {
-	return EditableTextLayout->GetSelection();
+	return SlateEditableTextLayout.GetSelection();
 }
 
 void SEditableText::SelectText(const FTextLocation& InSelectionStart, const FTextLocation& InCursorLocation)
 {
-	EditableTextLayout->SelectText(InSelectionStart, InCursorLocation);
+	SlateEditableTextLayout.SelectText(InSelectionStart, InCursorLocation);
 }
 
 void SEditableText::SynchronizeTextStyle()
 {
 	// Has the style used for this editable text changed?
 	bool bTextStyleChanged = false;
-	FTextBlockStyle NewTextStyle = EditableTextLayout->GetTextStyle();
+	FTextBlockStyle NewTextStyle = SlateEditableTextLayout.GetTextStyle();
 
 	// Sync from the font override
-	if (Font.IsSet())
+	if (bIsFontOverrideAttrSet)
 	{
-		const FSlateFontInfo& NewFontInfo = Font.Get();
+		const FSlateFontInfo& NewFontInfo = FontOverrideAttribute.Get();
 		if (!NewTextStyle.Font.IsIdenticalTo(NewFontInfo))
 		{
 			NewTextStyle.Font = NewFontInfo;
@@ -524,9 +553,9 @@ void SEditableText::SynchronizeTextStyle()
 	}
 
 	// Sync from the color override
-	if (ColorAndOpacity.IsSet())
+	if (bIsColorAndOpacityAttrSet)
 	{
-		const FSlateColor& NewColorAndOpacity = ColorAndOpacity.Get();
+		const FSlateColor& NewColorAndOpacity = ColorAndOpacityAttribute.Get();
 		if (NewTextStyle.ColorAndOpacity != NewColorAndOpacity)
 		{
 			NewTextStyle.ColorAndOpacity = NewColorAndOpacity;
@@ -535,9 +564,9 @@ void SEditableText::SynchronizeTextStyle()
 	}
 
 	// Sync from the highlight shape override
-	if (BackgroundImageSelected.IsSet())
+	if (bIsBackgroundImageSelectedAttrSet)
 	{
-		const FSlateBrush* NewSelectionBrush = BackgroundImageSelected.Get();
+		const FSlateBrush* NewSelectionBrush = BackgroundImageSelectedAttribute.Get();
 		if (NewSelectionBrush && NewTextStyle.HighlightShape != *NewSelectionBrush)
 		{
 			NewTextStyle.HighlightShape = *NewSelectionBrush;
@@ -547,8 +576,8 @@ void SEditableText::SynchronizeTextStyle()
 
 	if (bTextStyleChanged)
 	{
-		EditableTextLayout->SetTextStyle(NewTextStyle);
-		EditableTextLayout->ForceRefreshTextLayout(EditableTextLayout->GetEditableText());
+		SlateEditableTextLayout.SetTextStyle(NewTextStyle);
+		SlateEditableTextLayout.ForceRefreshTextLayout(SlateEditableTextLayout.GetEditableText());
 	}
 }
 
@@ -559,7 +588,7 @@ bool SEditableText::IsTextReadOnly() const
 
 bool SEditableText::IsTextPassword() const
 {
-	return bIsPassword.Get(false);
+	return bIsPasswordAttribute.Get();
 }
 
 bool SEditableText::IsMultiLineTextEdit() const
@@ -633,7 +662,7 @@ void SEditableText::EnsureActiveTick()
 	auto DoActiveTick = [this](double InCurrentTime, float InDeltaTime) -> EActiveTimerReturnType
 	{
 		// Continue if we still have focus, otherwise treat as a fire-and-forget Tick() request
-		const bool bShouldAppearFocused = HasKeyboardFocus() || EditableTextLayout->HasActiveContextMenu();
+		const bool bShouldAppearFocused = HasKeyboardFocus() || SlateEditableTextLayout.HasActiveContextMenu();
 		return (bShouldAppearFocused) ? EActiveTimerReturnType::Continue : EActiveTimerReturnType::Stop;
 	};
 
@@ -686,8 +715,8 @@ TSharedPtr<SWidget> SEditableText::BuildContextMenuContent() const
 	{
 		return OnContextMenuOpening.Execute();
 	}
-	
-	return EditableTextLayout->BuildDefaultContextMenu(MenuExtender);
+
+	return SlateEditableTextLayout.BuildDefaultContextMenu(MenuExtender);
 }
 
 void SEditableText::OnBeginTextEdit(const FText& InText)
@@ -708,21 +737,41 @@ void SEditableText::OnTextCommitted(const FText& InText, const ETextCommit::Type
 void SEditableText::OnCursorMoved(const FTextLocation& InLocation)
 {
 	FTextSelection TextSelection = GetSelection();
-	
+
 	OnCursorMovedWithSelectionCallback.ExecuteIfBound(InLocation, TextSelection);
 	Invalidate(EInvalidateWidgetReason::Layout);
 }
 
-float SEditableText::UpdateAndClampHorizontalScrollBar(const float InViewOffset, const float InViewFraction, const EVisibility InVisiblityOverride)
+float SEditableText::UpdateAndClampHorizontalScrollBar(const float InViewOffset, const float InViewFraction, const EVisibility InVisibilityOverride)
 {
-	return EditableTextLayout->GetScrollOffset().X;
+	return SlateEditableTextLayout.GetScrollOffset().X;
 }
 
-float SEditableText::UpdateAndClampVerticalScrollBar(const float InViewOffset, const float InViewFraction, const EVisibility InVisiblityOverride)
+float SEditableText::UpdateAndClampVerticalScrollBar(const float InViewOffset, const float InViewFraction, const EVisibility InVisibilityOverride)
 {
 	return 0.0f;
 }
 
+ETextJustify::Type SEditableText::GetJustification() const
+{
+	return SlateEditableTextLayout.GetJustification();
+}
+
+FSlateColor SEditableText::GetColorAndOpacity() const
+{
+	return ColorAndOpacityAttribute.Get();
+}
+
+const FSlateBrush* SEditableText::GetBackgroundImageSelected() const
+{
+	return BackgroundImageSelectedAttribute.Get();
+}
+
+float SEditableText::GetMinDesiredWidth() const
+{
+	return MinDesiredWidthAttribute.Get();
+}
+
 #if WITH_ACCESSIBILITY
 TSharedRef<FSlateAccessibleWidget> SEditableText::CreateAccessibleWidget()
 {
@@ -737,8 +786,5 @@ TOptional<FText> SEditableText::GetDefaultAccessibleText(EAccessibleType Accessi
 
 void SEditableText::ToggleVirtualKeyboard(bool bShow, const int32 UserIndex)
 {
-	if (EditableTextLayout.IsValid())
-	{
-		EditableTextLayout->ToggleVirtualKeyboard(bShow, UserIndex);
-	}
+	SlateEditableTextLayout.ToggleVirtualKeyboard(bShow, UserIndex);
 }
diff --git a/Engine/Source/Runtime/Slate/Private/Widgets/Text/SMultiLineEditableText.cpp b/Engine/Source/Runtime/Slate/Private/Widgets/Text/SMultiLineEditableText.cpp
--- a/Engine/Source/Runtime/Slate/Private/Widgets/Text/SMultiLineEditableText.cpp
+++ b/Engine/Source/Runtime/Slate/Private/Widgets/Text/SMultiLineEditableText.cpp
@@ -10,22 +10,25 @@
 #include "Framework/Text/TextEditHelper.h"
 #include "Framework/Text/PlainTextLayoutMarshaller.h"
 #include "Widgets/Text/SlateEditableTextLayout.h"
+#include "Widgets/Text/SlateTextBlockLayout.h"
 #include "Types/ReflectionMetadata.h"
 #include "Types/TrackedMetaData.h"
 
+SLATE_IMPLEMENT_WIDGET(SMultiLineEditableText)
+void SMultiLineEditableText::PrivateRegisterAttributes(FSlateAttributeInitializer& AttributeInitializer)
+{
+	FSlateEditableTextLayout::RegisterNestedAttributes<&SMultiLineEditableText::SlateEditableTextLayout>(AttributeInitializer);
+}
+
 SMultiLineEditableText::SMultiLineEditableText()
-	: bSelectAllTextWhenFocused(false)
+	: SlateEditableTextLayout(*this)
+	, bSelectAllTextWhenFocused(false)
 	, bIsReadOnly(false)
 	, AmountScrolledWhileRightMouseDown(0.0f)
 	, bIsSoftwareCursor(false)
 {
 }
 
-SMultiLineEditableText::~SMultiLineEditableText()
-{
-	// Needed to avoid "deletion of pointer to incomplete type 'FSlateEditableTextLayout'; no destructor called" error when using TUniquePtr
-}
-
 void SMultiLineEditableText::Construct( const FArguments& InArgs )
 {
 	bIsReadOnly = InArgs._IsReadOnly;
@@ -82,19 +85,19 @@ void SMultiLineEditableText::Construct( const FArguments& InArgs )
 		Marshaller = FPlainTextLayoutMarshaller::Create();
 	}
 
-	EditableTextLayout = MakeUnique<FSlateEditableTextLayout>(*this, InArgs._Text, TextStyle, InArgs._TextShapingMethod, InArgs._TextFlowDirection, InArgs._CreateSlateTextLayout, Marshaller.ToSharedRef(), Marshaller.ToSharedRef());
-	EditableTextLayout->SetHintText(InArgs._HintText);
-	EditableTextLayout->SetSearchText(InArgs._SearchText);
-	EditableTextLayout->SetTextWrapping(InArgs._WrapTextAt, InArgs._AutoWrapText, InArgs._WrappingPolicy);
-	EditableTextLayout->SetMargin(InArgs._Margin);
-	EditableTextLayout->SetJustification(InArgs._Justification);
-	EditableTextLayout->SetLineHeightPercentage(InArgs._LineHeightPercentage);
-	EditableTextLayout->SetDebugSourceInfo(TAttribute<FString>::Create(TAttribute<FString>::FGetter::CreateLambda([this]{ return FReflectionMetaData::GetWidgetDebugInfo(this); })));
-	EditableTextLayout->SetOverflowPolicy(InArgs._OverflowPolicy);
-	EditableTextLayout->SetFontFacesLoadingPaintPolicy(InArgs._FontFacesLoadingPaintPolicy);
+	SlateEditableTextLayout.Construct(InArgs._Text, TextStyle, InArgs._TextShapingMethod, InArgs._TextFlowDirection, InArgs._CreateSlateTextLayout, Marshaller.ToSharedRef(), Marshaller.ToSharedRef());
+	SlateEditableTextLayout.SetHintText(InArgs._HintText);
+	SlateEditableTextLayout.SetSearchText(InArgs._SearchText);
+	SlateEditableTextLayout.SetTextWrapping(InArgs._WrapTextAt, InArgs._AutoWrapText, InArgs._WrappingPolicy);
+	SlateEditableTextLayout.SetMargin(InArgs._Margin);
+	SlateEditableTextLayout.SetJustification(InArgs._Justification);
+	SlateEditableTextLayout.SetLineHeightPercentage(InArgs._LineHeightPercentage);
+	SlateEditableTextLayout.SetDebugSourceInfo(TAttribute<FString>::Create(TAttribute<FString>::FGetter::CreateLambda([this]{ return FReflectionMetaData::GetWidgetDebugInfo(this); })));
+	SlateEditableTextLayout.SetOverflowPolicy(InArgs._OverflowPolicy);
+	SlateEditableTextLayout.SetFontFacesLoadingPaintPolicy(InArgs._FontFacesLoadingPaintPolicy);
 	if (InArgs._OnAllFontFacesFinishLoading.IsBound())
 	{
-		EditableTextLayout->SetOnAllFontFacesFinishLoadingDelegate(InArgs._OnAllFontFacesFinishLoading);
+		SlateEditableTextLayout.SetOnAllFontFacesFinishLoadingDelegate(InArgs._OnAllFontFacesFinishLoading);
 	}
 
 	// build context menu extender
@@ -106,139 +109,139 @@ void SMultiLineEditableText::Construct( const FArguments& InArgs )
 
 void SMultiLineEditableText::GetCurrentTextLine(FString& OutTextLine) const
 {
-	EditableTextLayout->GetCurrentTextLine(OutTextLine);
+	SlateEditableTextLayout.GetCurrentTextLine(OutTextLine);
 }
 
 void SMultiLineEditableText::GetTextLine(const int32 InLineIndex, FString& OutTextLine) const
 {
-	EditableTextLayout->GetTextLine(InLineIndex, OutTextLine);
+	SlateEditableTextLayout.GetTextLine(InLineIndex, OutTextLine);
 }
 
 void SMultiLineEditableText::SetText(const TAttribute< FText >& InText)
 {
-	EditableTextLayout->SetText(InText);
+	SlateEditableTextLayout.SetText(InText);
 }
 
-int32 SMultiLineEditableText::GetTextLineCount()
+int32 SMultiLineEditableText::GetTextLineCount() const
 {
-	return EditableTextLayout->GetTextLineCount();
+	return SlateEditableTextLayout.GetTextLineCount();
 }
 
 FText SMultiLineEditableText::GetText() const
 {
-	return EditableTextLayout->GetText();
+	return SlateEditableTextLayout.GetText();
 }
 
 FText SMultiLineEditableText::GetPlainText() const
 {
-	return EditableTextLayout->GetPlainText();
+	return SlateEditableTextLayout.GetPlainText();
 }
 
 void SMultiLineEditableText::SetHintText(const TAttribute< FText >& InHintText)
 {
-	EditableTextLayout->SetHintText(InHintText);
+	SlateEditableTextLayout.SetHintText(InHintText);
 }
 
 FText SMultiLineEditableText::GetHintText() const
 {
-	return EditableTextLayout->GetHintText();
+	return SlateEditableTextLayout.GetHintText();
 }
 
 void SMultiLineEditableText::SetSearchText(const TAttribute<FText>& InSearchText)
 {
-	EditableTextLayout->SetSearchText(InSearchText);
+	SlateEditableTextLayout.SetSearchText(InSearchText);
 }
 
 FText SMultiLineEditableText::GetSearchText() const
 {
-	return EditableTextLayout->GetSearchText();
+	return SlateEditableTextLayout.GetSearchText();
 }
 
 int32 SMultiLineEditableText::GetSearchResultIndex() const
 {
-	return EditableTextLayout->GetSearchResultIndex();
+	return SlateEditableTextLayout.GetSearchResultIndex();
 }
 
 int32 SMultiLineEditableText::GetNumSearchResults() const
 {
-	return EditableTextLayout->GetNumSearchResults();
+	return SlateEditableTextLayout.GetNumSearchResults();
 }
 
 void SMultiLineEditableText::SetTextStyle(const FTextBlockStyle* InTextStyle)
 {
 	if (InTextStyle)
 	{
-		EditableTextLayout->SetTextStyle(*InTextStyle);
+		SlateEditableTextLayout.SetTextStyle(*InTextStyle);
 	}
 	else
 	{
 		FArguments Defaults;
 		check(Defaults._TextStyle);
-		EditableTextLayout->SetTextStyle(*Defaults._TextStyle);
+		SlateEditableTextLayout.SetTextStyle(*Defaults._TextStyle);
 	}
 }
 
 void SMultiLineEditableText::SetFont(const TAttribute< FSlateFontInfo >& InNewFont)
 {
-	FTextBlockStyle TextStyle = EditableTextLayout->GetTextStyle();
+	FTextBlockStyle TextStyle = SlateEditableTextLayout.GetTextStyle();
 	TextStyle.SetFont(InNewFont.Get());
-	EditableTextLayout->SetTextStyle(TextStyle);
+	SlateEditableTextLayout.SetTextStyle(TextStyle);
 }
 
 FSlateFontInfo SMultiLineEditableText::GetFont() const
 {
-	FTextBlockStyle TextStyle = EditableTextLayout->GetTextStyle();
+	FTextBlockStyle TextStyle = SlateEditableTextLayout.GetTextStyle();
 	return TextStyle.Font;
 }
 
 void SMultiLineEditableText::SetTextShapingMethod(const TOptional<ETextShapingMethod>& InTextShapingMethod)
 {
-	EditableTextLayout->SetTextShapingMethod(InTextShapingMethod);
+	SlateEditableTextLayout.SetTextShapingMethod(InTextShapingMethod);
 }
 
 void SMultiLineEditableText::SetTextFlowDirection(const TOptional<ETextFlowDirection>& InTextFlowDirection)
 {
-	EditableTextLayout->SetTextFlowDirection(InTextFlowDirection);
+	SlateEditableTextLayout.SetTextFlowDirection(InTextFlowDirection);
 }
 
 void SMultiLineEditableText::SetWrapTextAt(const TAttribute<float>& InWrapTextAt)
 {
-	EditableTextLayout->SetWrapTextAt(InWrapTextAt);
+	SlateEditableTextLayout.SetWrapTextAt(InWrapTextAt);
 }
 
 void SMultiLineEditableText::SetAutoWrapText(const TAttribute<bool>& InAutoWrapText)
 {
-	EditableTextLayout->SetAutoWrapText(InAutoWrapText);
+	SlateEditableTextLayout.SetAutoWrapText(InAutoWrapText);
 }
 
 void SMultiLineEditableText::SetWrappingPolicy(const TAttribute<ETextWrappingPolicy>& InWrappingPolicy)
 {
-	EditableTextLayout->SetWrappingPolicy(InWrappingPolicy);
+	SlateEditableTextLayout.SetWrappingPolicy(InWrappingPolicy);
 }
 
 void SMultiLineEditableText::SetLineHeightPercentage(const TAttribute<float>& InLineHeightPercentage)
 {
-	EditableTextLayout->SetLineHeightPercentage(InLineHeightPercentage);
+	SlateEditableTextLayout.SetLineHeightPercentage(InLineHeightPercentage);
 }
 
 void SMultiLineEditableText::SetApplyLineHeightToBottomLine(const TAttribute<bool>& InApplyLineHeightToBottomLine)
 {
-	EditableTextLayout->SetApplyLineHeightToBottomLine(InApplyLineHeightToBottomLine);
+	SlateEditableTextLayout.SetApplyLineHeightToBottomLine(InApplyLineHeightToBottomLine);
 }
 
 void SMultiLineEditableText::SetMargin(const TAttribute<FMargin>& InMargin)
 {
-	EditableTextLayout->SetMargin(InMargin);
+	SlateEditableTextLayout.SetMargin(InMargin);
 }
 
 void SMultiLineEditableText::SetJustification(const TAttribute<ETextJustify::Type>& InJustification)
 {
-	EditableTextLayout->SetJustification(InJustification);
+	SlateEditableTextLayout.SetJustification(InJustification);
 }
 
 void SMultiLineEditableText::SetOverflowPolicy(TOptional<ETextOverflowPolicy> InOverflowPolicy)
 {
-	EditableTextLayout->SetOverflowPolicy(InOverflowPolicy);
+	SlateEditableTextLayout.SetOverflowPolicy(InOverflowPolicy);
 }
 
 void SMultiLineEditableText::SetAllowContextMenu(const TAttribute< bool >& InAllowContextMenu)
@@ -283,21 +286,18 @@ void SMultiLineEditableText::SetClearKeyboardFocusOnCommit(const TAttribute<bool
 
 void SMultiLineEditableText::SetFontFacesLoadingPaintPolicy(EFontFacesLoadingPaintPolicy InFontFacesLoadingPaintPolicy)
 {
-	if (EditableTextLayout)
-	{
-		EditableTextLayout->SetFontFacesLoadingPaintPolicy(InFontFacesLoadingPaintPolicy);
-	}
+	SlateEditableTextLayout.SetFontFacesLoadingPaintPolicy(InFontFacesLoadingPaintPolicy);
 }
 
 void SMultiLineEditableText::OnHScrollBarMoved(const float InScrollOffsetFraction)
 {
-	EditableTextLayout->SetHorizontalScrollFraction(InScrollOffsetFraction);
+	SlateEditableTextLayout.SetHorizontalScrollFraction(InScrollOffsetFraction);
 	OnHScrollBarUserScrolled.ExecuteIfBound(InScrollOffsetFraction);
 }
 
 void SMultiLineEditableText::OnVScrollBarMoved(const float InScrollOffsetFraction)
 {
-	EditableTextLayout->SetVerticalScrollFraction(InScrollOffsetFraction);
+	SlateEditableTextLayout.SetVerticalScrollFraction(InScrollOffsetFraction);
 	OnVScrollBarUserScrolled.ExecuteIfBound(InScrollOffsetFraction);
 }
 
@@ -382,7 +382,7 @@ void SMultiLineEditableText::EnsureActiveTick()
 	auto DoActiveTick = [this](double InCurrentTime, float InDeltaTime) -> EActiveTimerReturnType
 	{
 		// Continue if we still have focus, otherwise treat as a fire-and-forget Tick() request
-		const bool bShouldAppearFocused = HasKeyboardFocus() || EditableTextLayout->HasActiveContextMenu();
+		const bool bShouldAppearFocused = HasKeyboardFocus() || SlateEditableTextLayout.HasActiveContextMenu();
 		return (bShouldAppearFocused) ? EActiveTimerReturnType::Continue : EActiveTimerReturnType::Stop;
 	};
 
@@ -436,7 +436,7 @@ TSharedPtr<SWidget> SMultiLineEditableText::BuildContextMenuContent() const
 		return OnContextMenuOpening.Execute();
 	}
 
-	return EditableTextLayout->BuildDefaultContextMenu(MenuExtender);
+	return SlateEditableTextLayout.BuildDefaultContextMenu(MenuExtender);
 }
 
 void SMultiLineEditableText::OnBeginTextEdit(const FText& InText)
@@ -473,7 +473,7 @@ float SMultiLineEditableText::UpdateAndClampHorizontalScrollBar(const float InVi
 		}
 	}
 
-	return EditableTextLayout->GetScrollOffset().X;
+	return SlateEditableTextLayout.GetScrollOffset().X;
 }
 
 float SMultiLineEditableText::UpdateAndClampVerticalScrollBar(const float InViewOffset, const float InViewFraction, const EVisibility InVisiblityOverride)
@@ -489,134 +489,134 @@ float SMultiLineEditableText::UpdateAndClampVerticalScrollBar(const float InView
 		}
 	}
 
-	return EditableTextLayout->GetScrollOffset().Y;
+	return SlateEditableTextLayout.GetScrollOffset().Y;
 }
 
 void SMultiLineEditableText::BeginEditTransaction()
 {
-	EditableTextLayout->BeginEditTransation();
+	SlateEditableTextLayout.BeginEditTransaction();
 }
 
 void SMultiLineEditableText::EndEditTransaction()
 {
-	EditableTextLayout->EndEditTransaction();
+	SlateEditableTextLayout.EndEditTransaction();
 }
 
 FReply SMultiLineEditableText::OnFocusReceived( const FGeometry& MyGeometry, const FFocusEvent& InFocusEvent )
 {
-	EditableTextLayout->HandleFocusReceived(InFocusEvent);
+	SlateEditableTextLayout.HandleFocusReceived(InFocusEvent);
 	return FReply::Handled();
 }
 
 void SMultiLineEditableText::OnFocusLost( const FFocusEvent& InFocusEvent )
 {
 	bIsSoftwareCursor = false;
-	EditableTextLayout->HandleFocusLost(InFocusEvent);
+	SlateEditableTextLayout.HandleFocusLost(InFocusEvent);
 }
 
 bool SMultiLineEditableText::AnyTextSelected() const
 {
-	return EditableTextLayout->AnyTextSelected();
+	return SlateEditableTextLayout.AnyTextSelected();
 }
 
 void SMultiLineEditableText::SelectAllText()
 {
-	EditableTextLayout->SelectAllText();
+	SlateEditableTextLayout.SelectAllText();
 }
 
 void SMultiLineEditableText::SelectText(const FTextLocation& InSelectionStart, const FTextLocation& InCursorLocation)
 {
-	EditableTextLayout->SelectText(InSelectionStart, InCursorLocation);
+	SlateEditableTextLayout.SelectText(InSelectionStart, InCursorLocation);
 }
 
 void SMultiLineEditableText::ClearSelection()
 {
-	EditableTextLayout->ClearSelection();
+	SlateEditableTextLayout.ClearSelection();
 }
 
 FText SMultiLineEditableText::GetSelectedText() const
 {
-	return EditableTextLayout->GetSelectedText();
+	return SlateEditableTextLayout.GetSelectedText();
 }
 
 FTextSelection SMultiLineEditableText::GetSelection() const
 {
-	return EditableTextLayout->GetSelection();
+	return SlateEditableTextLayout.GetSelection();
 }
 
 void SMultiLineEditableText::DeleteSelectedText()
 {
-	EditableTextLayout->DeleteSelectedText();
+	SlateEditableTextLayout.DeleteSelectedText();
 }
 
 void SMultiLineEditableText::InsertTextAtCursor(const FText& InText)
 {
-	EditableTextLayout->InsertTextAtCursor(InText.ToString());
+	SlateEditableTextLayout.InsertTextAtCursor(InText.ToString());
 }
 
 void SMultiLineEditableText::InsertTextAtCursor(const FString& InString)
 {
-	EditableTextLayout->InsertTextAtCursor(InString);
+	SlateEditableTextLayout.InsertTextAtCursor(InString);
 }
 
 void SMultiLineEditableText::InsertRunAtCursor(TSharedRef<IRun> InRun)
 {
-	EditableTextLayout->InsertRunAtCursor(InRun);
+	SlateEditableTextLayout.InsertRunAtCursor(InRun);
 }
 
 void SMultiLineEditableText::GoTo(const FTextLocation& NewLocation)
 {
-	EditableTextLayout->GoTo(NewLocation);
+	SlateEditableTextLayout.GoTo(NewLocation);
 }
 
 void SMultiLineEditableText::GoTo(const ETextLocation NewLocation)
 {
-	EditableTextLayout->GoTo(NewLocation);
+	SlateEditableTextLayout.GoTo(NewLocation);
 }
 
 void SMultiLineEditableText::ScrollTo(const FTextLocation& NewLocation)
 {
-	EditableTextLayout->ScrollTo(NewLocation);
+	SlateEditableTextLayout.ScrollTo(NewLocation);
 }
 
 void SMultiLineEditableText::ScrollTo(const ETextLocation NewLocation)
 {
-	EditableTextLayout->ScrollTo(NewLocation);
+	SlateEditableTextLayout.ScrollTo(NewLocation);
 }
 
 void SMultiLineEditableText::ApplyToSelection(const FRunInfo& InRunInfo, const FTextBlockStyle& InStyle)
 {
-	EditableTextLayout->ApplyToSelection(InRunInfo, InStyle);
+	SlateEditableTextLayout.ApplyToSelection(InRunInfo, InStyle);
 }
 
 void SMultiLineEditableText::BeginSearch(const FText& InSearchText, const ESearchCase::Type InSearchCase, const bool InReverse)
 {
-	EditableTextLayout->BeginSearch(InSearchText, InSearchCase, InReverse);
+	SlateEditableTextLayout.BeginSearch(InSearchText, InSearchCase, InReverse);
 }
 
 void SMultiLineEditableText::AdvanceSearch(const bool InReverse)
 {
-	EditableTextLayout->AdvanceSearch(InReverse);
+	SlateEditableTextLayout.AdvanceSearch(InReverse);
 }
 
 TSharedPtr<const IRun> SMultiLineEditableText::GetRunUnderCursor() const
 {
-	return EditableTextLayout->GetRunUnderCursor();
+	return SlateEditableTextLayout.GetRunUnderCursor();
 }
 
 TArray<TSharedRef<const IRun>> SMultiLineEditableText::GetSelectedRuns() const
 {
-	return EditableTextLayout->GetSelectedRuns();
+	return SlateEditableTextLayout.GetSelectedRuns();
 }
 
 FTextLocation SMultiLineEditableText::GetCursorLocation() const
 {
-	return EditableTextLayout->GetCursorLocation();
+	return SlateEditableTextLayout.GetCursorLocation();
 }
 
 TCHAR SMultiLineEditableText::GetCharacterAt(const FTextLocation& Location) const
 {
-	return EditableTextLayout->GetCharacterAt(Location);
+	return SlateEditableTextLayout.GetCharacterAt(Location);
 }
 
 TSharedPtr<const SScrollBar> SMultiLineEditableText::GetHScrollBar() const
@@ -631,7 +631,7 @@ TSharedPtr<const SScrollBar> SMultiLineEditableText::GetVScrollBar() const
 
 void SMultiLineEditableText::Refresh()
 {
-	EditableTextLayout->Refresh();
+	SlateEditableTextLayout.Refresh();
 }
 
 void SMultiLineEditableText::ForceScroll(int32 UserIndex, float ScrollAxisMagnitude)
@@ -645,25 +645,25 @@ void SMultiLineEditableText::ForceScroll(int32 UserIndex, float ScrollAxisMagnit
 
 void SMultiLineEditableText::Tick( const FGeometry& AllottedGeometry, const double InCurrentTime, const float InDeltaTime )
 {
-	EditableTextLayout->Tick(AllottedGeometry, InCurrentTime, InDeltaTime);
+	SlateEditableTextLayout.Tick(AllottedGeometry, InCurrentTime, InDeltaTime);
 }
 
 int32 SMultiLineEditableText::OnPaint( const FPaintArgs& Args, const FGeometry& AllottedGeometry, const FSlateRect& MyCullingRect, FSlateWindowElementList& OutDrawElements, int32 LayerId, const FWidgetStyle& InWidgetStyle, bool bParentEnabled ) const
 {
-	const FTextBlockStyle& EditableTextStyle = EditableTextLayout->GetTextStyle();
+	const FTextBlockStyle& EditableTextStyle = SlateEditableTextLayout.GetTextStyle();
 	const FLinearColor ForegroundColor = EditableTextStyle.ColorAndOpacity.GetColor(InWidgetStyle);
 
 	FWidgetStyle TextWidgetStyle = FWidgetStyle(InWidgetStyle)
 		.SetForegroundColor(ForegroundColor);
 
-	const bool bAutoWrap = EditableTextLayout->GetAutoWrapText();
-	if (bAutoWrap && EditableTextLayout->GetComputedWrappingWidth() == 0)
+	const bool bAutoWrap = SlateEditableTextLayout.GetAutoWrapText();
+	if (bAutoWrap && SlateEditableTextLayout.GetComputedWrappingWidth() == 0)
 	{
 		QUICK_SCOPE_CYCLE_COUNTER(STAT_Slate_AutoWrapRecompute)
 		const_cast<SMultiLineEditableText*>(this)->CacheDesiredSize(GetPrepassLayoutScaleMultiplier());
 	}
 
-	LayerId = EditableTextLayout->OnPaint(Args, AllottedGeometry, MyCullingRect, OutDrawElements, LayerId, TextWidgetStyle, ShouldBeEnabled(bParentEnabled));
+	LayerId = SlateEditableTextLayout.OnPaint(Args, AllottedGeometry, MyCullingRect, OutDrawElements, LayerId, TextWidgetStyle, ShouldBeEnabled(bParentEnabled));
 
 	if (bIsSoftwareCursor)
 	{
@@ -683,23 +683,23 @@ int32 SMultiLineEditableText::OnPaint( const FPaintArgs& Args, const FGeometry&
 
 void SMultiLineEditableText::CacheDesiredSize(float LayoutScaleMultiplier)
 {
-	EditableTextLayout->CacheDesiredSize(LayoutScaleMultiplier);
+	SlateEditableTextLayout.CacheDesiredSize(LayoutScaleMultiplier);
 	SWidget::CacheDesiredSize(LayoutScaleMultiplier);
 }
 
 FVector2D SMultiLineEditableText::ComputeDesiredSize( float LayoutScaleMultiplier ) const
 {
-	return EditableTextLayout->ComputeDesiredSize(LayoutScaleMultiplier);
+	return SlateEditableTextLayout.ComputeDesiredSize(LayoutScaleMultiplier);
 }
 
 FChildren* SMultiLineEditableText::GetChildren()
 {
-	return EditableTextLayout->GetChildren();
+	return SlateEditableTextLayout.GetChildren();
 }
 
 void SMultiLineEditableText::OnArrangeChildren( const FGeometry& AllottedGeometry, FArrangedChildren& ArrangedChildren ) const
 {
-	EditableTextLayout->OnArrangeChildren(AllottedGeometry, ArrangedChildren);
+	SlateEditableTextLayout.OnArrangeChildren(AllottedGeometry, ArrangedChildren);
 }
 
 bool SMultiLineEditableText::SupportsKeyboardFocus() const
@@ -719,7 +719,7 @@ FReply SMultiLineEditableText::OnKeyChar( const FGeometry& MyGeometry,const FCha
 
 	if (!Reply.IsEventHandled())
 	{
-		Reply = EditableTextLayout->HandleKeyChar(InCharacterEvent);
+		Reply = SlateEditableTextLayout.HandleKeyChar(InCharacterEvent);
 	}
 
 	return Reply;
@@ -737,7 +737,7 @@ FReply SMultiLineEditableText::OnKeyDown( const FGeometry& MyGeometry, const FKe
 
 	if (!Reply.IsEventHandled())
 	{
-		Reply = EditableTextLayout->HandleKeyDown(InKeyEvent);
+		Reply = SlateEditableTextLayout.HandleKeyDown(InKeyEvent);
 	}
 
 	return Reply;
@@ -745,20 +745,20 @@ FReply SMultiLineEditableText::OnKeyDown( const FGeometry& MyGeometry, const FKe
 
 FReply SMultiLineEditableText::OnKeyUp( const FGeometry& MyGeometry, const FKeyEvent& InKeyEvent )
 {
-	return EditableTextLayout->HandleKeyUp(InKeyEvent);
+	return SlateEditableTextLayout.HandleKeyUp(InKeyEvent);
 }
 
-FReply SMultiLineEditableText::OnMouseButtonDown( const FGeometry& MyGeometry, const FPointerEvent& MouseEvent ) 
+FReply SMultiLineEditableText::OnMouseButtonDown( const FGeometry& MyGeometry, const FPointerEvent& MouseEvent )
 {
 	if (MouseEvent.GetEffectingButton() == EKeys::RightMouseButton)
 	{
 		AmountScrolledWhileRightMouseDown = 0.0f;
 	}
 
-	return EditableTextLayout->HandleMouseButtonDown(MyGeometry, MouseEvent);
+	return SlateEditableTextLayout.HandleMouseButtonDown(MyGeometry, MouseEvent);
 }
 
-FReply SMultiLineEditableText::OnMouseButtonUp( const FGeometry& MyGeometry, const FPointerEvent& MouseEvent ) 
+FReply SMultiLineEditableText::OnMouseButtonUp( const FGeometry& MyGeometry, const FPointerEvent& MouseEvent )
 {
 	if (MouseEvent.GetEffectingButton() == EKeys::RightMouseButton)
 	{
@@ -774,7 +774,7 @@ FReply SMultiLineEditableText::OnMouseButtonUp( const FGeometry& MyGeometry, con
 		}
 	}
 
-	return EditableTextLayout->HandleMouseButtonUp(MyGeometry, MouseEvent);
+	return SlateEditableTextLayout.HandleMouseButtonUp(MyGeometry, MouseEvent);
 }
 
 FReply SMultiLineEditableText::OnMouseMove( const FGeometry& MyGeometry, const FPointerEvent& MouseEvent )
@@ -789,11 +789,11 @@ FReply SMultiLineEditableText::OnMouseMove( const FGeometry& MyGeometry, const F
 
 		if (IsRightClickScrolling())
 		{
-			const FVector2f PreviousScrollOffset = EditableTextLayout->GetScrollOffset();
+			const FVector2f PreviousScrollOffset = SlateEditableTextLayout.GetScrollOffset();
 
 			FVector2f NewScrollOffset = PreviousScrollOffset;
 			NewScrollOffset.Y -= ScrollByAmount;
-			EditableTextLayout->SetScrollOffset(NewScrollOffset, MyGeometry);
+			SlateEditableTextLayout.SetScrollOffset(NewScrollOffset, MyGeometry);
 
 			if (!bIsSoftwareCursor)
 			{
@@ -803,7 +803,7 @@ FReply SMultiLineEditableText::OnMouseMove( const FGeometry& MyGeometry, const F
 
 			if (PreviousScrollOffset.Y != NewScrollOffset.Y)
 			{
-				const float ScrollMax = EditableTextLayout->GetSize().Y - MyGeometry.GetLocalSize().Y;
+				const float ScrollMax = SlateEditableTextLayout.GetSize().Y - MyGeometry.GetLocalSize().Y;
 				const float ScrollbarOffset = (ScrollMax != 0.0f) ? NewScrollOffset.Y / ScrollMax : 0.0f;
 				OnVScrollBarUserScrolled.ExecuteIfBound(ScrollbarOffset);
 				SoftwareCursorPosition.Y += (PreviousScrollOffset.Y - NewScrollOffset.Y);
@@ -813,7 +813,7 @@ FReply SMultiLineEditableText::OnMouseMove( const FGeometry& MyGeometry, const F
 		}
 	}
 
-	return EditableTextLayout->HandleMouseMove(MyGeometry, MouseEvent);
+	return SlateEditableTextLayout.HandleMouseMove(MyGeometry, MouseEvent);
 }
 
 FReply SMultiLineEditableText::OnMouseWheel( const FGeometry& MyGeometry, const FPointerEvent& MouseEvent )
@@ -822,15 +822,15 @@ FReply SMultiLineEditableText::OnMouseWheel( const FGeometry& MyGeometry, const
 	{
 		const float ScrollAmount = -MouseEvent.GetWheelDelta() * GetGlobalScrollAmount();
 
-		const FVector2f PreviousScrollOffset = EditableTextLayout->GetScrollOffset();
-		
+		const FVector2f PreviousScrollOffset = SlateEditableTextLayout.GetScrollOffset();
+
 		FVector2f NewScrollOffset = PreviousScrollOffset;
 		NewScrollOffset.Y += ScrollAmount;
-		EditableTextLayout->SetScrollOffset(NewScrollOffset, MyGeometry);
+		SlateEditableTextLayout.SetScrollOffset(NewScrollOffset, MyGeometry);
 
 		if (PreviousScrollOffset.Y != NewScrollOffset.Y)
 		{
-			const float ScrollMax = EditableTextLayout->GetSize().Y - MyGeometry.GetLocalSize().Y;
+			const float ScrollMax = SlateEditableTextLayout.GetSize().Y - MyGeometry.GetLocalSize().Y;
 			const float ScrollbarOffset = (ScrollMax != 0.0f) ? NewScrollOffset.Y / ScrollMax : 0.0f;
 			OnVScrollBarUserScrolled.ExecuteIfBound(ScrollbarOffset);
 			return FReply::Handled();
@@ -842,7 +842,7 @@ FReply SMultiLineEditableText::OnMouseWheel( const FGeometry& MyGeometry, const
 
 FReply SMultiLineEditableText::OnMouseButtonDoubleClick(const FGeometry& MyGeometry, const FPointerEvent& MouseEvent)
 {
-	return EditableTextLayout->HandleMouseButtonDoubleClick(MyGeometry, MouseEvent);
+	return SlateEditableTextLayout.HandleMouseButtonDoubleClick(MyGeometry, MouseEvent);
 }
 
 FCursorReply SMultiLineEditableText::OnCursorQuery( const FGeometry& MyGeometry, const FPointerEvent& CursorEvent ) const
@@ -866,7 +866,6 @@ bool SMultiLineEditableText::ComputeVolatility() const
 {
 	return SWidget::ComputeVolatility()
 		|| HasKeyboardFocus()
-		|| EditableTextLayout->ComputeVolatility()
 		|| bIsReadOnly.IsBound();
 }
 
diff --git a/Engine/Source/Runtime/Slate/Private/Widgets/Text/SlateEditableTextLayout.cpp b/Engine/Source/Runtime/Slate/Private/Widgets/Text/SlateEditableTextLayout.cpp
--- a/Engine/Source/Runtime/Slate/Private/Widgets/Text/SlateEditableTextLayout.cpp
+++ b/Engine/Source/Runtime/Slate/Private/Widgets/Text/SlateEditableTextLayout.cpp
@@ -21,7 +21,7 @@
 
 /**
  * Ensure that text transactions are always completed.
- * i.e. never forget to call FinishChangingText();
+ * i.e. never forget to call EndEditTransaction();
  */
 struct FScopedEditableTextTransaction
 {
@@ -29,7 +29,7 @@ public:
 	explicit FScopedEditableTextTransaction(FSlateEditableTextLayout& InEditableTextLayout)
 		: EditableTextLayout(&InEditableTextLayout)
 	{
-		EditableTextLayout->BeginEditTransation();
+		EditableTextLayout->BeginEditTransaction();
 	}
 
 	~FScopedEditableTextTransaction()
@@ -65,7 +65,7 @@ bool IsCharAllowed(const TCHAR InChar)
 
 }
 
-FSlateEditableTextLayout::FSlateEditableTextLayout(ISlateEditableTextWidget& InOwnerWidget, const TAttribute<FText>& InInitialText, FTextBlockStyle InTextStyle, const TOptional<ETextShapingMethod> InTextShapingMethod, const TOptional<ETextFlowDirection> InTextFlowDirection, const FCreateSlateTextLayout& InCreateSlateTextLayout, TSharedRef<ITextLayoutMarshaller> InTextMarshaller, TSharedRef<ITextLayoutMarshaller> InHintTextMarshaller)
+void FSlateEditableTextLayout::Construct(TAttribute<FText> InInitialText, FTextBlockStyle InTextStyle, const TOptional<ETextShapingMethod>& InTextShapingMethod, const TOptional<ETextFlowDirection>& InTextFlowDirection, const FCreateSlateTextLayout& InCreateSlateTextLayout, TSharedRef<ITextLayoutMarshaller> InTextMarshaller, TSharedRef<ITextLayoutMarshaller> InHintTextMarshaller)
 {
 	CreateSlateTextLayout = InCreateSlateTextLayout;
 	if (!CreateSlateTextLayout.IsBound())
@@ -76,34 +76,28 @@ FSlateEditableTextLayout::FSlateEditableTextLayout(ISlateEditableTextWidget& InO
 			});
 	}
 
-	OwnerWidget = &InOwnerWidget;
-	Marshaller = InTextMarshaller;
-	HintMarshaller = InHintTextMarshaller;
-	TextStyle = InTextStyle;
-	HintTextStyle = TextStyle;
-	TextLayout = CreateSlateTextLayout.Execute(&InOwnerWidget.GetSlateWidget().Get(), TextStyle);
+	const TSharedRef<SWidget> OwnerSlateWidget = OwnerWidget->GetSlateWidget();
 
-	WrapTextAt = 0.0f;
-	AutoWrapText = false;
-	WrappingPolicy = ETextWrappingPolicy::DefaultWrapping;
+	Marshaller = MoveTemp(InTextMarshaller);
+	HintMarshaller = MoveTemp(InHintTextMarshaller);
+	TextStyle = MoveTemp(InTextStyle);
+	HintTextStyle = TextStyle;
+	TextLayout = CreateSlateTextLayout.Execute(&OwnerSlateWidget.Get(), TextStyle);
 
-	Margin = FMargin(0);
-	Justification = ETextJustify::Left;
-	LineHeightPercentage = 1.0f;
 	DebugSourceInfo = FString();
 
 	SearchCase = ESearchCase::IgnoreCase;
 
 	GraphemeBreakIterator = GSlateUseSharedBreakIterator ? FBreakIterator::GetCharacterBoundaryIterator() : FBreakIterator::CreateCharacterBoundaryIterator();
 
-	BoundText = InInitialText;
+	BoundText.Assign(OwnerSlateWidget.Get(), MoveTemp(InInitialText), ESlateAttributeBindAction::Initialize);
 
 	// Set the initial text - the same as SetText, but doesn't call OnTextChanged as that can cause a crash during construction
 	{
 		const bool bIsPassword = OwnerWidget->IsTextPassword();
 		TextLayout->SetIsPassword(bIsPassword);
 
-		const FText& InitialTextToSet = BoundText.Get(FText::GetEmpty());
+		const FText& InitialTextToSet = BoundText.Get();
 		SetEditableText(InitialTextToSet, true);
 
 		// Update the cached BoundText value to prevent it triggering another SetEditableText update again next Tick
@@ -143,13 +137,14 @@ FSlateEditableTextLayout::FSlateEditableTextLayout(ISlateEditableTextWidget& InO
 	bTextChangedByVirtualKeyboard = false;
 	bTextCommittedByVirtualKeyboard = false;
 	bSelectionChangedExternally = false;
+	bMoveCursorToDocumentEnd = false;
 	VirtualKeyboardTextCommitType = ETextCommit::Default;
 
 	CachedSize = FVector2f::ZeroVector;
 
 	auto ExecuteDeleteAction = [this]()
 	{
-		BeginEditTransation();
+		BeginEditTransaction();
 		DeleteSelectedText();
 		EndEditTransaction();
 	};
@@ -220,44 +215,21 @@ FSlateEditableTextLayout::~FSlateEditableTextLayout()
 	}
 }
 
-void FSlateEditableTextLayout::SetText(const TAttribute<FText>& InText)
+void FSlateEditableTextLayout::SetText(TAttribute<FText> InText)
 {
-	const FText PreviousText = BoundText.Get(FText::GetEmpty());
-	BoundText = InText;
-	const FText NewText = BoundText.Get(FText::GetEmpty());
-
-	// We need to force an update if the text doesn't match the editable text, as the editable 
-	// text may not match the current bound text since it may have been changed by the user
-	const FText EditableText = GetEditableText();
-	const bool bForceRefresh = !EditableText.ToString().Equals(NewText.ToString(), ESearchCase::CaseSensitive);
+	bMoveCursorToDocumentEnd = false;
 
-	// Only emit the "text changed" event if the text has actually been changed
-	const bool bHasTextChanged = OwnerWidget->GetSlateWidget()->HasAnyUserFocus().IsSet()
-		? !NewText.ToString().Equals(EditableText.ToString(), ESearchCase::CaseSensitive)
-		: !NewText.ToString().Equals(PreviousText.ToString(), ESearchCase::CaseSensitive);
+	// this will call OnBoundTextChanged right away
+	BoundText.Assign(OwnerWidget->GetSlateWidget().Get(), MoveTemp(InText), FText::GetEmpty(), ESlateAttributeBindAction::Initialize);
 
-	if (RefreshImpl(&NewText, bForceRefresh))
+	// Make sure we move the cursor to the end of the new text if we had keyboard focus
+	if (bMoveCursorToDocumentEnd && OwnerWidget->GetSlateWidget()->HasAnyUserFocus().IsSet() && !bWasFocusedByLastMouseDown)
 	{
-		// Make sure we move the cursor to the end of the new text if we had keyboard focus
-		if (OwnerWidget->GetSlateWidget()->HasAnyUserFocus().IsSet() && !bWasFocusedByLastMouseDown)
-		{
-			JumpTo(ETextLocation::EndOfDocument, ECursorAction::MoveCursor);
-		}
-
-		// Let outsiders know that the text content has been changed
-		if (bHasTextChanged)
-		{
-			OwnerWidget->OnTextChanged(NewText);
-		}
-	}
-
-	if (bHasTextChanged || BoundText.IsBound())
-	{
-		OwnerWidget->GetSlateWidget()->Invalidate(EInvalidateWidget::LayoutAndVolatility);
+		JumpTo(ETextLocation::EndOfDocument, ECursorAction::MoveCursor);
 	}
 }
 
-int32 FSlateEditableTextLayout::GetTextLineCount()
+int32 FSlateEditableTextLayout::GetTextLineCount() const
 {
 	return TextLayout->GetLineCount();
 }
@@ -266,15 +238,17 @@ FText FSlateEditableTextLayout::GetText() const
 {
 	SLATE_CROSS_THREAD_CHECK();
 
-	return BoundText.Get(FText::GetEmpty());
+	return BoundText.Get();
 }
 
-void FSlateEditableTextLayout::SetHintText(const TAttribute<FText>& InHintText)
+void FSlateEditableTextLayout::SetHintText(TAttribute<FText> InHintText)
 {
-	HintText = InHintText;
+	const TSharedRef<SWidget> OwnerSlateWidget = OwnerWidget->GetSlateWidget();
+
+	HintText.Assign(OwnerSlateWidget.Get(), MoveTemp(InHintText));
 
 	// If we have hint text that is either non-empty or bound to a delegate, we'll also need to make the hint text layout
-	if (HintText.IsBound() || !HintText.Get(FText::GetEmpty()).IsEmpty())
+	if (HintText.IsBound(OwnerSlateWidget.Get()) || !HintText.Get().IsEmpty())
 	{
 		HintTextLayout = MakeUnique<FSlateTextBlockLayout>(OwnerWidget->GetSlateWidgetPtr().Get(), HintTextStyle, TextLayout->GetTextShapingMethod(), TextLayout->GetTextFlowDirection(), CreateSlateTextLayout, HintMarshaller.ToSharedRef(), nullptr);
 		HintTextLayout->SetDebugSourceInfo(DebugSourceInfo);
@@ -289,18 +263,12 @@ void FSlateEditableTextLayout::SetHintText(const TAttribute<FText>& InHintText)
 
 FText FSlateEditableTextLayout::GetHintText() const
 {
-	return HintText.Get(FText::GetEmpty());
+	return HintText.Get();
 }
 
-void FSlateEditableTextLayout::SetSearchText(const TAttribute<FText>& InSearchText)
+void FSlateEditableTextLayout::SetSearchText(TAttribute<FText> InSearchText)
 {
-	const FText& SearchTextToSet = InSearchText.Get(FText::GetEmpty());
-
-	BoundSearchText = InSearchText;
-	BoundSearchTextLastTick = FTextSnapshot(SearchTextToSet);
-
-	BeginSearch(SearchTextToSet);
-	OwnerWidget->GetSlateWidget()->Invalidate(EInvalidateWidget::LayoutAndVolatility);
+	BoundSearchText.Assign(OwnerWidget->GetSlateWidget().Get(), MoveTemp(InSearchText));
 }
 
 FText FSlateEditableTextLayout::GetSearchText() const
@@ -325,6 +293,8 @@ void FSlateEditableTextLayout::SetTextStyle(const FTextBlockStyle& InTextStyle)
 
 	TextLayout->SetDefaultTextStyle(TextStyle);
 	Marshaller->MakeDirty(); // will regenerate the text using the new default style
+
+	OwnerWidget->GetSlateWidget()->Invalidate(EInvalidateWidget::Layout); //Using Layout as changing text block size can affect the size.
 }
 
 const FTextBlockStyle& FSlateEditableTextLayout::GetTextStyle() const
@@ -332,14 +302,14 @@ const FTextBlockStyle& FSlateEditableTextLayout::GetTextStyle() const
 	return TextStyle;
 }
 
-void FSlateEditableTextLayout::SetHintTextOpacity(const TAttribute<float>& InOpacity)
+void FSlateEditableTextLayout::SetHintTextOpacity(TAttribute<float> InOpacity)
 {
-	HintTextOpacity = InOpacity;
+	HintTextOpacity.Assign(OwnerWidget->GetSlateWidget().Get(), MoveTemp(InOpacity));
 }
 
 float FSlateEditableTextLayout::GetHintTextOpacity() const
 {
-	return HintTextOpacity.Get(0.35f);
+	return HintTextOpacity.Get();
 }
 
 void FSlateEditableTextLayout::SetCursorBrush(const TAttribute<const FSlateBrush*>& InCursorBrush)
@@ -467,76 +437,76 @@ FTextSelection FSlateEditableTextLayout::GetSelection() const
 
 void FSlateEditableTextLayout::SetTextShapingMethod(const TOptional<ETextShapingMethod>& InTextShapingMethod)
 {
-	TextLayout->SetTextShapingMethod((InTextShapingMethod.IsSet()) ? InTextShapingMethod.GetValue() : GetDefaultTextShapingMethod());
+	const ETextShapingMethod NewShapingMethod = InTextShapingMethod.Get(GetDefaultTextShapingMethod());
+	if (NewShapingMethod != TextLayout->GetTextShapingMethod())
+	{
+		TextLayout->SetTextShapingMethod((InTextShapingMethod.IsSet()) ? InTextShapingMethod.GetValue() : GetDefaultTextShapingMethod());
+
+		OwnerWidget->GetSlateWidget()->Invalidate(EInvalidateWidget::Layout);
+	}
 }
 
 void FSlateEditableTextLayout::SetTextFlowDirection(const TOptional<ETextFlowDirection>& InTextFlowDirection)
 {
-	TextLayout->SetTextFlowDirection((InTextFlowDirection.IsSet()) ? InTextFlowDirection.GetValue() : GetDefaultTextFlowDirection());
+	const ETextFlowDirection NewFlowDirection = InTextFlowDirection.Get(GetDefaultTextFlowDirection());
+	if (NewFlowDirection != TextLayout->GetTextFlowDirection())
+	{
+		TextLayout->SetTextFlowDirection(NewFlowDirection);
+
+		OwnerWidget->GetSlateWidget()->Invalidate(EInvalidateWidget::Layout);
+	}
 }
 
-void FSlateEditableTextLayout::SetTextWrapping(const TAttribute<float>& InWrapTextAt, const TAttribute<bool>& InAutoWrapText, const TAttribute<ETextWrappingPolicy>& InWrappingPolicy)
+void FSlateEditableTextLayout::SetTextWrapping(TAttribute<float> InWrapTextAt, TAttribute<bool> InAutoWrapText, TAttribute<ETextWrappingPolicy> InWrappingPolicy)
 {
-	WrapTextAt = InWrapTextAt;
-	AutoWrapText = InAutoWrapText;
-	WrappingPolicy = InWrappingPolicy;
+	const TSharedRef<SWidget> OwnerSlateWidget = OwnerWidget->GetSlateWidget();
 
-	OwnerWidget->GetSlateWidget()->Invalidate(EInvalidateWidget::LayoutAndVolatility);
+	WrapTextAt.Assign(OwnerSlateWidget.Get(), MoveTemp(InWrapTextAt));
+	AutoWrapText.Assign(OwnerSlateWidget.Get(), MoveTemp(InAutoWrapText));
+	WrappingPolicy.Assign(OwnerSlateWidget.Get(), MoveTemp(InWrappingPolicy));
 }
 
-void FSlateEditableTextLayout::SetWrapTextAt(const TAttribute<float>& InWrapTextAt)
+void FSlateEditableTextLayout::SetWrapTextAt(TAttribute<float> InWrapTextAt)
 {
-	WrapTextAt = InWrapTextAt;
-
-	OwnerWidget->GetSlateWidget()->Invalidate(EInvalidateWidget::LayoutAndVolatility);
+	WrapTextAt.Assign(OwnerWidget->GetSlateWidget().Get(), MoveTemp(InWrapTextAt));
 }
 
-void FSlateEditableTextLayout::SetAutoWrapText(const TAttribute<bool>& InAutoWrapText)
+void FSlateEditableTextLayout::SetAutoWrapText(TAttribute<bool> InAutoWrapText)
 {
-	AutoWrapText = InAutoWrapText;
-
-	OwnerWidget->GetSlateWidget()->Invalidate(EInvalidateWidget::LayoutAndVolatility);
+	AutoWrapText.Assign(OwnerWidget->GetSlateWidget().Get(), MoveTemp(InAutoWrapText));
 }
 
-void FSlateEditableTextLayout::SetWrappingPolicy(const TAttribute<ETextWrappingPolicy>& InWrappingPolicy)
+void FSlateEditableTextLayout::SetWrappingPolicy(TAttribute<ETextWrappingPolicy> InWrappingPolicy)
 {
-	WrappingPolicy = InWrappingPolicy;
-
-	OwnerWidget->GetSlateWidget()->Invalidate(EInvalidateWidget::LayoutAndVolatility);
+	WrappingPolicy.Assign(OwnerWidget->GetSlateWidget().Get(), MoveTemp(InWrappingPolicy));
 }
 
-void FSlateEditableTextLayout::SetMargin(const TAttribute<FMargin>& InMargin)
+void FSlateEditableTextLayout::SetMargin(TAttribute<FMargin> InMargin)
 {
-	Margin = InMargin;
-
-	OwnerWidget->GetSlateWidget()->Invalidate(EInvalidateWidget::LayoutAndVolatility);
+	Margin.Assign(OwnerWidget->GetSlateWidget().Get(), MoveTemp(InMargin));
 }
 
-void FSlateEditableTextLayout::SetJustification(const TAttribute<ETextJustify::Type>& InJustification)
+void FSlateEditableTextLayout::SetJustification(TAttribute<ETextJustify::Type> InJustification)
 {
-	if (!Justification.IdenticalTo(InJustification))
-	{
-		Justification = InJustification;
-
-		OwnerWidget->GetSlateWidget()->Invalidate(EInvalidateWidget::LayoutAndVolatility);
-	}
+	Justification.Assign(OwnerWidget->GetSlateWidget().Get(), MoveTemp(InJustification));
 }
 
-void FSlateEditableTextLayout::SetLineHeightPercentage(const TAttribute<float>& InLineHeightPercentage)
+ETextJustify::Type FSlateEditableTextLayout::GetJustification() const
 {
-	LineHeightPercentage = InLineHeightPercentage;
-
-	OwnerWidget->GetSlateWidget()->Invalidate(EInvalidateWidget::LayoutAndVolatility);
+	return Justification.Get();
 }
 
-void FSlateEditableTextLayout::SetApplyLineHeightToBottomLine(const TAttribute<bool>& InApplyLineHeightToBottomLine)
+void FSlateEditableTextLayout::SetLineHeightPercentage(TAttribute<float> InLineHeightPercentage)
 {
-	ApplyLineHeightToBottomLine = InApplyLineHeightToBottomLine;
+	LineHeightPercentage.Assign(OwnerWidget->GetSlateWidget().Get(), MoveTemp(InLineHeightPercentage));
+}
 
-	OwnerWidget->GetSlateWidget()->Invalidate(EInvalidateWidget::LayoutAndVolatility);
+void FSlateEditableTextLayout::SetApplyLineHeightToBottomLine(TAttribute<bool> InApplyLineHeightToBottomLine)
+{
+	ApplyLineHeightToBottomLine.Assign(OwnerWidget->GetSlateWidget().Get(), MoveTemp(InApplyLineHeightToBottomLine));
 }
 
-void FSlateEditableTextLayout::SetOverflowPolicy(TOptional<ETextOverflowPolicy> InOverflowPolicy)
+void FSlateEditableTextLayout::SetOverflowPolicy(const TOptional<ETextOverflowPolicy>& InOverflowPolicy)
 {
 	if(OverflowPolicyOverride != InOverflowPolicy)
 	{
@@ -603,7 +573,7 @@ void FSlateEditableTextLayout::EnableTextInputMethodContext()
 
 bool FSlateEditableTextLayout::Refresh()
 {
-	const FText& TextToSet = BoundText.Get(FText::GetEmpty());
+	const FText& TextToSet = BoundText.Get();
 	return RefreshImpl(&TextToSet);
 }
 
@@ -622,7 +592,7 @@ bool FSlateEditableTextLayout::RefreshImpl(const FText* InTextToSet, const bool
 		if (bForce || !BoundTextLastTick.IsDisplayStringEqualTo(*InTextToSet))
 		{
 			// The source text has changed, so update the internal editable text
-			bHasSetText = SetEditableText(*InTextToSet, true);
+			bHasSetText = SetEditableText(*InTextToSet, bForce);
 		}
 
 		// Update this even if the text is lexically identical, as it will update the pointer compared by IdenticalTo for the next Tick
@@ -892,7 +862,7 @@ bool FSlateEditableTextLayout::HandleFocusReceived(const FFocusEvent& InFocusEve
 	// Force update of last cursor interaction time to ensure the caret blinks at the correct frequency regardless of how focus is set.
 	CursorInfo.UpdateLastCursorInteractionTime();
 
-	OwnerWidget->OnBeginTextEdit(BoundText.Get(FText::GetEmpty()));
+	OwnerWidget->OnBeginTextEdit(BoundText.Get());
 	return true;
 }
 
@@ -3179,7 +3149,7 @@ bool FSlateEditableTextLayout::HasTextChangedFromOriginal() const
 	return bHasChanged;
 }
 
-void FSlateEditableTextLayout::BeginEditTransation()
+void FSlateEditableTextLayout::BeginEditTransaction()
 {
 	NumTransactionsOpened += 1;
 	
@@ -3411,37 +3381,36 @@ void FSlateEditableTextLayout::Redo()
 
 void FSlateEditableTextLayout::SaveText(const FText& TextToSave)
 {
+	const TSharedRef<SWidget> OwnerSlateWidget = OwnerWidget->GetSlateWidget();
 	// Don't set text if the text attribute has a 'getter' binding on it, otherwise we'd blow away
 	// that binding.  If there is a getter binding, then we'll assume it will provide us with
 	// updated text after we've fired our 'text changed' callbacks
-	if (!BoundText.IsBound())
+	if (!BoundText.IsBound(OwnerSlateWidget.Get()))
 	{
-		BoundText.Set(TextToSave);
+		BoundText.Assign(OwnerSlateWidget.Get(), TextToSave);
 	}
 }
 
 void FSlateEditableTextLayout::LoadText()
 {
+	const TSharedRef<SWidget> OwnerSlateWidget = OwnerWidget->GetSlateWidget();
 	// We only need to do this if we're bound to a delegate, otherwise the text layout will already be up-to-date
 	// either from Construct, or a call to SetText
-	if (BoundText.IsBound())
-	{
-		SetText(BoundText);
+	if (BoundText.IsBound(OwnerSlateWidget.Get()))
+	{
+		BoundText.UpdateNow(OwnerSlateWidget.Get());
+		// If we arrive here when the text is commited and `BoundText` is bound,
+		// the new text might have been refused in which case `BoundText.UpdateNow` would have not triggerred `OnBoundTextChanged`
+		// If this was not called, the Edited text would now not match the `BoundText` which is a problem, so we call it manually.
+		// Calling it twice would not be a problem (even though not ideal) as nothing would trigger now that the content is identical.
+		OnBoundTextChanged(OwnerSlateWidget.Get(), *this);
 		TextLayout->UpdateIfNeeded();
 	}
 }
 
 bool FSlateEditableTextLayout::ComputeVolatility() const
 {
-	return BoundText.IsBound()
-		|| HintText.IsBound()
-		|| BoundSearchText.IsBound()
-		|| WrapTextAt.IsBound()
-		|| AutoWrapText.IsBound()
-		|| WrappingPolicy.IsBound()
-		|| Margin.IsBound()
-		|| Justification.IsBound()
-		|| LineHeightPercentage.IsBound();
+	return false;
 }
 
 void FSlateEditableTextLayout::UpdateTextChangedByVirtualKeyboard()
@@ -3509,7 +3478,7 @@ void FSlateEditableTextLayout::Tick(const FGeometry& AllottedGeometry, const dou
 
 	// Update the search before we process the next PositionToScrollIntoView
 	{
-		const FText& SearchTextToSet = BoundSearchText.Get(FText::GetEmpty());
+		const FText& SearchTextToSet = BoundSearchText.Get();
 		if (!BoundSearchTextLastTick.IdenticalTo(SearchTextToSet))
 		{
 			// The pointer used by the bound text has changed, however the text may still be the same - check that now
@@ -3598,9 +3567,9 @@ void FSlateEditableTextLayout::Tick(const FGeometry& AllottedGeometry, const dou
 	TextLayout->SetVisibleRegion(AllottedGeometry.Size, FVector2D(ScrollOffset) * TextLayout->GetScale());
 }
 
-int32 FSlateEditableTextLayout::OnPaint(const FPaintArgs& Args, const FGeometry& AllottedGeometry, const FSlateRect& MyCullingRect, FSlateWindowElementList& OutDrawElements, int32 LayerId, const FWidgetStyle& InWidgetStyle, bool bParentEnabled)
+int32 FSlateEditableTextLayout::OnPaint(const FPaintArgs& Args, const FGeometry& AllottedGeometry, const FSlateRect& MyCullingRect, FSlateWindowElementList& OutDrawElements, int32 LayerId, const FWidgetStyle& InWidgetStyle, bool bParentEnabled) const
 {
-	// Update the auto-wrap size now that we have computed paint geometry; won't take affect until text frame
+	// Update the auto-wrap size now that we have computed paint geometry; won't take effect until text frame
 	// Note: This is done here rather than in Tick(), because Tick() doesn't get called while resizing windows, but OnPaint() does
 	CachedSize = FVector2f(AllottedGeometry.GetLocalSize());
 
@@ -3814,11 +3783,13 @@ void FSlateEditableTextLayout::FVirtualKeyboardEntry::SetTextFromVirtualKeyboard
 {
 	check(IsInGameThread());
 
+	const TSharedRef<SWidget> OwnerSlateWidget = OwnerLayout->OwnerWidget->GetSlateWidget();
+
 	// Only set the text if the text attribute doesn't have a getter binding (otherwise it would be blown away).
 	// If it is bound, we'll assume that OnTextCommitted will handle the update.
-	if (!OwnerLayout->BoundText.IsBound())
+	if (!OwnerLayout->BoundText.IsBound(OwnerSlateWidget.Get()))
 	{
-		OwnerLayout->BoundText.Set(InNewText);
+		OwnerLayout->BoundText.Assign(OwnerSlateWidget.Get(), InNewText);
 	}
 
 	// Update the internal editable text
@@ -4184,7 +4155,7 @@ void FSlateEditableTextLayout::FTextInputMethodContext::BeginComposition()
 	{
 		bIsComposing = true;
 
-		OwnerLayout->BeginEditTransation();
+		OwnerLayout->BeginEditTransaction();
 		OwnerLayout->UpdateCursorHighlight();
 	}
 }
@@ -4243,4 +4214,39 @@ void FSlateEditableTextLayout::SetFontFacesLoadingPaintPolicy(EFontFacesLoadingP
 	{
 		TextLayout->SetFontFacesLoadingPaintPolicy(InFontFacesLoadingPaintPolicy);
 	}
-}
\ No newline at end of file
+}
+
+void FSlateEditableTextLayout::OnBoundTextChanged(SWidget& Widget, FSlateEditableTextLayout& SlateTextLayout)
+{
+	const FText NewText = SlateTextLayout.BoundText.Get();
+
+	// We need to force an update if the text doesn't match the editable text, as the editable
+	// text may not match the current bound text since it may have been changed by the user
+	const FText EditableText = SlateTextLayout.GetEditableText();
+	const bool bForceRefresh = !EditableText.ToString().Equals(NewText.ToString(), ESearchCase::CaseSensitive);
+
+	// Only emit the "text changed" event if the text has actually been changed
+	const bool bHasTextChanged = Widget.HasAnyUserFocus().IsSet()
+		? !NewText.ToString().Equals(EditableText.ToString(), ESearchCase::CaseSensitive)
+		: !NewText.ToString().Equals(SlateTextLayout.PreviousBoundText.ToString(), ESearchCase::CaseSensitive);
+
+	if (SlateTextLayout.RefreshImpl(&NewText, bForceRefresh))
+	{
+		SlateTextLayout.bMoveCursorToDocumentEnd = true;
+
+		// Let outsiders know that the text content has been changed
+		if (bHasTextChanged)
+		{
+			SlateTextLayout.OwnerWidget->OnTextChanged(NewText);
+		}
+	}
+
+	SlateTextLayout.PreviousBoundText = SlateTextLayout.BoundText.Get();
+}
+
+void FSlateEditableTextLayout::OnBoundSearchTextChanged(SWidget& Widget, FSlateEditableTextLayout& SlateTextLayout)
+{
+	const FText& SearchTextToSet = SlateTextLayout.BoundSearchText.Get();
+	SlateTextLayout.BoundSearchTextLastTick = FTextSnapshot(SearchTextToSet);
+	SlateTextLayout.BeginSearch(SearchTextToSet);
+}
diff --git a/Engine/Source/Runtime/Slate/Public/Widgets/Input/SEditableText.h b/Engine/Source/Runtime/Slate/Public/Widgets/Input/SEditableText.h
--- a/Engine/Source/Runtime/Slate/Public/Widgets/Input/SEditableText.h
+++ b/Engine/Source/Runtime/Slate/Public/Widgets/Input/SEditableText.h
@@ -19,13 +19,13 @@
 #include "Framework/Text/TextLayout.h"
 #include "Framework/Text/SlateAsyncFontLoadingTypes.h"
 #include "Widgets/Text/ISlateEditableTextWidget.h"
+#include "Widgets/Text/SlateEditableTextLayout.h"
 
 class FActiveTimerHandle;
 class FArrangedChildren;
 class FChildren;
 class FPaintArgs;
 class FPlainTextLayoutMarshaller;
-class FSlateEditableTextLayout;
 class FSlateWindowElementList;
 class IBreakIterator;
 struct FTextLocation;
@@ -34,11 +34,18 @@ enum class ETextShapingMethod : uint8;
 
 DECLARE_DELEGATE_TwoParams(FOnCursorMovedWithSelection, const FTextLocation&, const FTextSelection&);
 
+template <>
+struct TWidgetTypeTraits<class SEditableText>
+{
+	static constexpr bool SupportsInvalidation() { return true; }
+};
+
 /**
  * Editable text widget
  */
 class SEditableText : public SWidget, public ISlateEditableTextWidget
 {
+	SLATE_DECLARE_WIDGET_API(SEditableText, SWidget, SLATE_API)
 public:
 	SLATE_BEGIN_ARGS( SEditableText )
 		: _Text()
@@ -203,9 +210,6 @@ public:
 	/** Constructor */
 	SLATE_API SEditableText();
 
-	/** Destructor */
-	SLATE_API ~SEditableText();
-
 	/**
 	 * Construct this widget
 	 *
@@ -218,7 +222,7 @@ public:
 	 *
 	 * @param  InNewText  The new text
 	 */
-	SLATE_API void SetText( const TAttribute< FText >& InNewText );
+	SLATE_API void SetText(TAttribute<FText> InNewText);
 	
 	/**
 	 * Returns the text string
@@ -238,47 +242,50 @@ public:
 	SLATE_API bool SetEditableText(const FText& InNewText);
 
 	/** See the HintText attribute */
-	SLATE_API void SetHintText( const TAttribute< FText >& InHintText );
+	SLATE_API void SetHintText(TAttribute<FText> InHintText);
 	
 	/** Get the text that appears when there is no text in the text box */
 	SLATE_API FText GetHintText() const;
 	
 	/** Sets the hint text opacity */
-	SLATE_API void SetHintTextOpacity(const TAttribute<float>& InHintTextOpacity);
+	SLATE_API void SetHintTextOpacity(TAttribute<float> InHintTextOpacity);
 	
 	/** Get the hint text opacity */
 	SLATE_API float GetHintTextOpacity() const;
 
 	/** Set the text that is currently being searched for (if any) */
-	SLATE_API void SetSearchText(const TAttribute<FText>& InSearchText);
+	SLATE_API void SetSearchText(TAttribute<FText> InSearchText);
 
 	/** Get the text that is currently being searched for (if any) */
 	SLATE_API FText GetSearchText() const;
 
 	/** See the IsReadOnly attribute */
-	SLATE_API void SetIsReadOnly( TAttribute< bool > InIsReadOnly );
+	SLATE_API void SetIsReadOnly(TAttribute<bool> InIsReadOnly);
 	
 	/** See the IsPassword attribute */
-	SLATE_API void SetIsPassword( TAttribute< bool > InIsPassword );
+	SLATE_API void SetIsPassword(TAttribute<bool> InIsPassword);
 	
 	/** See the ColorAndOpacity attribute */
 	SLATE_API void SetColorAndOpacity(TAttribute<FSlateColor> Color);
 
+	/** See the BackgroundImageSelected attribute */
+	SLATE_API void SetBackgroundImageSelected(TAttribute<const FSlateBrush*> InBackgroundImageSelected);
+
 	/** See the AllowContextMenu attribute */
-	SLATE_API void SetAllowContextMenu(const TAttribute< bool >& InAllowContextMenu);
+	SLATE_API void SetAllowContextMenu(TAttribute<bool> InAllowContextMenu);
 
 	/** See the EnableIntegratedKeyboard attribute */
-	SLATE_API void SetEnableIntegratedKeyboard(const TAttribute< bool >& InEnableIntegratedKeyboard);
+	SLATE_API void SetEnableIntegratedKeyboard(TAttribute<bool> InEnableIntegratedKeyboard);
 
 	/** Set the VirtualKeyboardDismissAction attribute */
-	SLATE_API void SetVirtualKeyboardDismissAction(TAttribute< EVirtualKeyboardDismissAction > InVirtualKeyboardDismissAction);
+	SLATE_API void SetVirtualKeyboardDismissAction(TAttribute<EVirtualKeyboardDismissAction> InVirtualKeyboardDismissAction);
 
 	/**
 	 * Sets the font used to draw the text
 	 *
 	 * @param  InNewFont	The new font to use
 	 */
-	SLATE_API void SetFont(const TAttribute< FSlateFontInfo >& InNewFont);
+	SLATE_API void SetFont(TAttribute<FSlateFontInfo> InNewFont);
 
 	/** Gets the font used to draw the text. */
 	SLATE_API FSlateFontInfo GetFont() const;
@@ -286,7 +293,7 @@ public:
 	/**
 	 * Sets the text style used to draw the text
 	 *
-	 * @param  NewTextStyle	The new text style to use
+	 * @param  InNewTextStyle	The new text style to use
 	 */
 	SLATE_API void SetTextStyle( const FEditableTextStyle& InNewTextStyle );
 
@@ -303,52 +310,53 @@ public:
 	 *
 	 * @param  InMinDesiredWidth	The minimum width
 	 */
-	SLATE_API void SetMinDesiredWidth(const TAttribute<float>& InMinDesiredWidth);
+	SLATE_API void SetMinDesiredWidth(TAttribute<float> InMinDesiredWidth);
 
 	/**
 	 * Workaround as we loose focus when the auto completion closes.
 	 *
 	 * @param  InIsCaretMovedWhenGainFocus	Workaround
 	 */
-	SLATE_API void SetIsCaretMovedWhenGainFocus(const TAttribute<bool>& InIsCaretMovedWhenGainFocus);
+	SLATE_API void SetIsCaretMovedWhenGainFocus(TAttribute<bool> InIsCaretMovedWhenGainFocus);
 
 	/**
 	 * Sets whether to select all text when the user clicks to give focus on the widget
 	 *
 	 * @param  InSelectAllTextWhenFocused	Select all text when the user clicks?
 	 */
-	SLATE_API void SetSelectAllTextWhenFocused(const TAttribute<bool>& InSelectAllTextWhenFocused);
+	SLATE_API void SetSelectAllTextWhenFocused(TAttribute<bool> InSelectAllTextWhenFocused);
 
 	/**
 	 * Sets whether to allow the user to back out of changes when they press the escape key
 	 *
 	 * @param  InRevertTextOnEscape			Allow the user to back out of changes?
 	 */
-	SLATE_API void SetRevertTextOnEscape(const TAttribute<bool>& InRevertTextOnEscape);
+	SLATE_API void SetRevertTextOnEscape(TAttribute<bool> InRevertTextOnEscape);
 
 	/**
 	 * Sets whether to clear keyboard focus when pressing enter to commit changes
 	 *
 	 * @param  InClearKeyboardFocusOnCommit		Clear keyboard focus when pressing enter?
 	 */
-	SLATE_API void SetClearKeyboardFocusOnCommit(const TAttribute<bool>& InClearKeyboardFocusOnCommit);
+	SLATE_API void SetClearKeyboardFocusOnCommit(TAttribute<bool> InClearKeyboardFocusOnCommit);
 
 	/**
 	 * Sets whether to select all text when pressing enter to commit changes
 	 *
 	 * @param  InSelectAllTextOnCommit		Select all text when pressing enter?
 	 */
-	SLATE_API void SetSelectAllTextOnCommit(const TAttribute<bool>& InSelectAllTextOnCommit);
+	SLATE_API void SetSelectAllTextOnCommit(TAttribute<bool> InSelectAllTextOnCommit);
 
 	/**
 	 * Sets whether to select word on the mouse double click
 	 *
 	 * @param  InSelectWordOnMouseDoubleClick		Select word on the mouse double click
 	 */
-	SLATE_API void SetSelectWordOnMouseDoubleClick(const TAttribute<bool>& InSelectWordOnMouseDoubleClick);
+	SLATE_API void SetSelectWordOnMouseDoubleClick(TAttribute<bool> InSelectWordOnMouseDoubleClick);
 
 	/** See Justification attribute */
-	SLATE_API void SetJustification(const TAttribute<ETextJustify::Type>& InJustification);
+	SLATE_API void SetJustification(TAttribute<ETextJustify::Type> InJustification);
+
 	/**
 	 * Sets the OnKeyCharHandler to provide first chance handling of the OnKeyChar event
 	 *
@@ -356,7 +364,7 @@ public:
 	 */
 	void SetOnKeyCharHandler(FOnKeyChar InOnKeyCharHandler)
 	{
-		OnKeyCharHandler = InOnKeyCharHandler;
+		OnKeyCharHandler = MoveTemp(InOnKeyCharHandler);
 	}
 
 	/**
@@ -366,7 +374,7 @@ public:
 	 */
 	void SetOnKeyDownHandler(FOnKeyDown InOnKeyDownHandler)
 	{
-		OnKeyDownHandler = InOnKeyDownHandler;
+		OnKeyDownHandler = MoveTemp(InOnKeyDownHandler);
 	}
 
 	/** See TextShapingMethod attribute */
@@ -376,7 +384,7 @@ public:
 	SLATE_API void SetTextFlowDirection(const TOptional<ETextFlowDirection>& InTextFlowDirection);
 
 	/** Sets the overflow policy for this text block */
-	SLATE_API void SetOverflowPolicy(TOptional<ETextOverflowPolicy> InOverflowPolicy);
+	SLATE_API void SetOverflowPolicy(const TOptional<ETextOverflowPolicy>& InOverflowPolicy);
 
 	/** Sets the paint policy to use for text if the font faces are still async loading. */
 	SLATE_API void SetFontFacesLoadingPaintPolicy(EFontFacesLoadingPaintPolicy InFontFacesLoadingPaintPolicy);
@@ -443,7 +451,7 @@ protected:
 	SLATE_API virtual FReply OnMouseMove( const FGeometry& InMyGeometry, const FPointerEvent& InMouseEvent ) override;
 	SLATE_API virtual FReply OnMouseButtonDoubleClick( const FGeometry& InMyGeometry, const FPointerEvent& InMouseEvent ) override;
 	SLATE_API virtual FCursorReply OnCursorQuery( const FGeometry& MyGeometry, const FPointerEvent& CursorEvent ) const override;
-	SLATE_API virtual const FSlateBrush* GetFocusBrush() const;
+	SLATE_API virtual const FSlateBrush* GetFocusBrush() const override;
 	SLATE_API virtual bool IsInteractable() const override;
 	SLATE_API virtual bool ComputeVolatility() const override;
 #if WITH_ACCESSIBILITY
@@ -487,31 +495,75 @@ protected:
 	SLATE_API virtual void OnTextChanged(const FText& InText) override;
 	SLATE_API virtual void OnTextCommitted(const FText& InText, const ETextCommit::Type InTextAction) override;
 	SLATE_API virtual void OnCursorMoved(const FTextLocation& InLocation) override;
-	SLATE_API virtual float UpdateAndClampHorizontalScrollBar(const float InViewOffset, const float InViewFraction, const EVisibility InVisiblityOverride) override;
-	SLATE_API virtual float UpdateAndClampVerticalScrollBar(const float InViewOffset, const float InViewFraction, const EVisibility InVisiblityOverride) override;
+	SLATE_API virtual float UpdateAndClampHorizontalScrollBar(const float InViewOffset, const float InViewFraction, const EVisibility InVisibilityOverride) override;
+	SLATE_API virtual float UpdateAndClampVerticalScrollBar(const float InViewOffset, const float InViewFraction, const EVisibility InVisibilityOverride) override;
 	//~ End ISlateEditableTextWidget Interface
 
+protected:
+	SLATE_API ETextJustify::Type GetJustification() const;
+
+	TSlateAttributeRef<FSlateFontInfo> GetFontOverrideAttribute() const
+	{
+		return TSlateAttributeRef<FSlateFontInfo>{SharedThis(this), FontOverrideAttribute};
+	}
+
+	SLATE_API FSlateColor GetColorAndOpacity() const;
+	TSlateAttributeRef<FSlateColor> GetColorAndOpacityAttribute() const
+	{
+		return TSlateAttributeRef<FSlateColor>{SharedThis(this), ColorAndOpacityAttribute};
+	}
+
+	SLATE_API const FSlateBrush* GetBackgroundImageSelected() const;
+	TSlateAttributeRef<const FSlateBrush*> GetBackgroundImageSelectedAttribute() const
+	{
+		return TSlateAttributeRef<const FSlateBrush*>{SharedThis(this), BackgroundImageSelectedAttribute};
+	}
+
+	TSlateAttributeRef<bool> GetIsPasswordAttribute() const
+	{
+		return TSlateAttributeRef<bool>{SharedThis(this), bIsPasswordAttribute};
+	}
+
+	SLATE_API float GetMinDesiredWidth() const;
+	TSlateAttributeRef<float> GetMinDesiredWidthAttribute() const
+	{
+		return TSlateAttributeRef<float>{SharedThis(this), MinDesiredWidthAttribute};
+	}
+
 protected:
 	/** Text marshaller used by the editable text layout */
 	TSharedPtr<FPlainTextLayoutMarshaller> PlainTextMarshaller;
 
+#if WITH_EDITORONLY_DATA
 	/** The text layout that deals with the editable text */
+	UE_DEPRECATED(5.8, "Use SlateEditableTextLayout, this variable will never be set.")
 	TUniquePtr<FSlateEditableTextLayout> EditableTextLayout;
+#endif
+	/** The text layout that deals with the editable text */
+	FSlateEditableTextLayout SlateEditableTextLayout;
 
+#if WITH_EDITORONLY_DATA
 	/** The font used to draw the text */
-	TAttribute<FSlateFontInfo> Font;
+	UE_DEPRECATED(5.8, "Use SetFont / GetFont / GetFontOverrideAttribute")
+	TSlateDeprecatedTAttribute<FSlateFontInfo> Font;
 
 	/** Text color and opacity */
-	TAttribute<FSlateColor> ColorAndOpacity;
+	UE_DEPRECATED(5.8, "Use SetColorAndOpacity / GetColorAndOpacity / GetColorAndOpacityAttribute")
+	TSlateDeprecatedTAttribute<FSlateColor> ColorAndOpacity;
 
 	/** Background image for the selected text */
-	TAttribute<const FSlateBrush*> BackgroundImageSelected;
+	UE_DEPRECATED(5.8, "Use SetBackgroundImageSelected / GetBackgroundImageSelected / GetBackgroundImageSelectedAttribute")
+	TSlateDeprecatedTAttribute<const FSlateBrush*> BackgroundImageSelected;
+#endif
 
 	/** Sets whether this text box can actually be modified interactively by the user */
 	TAttribute<bool> bIsReadOnly;
 
+#if WITH_EDITORONLY_DATA
 	/** Sets whether this text box is for storing a password */
-	TAttribute<bool> bIsPassword;
+	UE_DEPRECATED(5.8, "Use SetIsPassword / IsTextPassword / GetIsPasswordAttribute")
+	TSlateDeprecatedTAttribute<bool> bIsPassword;
+#endif
 
 	/** Workaround as we loose focus when the auto completion closes. */
 	TAttribute<bool> bIsCaretMovedWhenGainFocus;
@@ -555,8 +607,11 @@ protected:
 	/** Called whenever the user's selection and/or cursor position of the text has changed. */
 	FOnCursorMovedWithSelection OnCursorMovedWithSelectionCallback;
 
+#if WITH_EDITORONLY_DATA
 	/** Prevents the editable text from being smaller than desired in certain cases (e.g. when it is empty) */
-	TAttribute<float> MinDesiredWidth;
+	UE_DEPRECATED(5.8, "Use SetMinDesiredWidth / GetMinDesiredWidth / GetMinDesiredWidthAttribute")
+	TSlateDeprecatedTAttribute<float> MinDesiredWidth;
+#endif
 
 	/** Menu extender for right-click context menu */
 	TSharedPtr<FExtender> MenuExtender;
@@ -576,12 +631,31 @@ protected:
 	/** The type of virtual keyboard to use for editing this text on mobile */
 	TAttribute<EKeyboardType> VirtualKeyboardType;
 
-	/** Additional options used by the virtual keyboard summoned by this widget */
-	FVirtualKeyboardOptions VirtualKeyboardOptions;
-
 	/** The type of event that will trigger the display of the virtual keyboard */
 	TAttribute<EVirtualKeyboardTrigger> VirtualKeyboardTrigger;
 
 	/** The message action to take when the virtual keyboard is dismissed by the user */
 	TAttribute<EVirtualKeyboardDismissAction> VirtualKeyboardDismissAction;
+
+	/** Additional options used by the virtual keyboard summoned by this widget */
+	FVirtualKeyboardOptions VirtualKeyboardOptions;
+
+private:
+	/** The font used to draw the text */
+	TSlateAttribute<FSlateFontInfo> FontOverrideAttribute;
+	bool bIsFontOverrideAttrSet = false;
+
+	/** Text color and opacity */
+	TSlateAttribute<FSlateColor> ColorAndOpacityAttribute;
+	bool bIsColorAndOpacityAttrSet = false;
+
+	/** Background image for the selected text */
+	TSlateAttribute<const FSlateBrush*> BackgroundImageSelectedAttribute;
+	bool bIsBackgroundImageSelectedAttrSet = false;
+
+	/** Sets whether this text box is for storing a password */
+	TSlateAttribute<bool> bIsPasswordAttribute;
+
+	/** Prevents the editable text from being smaller than desired in certain cases (e.g. when it is empty) */
+	TSlateAttribute<float> MinDesiredWidthAttribute;
 };
diff --git a/Engine/Source/Runtime/Slate/Public/Widgets/Text/SMultiLineEditableText.h b/Engine/Source/Runtime/Slate/Public/Widgets/Text/SMultiLineEditableText.h
--- a/Engine/Source/Runtime/Slate/Public/Widgets/Text/SMultiLineEditableText.h
+++ b/Engine/Source/Runtime/Slate/Public/Widgets/Text/SMultiLineEditableText.h
@@ -3,6 +3,7 @@
 #pragma once
 
 #include "CoreMinimal.h"
+#include "SlateEditableTextLayout.h"
 #include "Misc/Attribute.h"
 #include "SlateGlobals.h"
 #include "Layout/Margin.h"
@@ -28,7 +29,6 @@
 class FActiveTimerHandle;
 class FArrangedChildren;
 class FPaintArgs;
-class FSlateEditableTextLayout;
 class FSlateWindowElementList;
 class ITextLayoutMarshaller;
 enum class ETextShapingMethod : uint8;
@@ -41,6 +41,7 @@ class FSlateEditableTextLayout;
 /** An editable text widget that supports multiple lines and soft word-wrapping. */
 class SMultiLineEditableText : public SWidget, public ISlateEditableTextWidget
 {
+	SLATE_DECLARE_WIDGET_API(SMultiLineEditableText, SWidget, SLATE_API)
 public:
 
 	/** Used to merge multiple text edit transactions within a scope */
@@ -244,7 +245,6 @@ public:
 	SLATE_END_ARGS()
 
 	SLATE_API SMultiLineEditableText();
-	SLATE_API ~SMultiLineEditableText();
 
 	SLATE_API void Construct( const FArguments& InArgs );
 
@@ -355,7 +355,13 @@ public:
 	SLATE_API void SetIsReadOnly(const TAttribute< bool >& InIsReadOnly);
 
 	/** Get the number of Text Lines */
-	SLATE_API int32 GetTextLineCount();
+	UE_DEPRECATED(5.8, "Use the const version instead.")
+	int32 GetTextLineCount()
+	{
+		return const_cast<const SMultiLineEditableText*>(this)->GetTextLineCount();
+	}
+	/** Get the number of Text Lines */
+	SLATE_API int32 GetTextLineCount() const;
 
 	/**
 	 * Sets whether to select all text when the user clicks to give focus on the widget
@@ -557,7 +563,10 @@ protected:
 
 protected:
 	/** The text layout that deals with the editable text */
+	UE_DEPRECATED(5.8, "Use SlateEditableTextLayout, this variablle will never be set.")
 	TUniquePtr<FSlateEditableTextLayout> EditableTextLayout;
+	/** The text layout that deals with the editable text */
+	FSlateEditableTextLayout SlateEditableTextLayout;
 
 	/** Whether to allow multi-line text */
 	TAttribute<bool> bAllowMultiLine;
diff --git a/Engine/Source/Runtime/Slate/Public/Widgets/Text/SlateEditableTextLayout.h b/Engine/Source/Runtime/Slate/Public/Widgets/Text/SlateEditableTextLayout.h
--- a/Engine/Source/Runtime/Slate/Public/Widgets/Text/SlateEditableTextLayout.h
+++ b/Engine/Source/Runtime/Slate/Public/Widgets/Text/SlateEditableTextLayout.h
@@ -19,6 +19,7 @@
 #include "Widgets/Text/SlateEditableTextTypes.h"
 #include "Framework/Text/SlateTextLayoutFactory.h"
 #include "GenericPlatform/ITextInputMethodSystem.h"
+#include "Types/SlateNestedAttribute.h"
 
 class FArrangedChildren;
 class FExtender;
@@ -33,19 +34,85 @@ enum class ETextShapingMethod : uint8;
 enum class EFontFacesLoadingPaintPolicy : uint8;
 
 /** Class to handle the cached layout of SEditableText/SMultiLineEditableText by proxying around a FTextLayout */
-class FSlateEditableTextLayout
+class FSlateEditableTextLayout : public FSlateNestedAttributeContainer
 {
 public:
-	SLATE_API FSlateEditableTextLayout(ISlateEditableTextWidget& InOwnerWidget, const TAttribute<FText>& InInitialText, FTextBlockStyle InTextStyle, const TOptional<ETextShapingMethod> InTextShapingMethod, const TOptional<ETextFlowDirection> InTextFlowDirection, const FCreateSlateTextLayout& InCreateSlateTextLayout, TSharedRef<ITextLayoutMarshaller> InTextMarshaller, TSharedRef<ITextLayoutMarshaller> InHintTextMarshaller);
+	template<auto InPointerToAttributeOwner>
+	static void RegisterNestedAttributes(FSlateAttributeInitializer& AttributeInitializer)
+	{
+		SLATE_ADD_NESTED_MEMBER_ATTRIBUTE_DEFINITION(AttributeInitializer, InPointerToAttributeOwner, BoundText, EInvalidateWidgetReason::Layout)
+			.OnValueChanged_Static(OnBoundTextChanged);
+		SLATE_ADD_NESTED_MEMBER_ATTRIBUTE_DEFINITION(AttributeInitializer, InPointerToAttributeOwner, HintText, EInvalidateWidgetReason::Layout);
+		SLATE_ADD_NESTED_MEMBER_ATTRIBUTE_DEFINITION(AttributeInitializer, InPointerToAttributeOwner, HintTextOpacity, EInvalidateWidgetReason::Paint);
+		SLATE_ADD_NESTED_MEMBER_ATTRIBUTE_DEFINITION(AttributeInitializer, InPointerToAttributeOwner, BoundSearchText, EInvalidateWidgetReason::Layout)
+			.OnValueChanged_Static(OnBoundSearchTextChanged);
+
+		SLATE_ADD_NESTED_MEMBER_ATTRIBUTE_DEFINITION(AttributeInitializer, InPointerToAttributeOwner, WrapTextAt, EInvalidateWidgetReason::Layout);
+		SLATE_ADD_NESTED_MEMBER_ATTRIBUTE_DEFINITION(AttributeInitializer, InPointerToAttributeOwner, AutoWrapText, EInvalidateWidgetReason::Layout);
+		SLATE_ADD_NESTED_MEMBER_ATTRIBUTE_DEFINITION(AttributeInitializer, InPointerToAttributeOwner, WrappingPolicy, EInvalidateWidgetReason::Layout);
+
+		SLATE_ADD_NESTED_MEMBER_ATTRIBUTE_DEFINITION(AttributeInitializer, InPointerToAttributeOwner, Margin, EInvalidateWidgetReason::Layout);
+
+		SLATE_ADD_NESTED_MEMBER_ATTRIBUTE_DEFINITION(AttributeInitializer, InPointerToAttributeOwner, Justification, EInvalidateWidgetReason::Layout);
+		SLATE_ADD_NESTED_MEMBER_ATTRIBUTE_DEFINITION(AttributeInitializer, InPointerToAttributeOwner, LineHeightPercentage, EInvalidateWidgetReason::Layout);
+		SLATE_ADD_NESTED_MEMBER_ATTRIBUTE_DEFINITION(AttributeInitializer, InPointerToAttributeOwner, ApplyLineHeightToBottomLine, EInvalidateWidgetReason::Layout);
+	}
+
+	template<typename RealWidgetType>
+	UE_DEPRECATED(5.8, "Use the other base constructor and call Construct during widget construction")
+	FSlateEditableTextLayout(RealWidgetType& InOwnerWidget, TAttribute<FText> InInitialText, FTextBlockStyle InTextStyle, const TOptional<ETextShapingMethod> InTextShapingMethod, const TOptional<ETextFlowDirection> InTextFlowDirection, const FCreateSlateTextLayout& InCreateSlateTextLayout, TSharedRef<ITextLayoutMarshaller> InTextMarshaller, TSharedRef<ITextLayoutMarshaller> InHintTextMarshaller)
+		: FSlateEditableTextLayout(InOwnerWidget)
+	{
+		Construct(MoveTemp(InInitialText), MoveTemp(InTextStyle), InTextShapingMethod, InTextFlowDirection, InCreateSlateTextLayout, MoveTemp(InTextMarshaller), MoveTemp(InHintTextMarshaller));
+	}
+
+
+	template<typename RealWidgetType>
+	FSlateEditableTextLayout(RealWidgetType& InOwnerWidget)
+		: OwnerWidget(&InOwnerWidget)
+		, BoundText(InOwnerWidget, FText::GetEmpty())
+		, bWasPasswordLastTick(false)
+		, HintText(InOwnerWidget, FText::GetEmpty())
+		, HintTextOpacity(InOwnerWidget, 0.35f)
+		, BoundSearchText(InOwnerWidget, FText::GetEmpty())
+		, SearchCase(ESearchCase::IgnoreCase)
+		, CurrentSearchResultIndex(0)
+		, WrapTextAt(InOwnerWidget, 0.0f)
+		, AutoWrapText(InOwnerWidget, false)
+		, WrappingPolicy(InOwnerWidget, ETextWrappingPolicy::DefaultWrapping)
+		, Margin(InOwnerWidget, FMargin(0))
+		, Justification(InOwnerWidget, ETextJustify::Left)
+		, LineHeightPercentage(InOwnerWidget, 1.0f)
+		, ApplyLineHeightToBottomLine(InOwnerWidget, false)
+		, bHasRegisteredTextInputMethodContext(false)
+		, PreferredCursorScreenOffsetInLine(0)
+		, CurrentUndoLevel(0)
+		, NumTransactionsOpened(0)
+		, bIsDragSelecting(false)
+		, bWasFocusedByLastMouseDown(false)
+		, bHasDragSelectedSinceFocused(false)
+		, bTextChangedByVirtualKeyboard(false)
+		, bTextCommittedByVirtualKeyboard(false)
+		, VirtualKeyboardTextCommitType()
+		, bSelectionChangedExternally(false)
+		, ExternalSelectionStart(0)
+		, ExternalSelectionEnd(0)
+	{
+		static_assert(std::is_base_of_v<SWidget, RealWidgetType>, "RealWidgetType must be a Widget type.");
+		static_assert(std::is_base_of_v<ISlateEditableTextWidget, RealWidgetType>, "RealWidgetType must derive from ISlateEditableTextWidget.");
+	}
+
 	SLATE_API ~FSlateEditableTextLayout();
 
-	SLATE_API void SetText(const TAttribute<FText>& InText);
+	SLATE_API void Construct(TAttribute<FText> InInitialText, FTextBlockStyle InTextStyle, const TOptional<ETextShapingMethod>& InTextShapingMethod, const TOptional<ETextFlowDirection>& InTextFlowDirection, const FCreateSlateTextLayout& InCreateSlateTextLayout, TSharedRef<ITextLayoutMarshaller> InTextMarshaller, TSharedRef<ITextLayoutMarshaller> InHintTextMarshaller);
+
+	SLATE_API void SetText(TAttribute<FText> InText);
 	SLATE_API FText GetText() const;
 
-	SLATE_API void SetHintText(const TAttribute<FText>& InHintText);
+	SLATE_API void SetHintText(TAttribute<FText> InHintText);
 	SLATE_API FText GetHintText() const;
 
-	SLATE_API void SetSearchText(const TAttribute<FText>& InSearchText);
+	SLATE_API void SetSearchText(TAttribute<FText> InSearchText);
 	SLATE_API FText GetSearchText() const;
 	
 	/** Get the index of the search result (0 if none) */
@@ -57,11 +124,17 @@ public:
 	SLATE_API void SetTextStyle(const FTextBlockStyle& InTextStyle);
 	SLATE_API const FTextBlockStyle& GetTextStyle() const;
 	
-	SLATE_API void SetHintTextOpacity(const TAttribute<float>& InOpacity);
+	SLATE_API void SetHintTextOpacity(TAttribute<float> InOpacity);
 	SLATE_API float GetHintTextOpacity() const;
 
 	/** Get the number of Text Lines */
-	SLATE_API int32 GetTextLineCount();
+	UE_DEPRECATED(5.8, "Use the const version instead.")
+	int32 GetTextLineCount()
+	{
+		return const_cast<const FSlateEditableTextLayout*>(this)->GetTextLineCount();
+	}
+	/** Get the number of Text Lines */
+	SLATE_API int32 GetTextLineCount() const;
 
 	/** Set the brush to use when drawing the cursor */
 	SLATE_API void SetCursorBrush(const TAttribute<const FSlateBrush*>& InCursorBrush);
@@ -102,31 +175,33 @@ public:
 	SLATE_API void SetTextFlowDirection(const TOptional<ETextFlowDirection>& InTextFlowDirection);
 
 	/** Set the wrapping to use for this document */
-	SLATE_API void SetTextWrapping(const TAttribute<float>& InWrapTextAt, const TAttribute<bool>& InAutoWrapText, const TAttribute<ETextWrappingPolicy>& InWrappingPolicy);
+	SLATE_API void SetTextWrapping(TAttribute<float> InWrapTextAt, TAttribute<bool> InAutoWrapText, TAttribute<ETextWrappingPolicy> InWrappingPolicy);
 
 	/** Set whether text wraps onto a new line when it's length exceeds this width; if this value is zero or negative, no wrapping occurs */
-	SLATE_API void SetWrapTextAt(const TAttribute<float>& InWrapTextAt);
+	SLATE_API void SetWrapTextAt(TAttribute<float> InWrapTextAt);
 
 	/** Set whether to wrap text automatically based on the widget's computed horizontal space */
-	SLATE_API void SetAutoWrapText(const TAttribute<bool>& InAutoWrapText);
+	SLATE_API void SetAutoWrapText(TAttribute<bool> InAutoWrapText);
 
 	/** Set the wrapping policy to use */
-	SLATE_API void SetWrappingPolicy(const TAttribute<ETextWrappingPolicy>& InWrappingPolicy);
+	SLATE_API void SetWrappingPolicy(TAttribute<ETextWrappingPolicy> InWrappingPolicy);
 
 	/** Set the amount of blank space left around the edges of text area */
-	SLATE_API void SetMargin(const TAttribute<FMargin>& InMargin);
+	SLATE_API void SetMargin(TAttribute<FMargin> InMargin);
 
 	/** Set how the text should be aligned with the margin */
-	SLATE_API void SetJustification(const TAttribute<ETextJustify::Type>& InJustification);
+	SLATE_API void SetJustification(TAttribute<ETextJustify::Type> InJustification);
+	/** Get how the text should be aligned with the margin */
+	SLATE_API ETextJustify::Type GetJustification() const;
 
 	/** Set the amount to scale each lines height by */
-	SLATE_API void SetLineHeightPercentage(const TAttribute<float>& InLineHeightPercentage);
+	SLATE_API void SetLineHeightPercentage(TAttribute<float> InLineHeightPercentage);
 
 	/** Set whether to leave extra space below the last line due to line height */
-	SLATE_API void SetApplyLineHeightToBottomLine(const TAttribute<bool>& InApplyLineHeightToBottomLine);
+	SLATE_API void SetApplyLineHeightToBottomLine(TAttribute<bool> InApplyLineHeightToBottomLine);
 
 	/** Set the text overflow policy that should be used to determine what happens to clipped text */
-	SLATE_API void SetOverflowPolicy(TOptional<ETextOverflowPolicy> InOverflowPolicy);
+	SLATE_API void SetOverflowPolicy(const TOptional<ETextOverflowPolicy>& InOverflowPolicy);
 
 	/** Set the information used to help identify who owns this text layout in the case of an error */
 	SLATE_API void SetDebugSourceInfo(const TAttribute<FString>& InDebugSourceInfo);
@@ -140,7 +215,7 @@ public:
 	/** Register and activate the IME context for this text layout */
 	SLATE_API void EnableTextInputMethodContext();
 
-	/** Refresh this editable text immediately, rather than wait for the usual caching mechanisms to take affect on the text Tick */
+	/** Refresh this editable text immediately, rather than wait for the usual caching mechanisms to take effect on the text Tick */
 	SLATE_API bool Refresh();
 
 	/** Force the text layout to be updated from the marshaller */
@@ -311,7 +386,7 @@ public:
 	/**
 	 * Given a location and a Direction to offset, return a new location.
 	 *
-	 * @param Location    Cursor location from which to offset
+	 * @param CurrentLocation    Cursor location from which to offset
 	 * @param Direction   Positive means right, negative means left.
 	 */
 	SLATE_API FTextLocation TranslatedLocation(const FTextLocation& CurrentLocation, int8 Direction) const;
@@ -354,8 +429,13 @@ public:
 	/** Returns whether the current text varies from the original */
 	SLATE_API bool HasTextChangedFromOriginal() const;
 
+	UE_DEPRECATED(5.8, "Use the properly spelled version BeginEditTransaction")
+	void BeginEditTransation()
+	{
+		BeginEditTransaction();
+	}
 	/** Called to begin an undoable editable text transaction */
-	SLATE_API void BeginEditTransation();
+	SLATE_API void BeginEditTransaction();
 
 	/** Called to end an undoable editable text transaction */
 	SLATE_API void EndEditTransaction();
@@ -385,11 +465,17 @@ public:
 
 	SLATE_API void LoadText();
 
+	UE_DEPRECATED(5.8, "ComputeVolatility always returns false with the addition of TSlateNestedAttributes.")
 	SLATE_API bool ComputeVolatility() const;
 
 	SLATE_API void Tick(const FGeometry& AllottedGeometry, const double InCurrentTime, const float InDeltaTime);
 
-	SLATE_API int32 OnPaint(const FPaintArgs& Args, const FGeometry& AllottedGeometry, const FSlateRect& MyCullingRect, FSlateWindowElementList& OutDrawElements, int32 LayerId, const FWidgetStyle& InWidgetStyle, bool bParentEnabled);
+	UE_DEPRECATED(5.8, "Use the const version instead.")
+	int32 OnPaint(const FPaintArgs& Args, const FGeometry& AllottedGeometry, const FSlateRect& MyCullingRect, FSlateWindowElementList& OutDrawElements, int32 LayerId, const FWidgetStyle& InWidgetStyle, bool bParentEnabled)
+	{
+		return const_cast<const FSlateEditableTextLayout*>(this)->OnPaint(Args, AllottedGeometry, MyCullingRect, OutDrawElements, LayerId, InWidgetStyle, bParentEnabled);
+	}
+	SLATE_API int32 OnPaint(const FPaintArgs& Args, const FGeometry& AllottedGeometry, const FSlateRect& MyCullingRect, FSlateWindowElementList& OutDrawElements, int32 LayerId, const FWidgetStyle& InWidgetStyle, bool bParentEnabled) const;
 
 	SLATE_API void CacheDesiredSize(float LayoutScaleMultiplier);
 
@@ -429,6 +515,11 @@ public:
 	/** Sets the paint policy to use for text if the font faces are still async loading. */
 	SLATE_API void SetFontFacesLoadingPaintPolicy(EFontFacesLoadingPaintPolicy InFontFacesLoadingPaintPolicy);
 private:
+	/** Function called when the BoundText attribute changes. Will call the Widget's OnTextChanged */
+	SLATE_API static void OnBoundTextChanged(SWidget& Widget, FSlateEditableTextLayout& SlateTextLayout);
+	/** Function called when the BoundSearchText attribute changes */
+	SLATE_API static void OnBoundSearchTextChanged(SWidget& Widget, FSlateEditableTextLayout& SlateTextLayout);
+
 	/** Insert the given text at the current cursor position, correctly taking into account new line characters */
 	SLATE_API void InsertTextAtCursorImpl(const FString& InString);
 
@@ -556,10 +647,11 @@ private:
 	FTextBlockStyle TextStyle;
 
 	/** Style used to draw the hint text (only valid when HintTextLayout is set) */
-	FTextBlockStyle HintTextStyle;
+	mutable FTextBlockStyle HintTextStyle;
 
 	/** The text displayed in this text block */
-	TAttribute<FText> BoundText;
+	TSlateNestedAttribute<FText> BoundText;
+	FText PreviousBoundText;
 
 	/** The state of BoundText last Tick() (used to allow updates when the text is changed) */
 	FTextSnapshot BoundTextLastTick;
@@ -568,13 +660,13 @@ private:
 	bool bWasPasswordLastTick;
 
 	/** The text that appears when there is no text in the text box */
-	TAttribute<FText> HintText;
+	TSlateNestedAttribute<FText> HintText;
 	
-	/** Override opacity of the hint text */
-	TAttribute<float> HintTextOpacity;
+	/** Opacity of the hint text */
+	TSlateNestedAttribute<float> HintTextOpacity;
 
 	/** The text to be searched for */
-	TAttribute<FText> BoundSearchText;
+	TSlateNestedAttribute<FText> BoundSearchText;
 
 	/** The state of BoundSearchText last Tick() (used to allow updates when the text is changed) */
 	FTextSnapshot BoundSearchTextLastTick;
@@ -592,25 +684,25 @@ private:
 	int32 CurrentSearchResultIndex;
 
 	/** Whether text wraps onto a new line when it's length exceeds this width; if this value is zero or negative, no wrapping occurs. */
-	TAttribute<float> WrapTextAt;
+	TSlateNestedAttribute<float> WrapTextAt;
 
 	/** True if we're wrapping text automatically based on the computed horizontal space for this widget */
-	TAttribute<bool> AutoWrapText;
+	TSlateNestedAttribute<bool> AutoWrapText;
 
 	/** The wrapping policy we're using */
-	TAttribute<ETextWrappingPolicy> WrappingPolicy;
+	TSlateNestedAttribute<ETextWrappingPolicy> WrappingPolicy;
 
 	/** The amount of blank space left around the edges of text area */
-	TAttribute<FMargin> Margin;
+	TSlateNestedAttribute<FMargin> Margin;
 
 	/** How the text should be aligned with the margin */
-	TAttribute<ETextJustify::Type> Justification;
+	TSlateNestedAttribute<ETextJustify::Type> Justification;
 
 	/** The amount to scale each lines height by */
-	TAttribute<float> LineHeightPercentage;
+	TSlateNestedAttribute<float> LineHeightPercentage;
 
 	/** Whether to leave extra space below the last line due to line height */
-	TAttribute<bool> ApplyLineHeightToBottomLine;
+	TSlateNestedAttribute<bool> ApplyLineHeightToBottomLine;
 
 	/** The information used to help identify who owns this text layout in the case of an error */
 	TAttribute<FString> DebugSourceInfo;
@@ -663,7 +755,7 @@ private:
 	/** Current undo state level that we've rolled back to, or INDEX_NONE if we haven't undone. Used for 'Redo'. */
 	int32 CurrentUndoLevel;
 
-	/** Undo state that will be pushed if text is actually changed between calls to BeginEditTransation() and EndEditTransaction() */
+	/** Undo state that will be pushed if text is actually changed between calls to BeginEditTransaction() and EndEditTransaction() */
 	TOptional<SlateEditableTextTypes::FUndoState> StateBeforeChangingText;
 
 	/** Track the number transactions that are opened */
@@ -697,7 +789,7 @@ private:
 	TOptional<ETextOverflowPolicy> OverflowPolicyOverride;
 
 	/** The last known size of the widget from the previous OnPaint, used to recalculate wrapping */
-	FVector2f CachedSize;
+	mutable FVector2f CachedSize;
 
 	/** A list commands to execute if a user presses the corresponding key-binding in the text box */
 	TSharedPtr<FUICommandList> UICommandList;
@@ -708,6 +800,9 @@ private:
 	/** Whether the cursor position has been changed externally*/
 	bool bSelectionChangedExternally;
 
+	/** Whether the cursor position would have moved in OnBoundTextChanged */
+	bool bMoveCursorToDocumentEnd;
+
 	/** The boundaries of the external selection */
 	int ExternalSelectionStart;
 	int ExternalSelectionEnd;
diff --git a/Engine/Source/Runtime/SlateCore/Public/Types/Attributes/SlateAttributeNestedMember.inl b/Engine/Source/Runtime/SlateCore/Public/Types/Attributes/SlateAttributeNestedMember.inl
new file mode 100644
--- /dev/null
+++ b/Engine/Source/Runtime/SlateCore/Public/Types/Attributes/SlateAttributeNestedMember.inl
@@ -0,0 +1,50 @@
+// Copyright Epic Games, Inc. All Rights Reserved.
+
+#pragma once
+
+#include "Traits/MemberPtrOuter.h"
+#include "Traits/RemoveMemberPtr.h"
+
+/** */
+namespace SlateAttributePrivate
+{
+	/**
+	 * A Slate Attribute that can live in a class/struct living in a widget.
+	 * ex: Widget -> FAttributeHolder -> TSlateNestedMemberAttribute
+	 */
+	template<typename InObjectType, typename InInvalidationReasonPredicate, typename InComparePredicate>
+	struct TSlateNestedMemberAttribute : public TSlateMemberAttribute<InObjectType, InInvalidationReasonPredicate, InComparePredicate>
+	{
+		/**
+		 * A Nested Attribute is just a TSlateMemberAttribute with an offset carefully computed.
+		 * This function checks that the passed in Pointer to Members are what we expect for a nested attribute.
+		 * See SLATE_ADD_NESTED_MEMBER_ATTRIBUTE_DEFINITION
+		 *
+		 * @tparam PointerToAttributeOwnerType The pointer to member from the Widget to the class/struct owning the nested attribute
+		 * @tparam PointerToAttributeType The pointer to member from the class/struct owning the nested attribute to the nested attribute
+		 * @return True if valid, otherwise a static_assert would have been raised
+		 */
+		template<
+			typename PointerToAttributeOwnerType,
+			typename PointerToAttributeType
+		>
+		static constexpr bool IsNestedAttributeValid()
+		{
+			static_assert(std::is_member_pointer_v<PointerToAttributeOwnerType>, "The PointerToAttributeOwnerType needs to be a Pointer to Member.");
+			static_assert(std::is_member_pointer_v<PointerToAttributeType>, "The PointerToAttributeType needs to be a Pointer to Member.");
+
+			using WidgetType = TMemberPtrOuter_T<PointerToAttributeOwnerType>;
+			using AttributeOwnerType = TRemoveMemberPtr_T<PointerToAttributeOwnerType>;
+			using AttributeContainerType = TMemberPtrOuter_T<PointerToAttributeType>;
+			using AttributeType = TRemoveMemberPtr_T<PointerToAttributeType>;
+
+			static_assert(std::is_base_of<SWidget, WidgetType>::value, "The class owning the PointerToAttributeOwner needs to derive from SWidget.");
+			static_assert(std::is_same<AttributeOwnerType, AttributeContainerType>::value, "The Attribute Owner needs to live inside the widget and needs to own the Attribute.");
+			static_assert(std::is_base_of<TSlateNestedMemberAttribute, AttributeType>::value, "The Attribute needs to be a TSlateNestedMemberAttribute.");
+			return true;
+		}
+
+		using TSlateMemberAttribute<InObjectType, InInvalidationReasonPredicate, InComparePredicate>::TSlateMemberAttribute;
+	};
+
+} // SlateAttributePrivate
diff --git a/Engine/Source/Runtime/SlateCore/Public/Types/SlateAttribute.h b/Engine/Source/Runtime/SlateCore/Public/Types/SlateAttribute.h
--- a/Engine/Source/Runtime/SlateCore/Public/Types/SlateAttribute.h
+++ b/Engine/Source/Runtime/SlateCore/Public/Types/SlateAttribute.h
@@ -277,4 +277,5 @@ struct TSlateDeprecatedTAttribute
 #include "Types/Attributes/SlateAttributeContained.inl"
 #include "Types/Attributes/SlateAttributeManaged.inl"
 #include "Types/Attributes/SlateAttributeMember.inl"
+#include "Types/Attributes/SlateAttributeNestedMember.inl"
 // IWYU pragma: end_exports
diff --git a/Engine/Source/Runtime/SlateCore/Public/Types/SlateAttributeDescriptor.h b/Engine/Source/Runtime/SlateCore/Public/Types/SlateAttributeDescriptor.h
--- a/Engine/Source/Runtime/SlateCore/Public/Types/SlateAttributeDescriptor.h
+++ b/Engine/Source/Runtime/SlateCore/Public/Types/SlateAttributeDescriptor.h
@@ -13,6 +13,9 @@
 #include "UObject/NameTypes.h"
 #include "Widgets/InvalidateWidgetReason.h"
 #include "Templates/Identity.h"
+#include "Traits/FunctionParameterType.h"
+#include "Traits/MemberPtrOuter.h"
+#include "Traits/RemoveMemberPtr.h"
 
 #include <type_traits>
 
@@ -22,6 +25,17 @@ class SWidget;
 namespace SlateAttributePrivate
 {
 	enum class ESlateAttributeType : uint8;
+
+	/**
+	 * Find the offset of a member in its class based on a Pointer to Member variable
+	 */
+	template <typename InClassType, typename InMemberType>
+	std::size_t OffsetOfPointerToMember(InMemberType InClassType::* Member)
+	{
+		const InClassType* Object = reinterpret_cast<const InClassType*>(alignof(InClassType));
+		return reinterpret_cast<std::size_t>(&(Object->*Member))
+			- reinterpret_cast<std::size_t>(Object);;
+	}
 }
 
 
@@ -71,6 +85,12 @@ public:
 		{
 		}
 
+		explicit FInvalidateWidgetReasonAttribute(FGetter InGetter)
+			: Reason(EInvalidateWidgetReason::None)
+			, Getter(MoveTemp(InGetter))
+		{
+		}
+
 		bool IsBound() const
 		{
 			return Getter.IsBound();
@@ -86,9 +106,119 @@ public:
 		FGetter Getter;
 	};
 
+	/**
+	 * Adjusted version of FInvalidateWidgetReasonAttribute for Nested Attributes.
+	 * The only difference is the signature of the getter, as we need to pass the PointerToAttributeOwner for a nested attribute
+	 */
+	template<typename PointerToAttributeOwnerType>
+	struct TInvalidateWidgetReasonNestedAttribute : FInvalidateWidgetReasonAttribute
+	{
+		static_assert(std::is_member_pointer_v<PointerToAttributeOwnerType>, "PointerToAttributeOwnerType needs to be a Pointer to Member."); \
+		using WidgetType = TMemberPtrOuter_T<PointerToAttributeOwnerType>;
+		using AttributeOwnerType = TRemoveMemberPtr_T<PointerToAttributeOwnerType>;
+
+		// using "not checked" delegate to disable race detection. This delegate is used in a static (FSlateWidgetClassData) and will be destructed after the tls array used in the race detection code (which will result in a use after free).
+		using FNestedGetter = TDelegate<EInvalidateWidgetReason(const SWidget&, const AttributeOwnerType&), FNotThreadSafeNotCheckedDelegateUserPolicy>;
+		using FNestedGetterTyped = TDelegate<EInvalidateWidgetReason(const WidgetType&, const AttributeOwnerType&), FNotThreadSafeNotCheckedDelegateUserPolicy>;
+
+
+		TInvalidateWidgetReasonNestedAttribute(const TInvalidateWidgetReasonNestedAttribute&) = default;
+		TInvalidateWidgetReasonNestedAttribute(TInvalidateWidgetReasonNestedAttribute&&) = default;
+		TInvalidateWidgetReasonNestedAttribute& operator=(const TInvalidateWidgetReasonNestedAttribute&) = default;
+		TInvalidateWidgetReasonNestedAttribute& operator=(TInvalidateWidgetReasonNestedAttribute&&) = default;
+
+		/** Default constructor. */
+		using FInvalidateWidgetReasonAttribute::FInvalidateWidgetReasonAttribute;
+
+		explicit TInvalidateWidgetReasonNestedAttribute(PointerToAttributeOwnerType InPointerToAttributeOwner, EInvalidateWidgetReason InReason)
+			: FInvalidateWidgetReasonAttribute(InReason)
+		{
+		}
+
+		template<typename... PayloadTypes>
+		explicit TInvalidateWidgetReasonNestedAttribute(PointerToAttributeOwnerType InPointerToAttributeOwner, TIdentity_T<typename FNestedGetter::template TFuncPtr<PayloadTypes...>> InFuncPtr, PayloadTypes&&... InputPayload)
+			: FInvalidateWidgetReasonAttribute(
+				FGetter::CreateStatic(
+					InvalidateWidgetReasonConverter,
+					FNestedGetter::CreateStatic(InFuncPtr, Forward<PayloadTypes>(InputPayload)...),
+					InPointerToAttributeOwner
+				))
+		{
+		}
+
+		template<typename... PayloadTypes, typename U = std::enable_if<!std::is_same_v<WidgetType, SWidget>>::type>
+		explicit TInvalidateWidgetReasonNestedAttribute(PointerToAttributeOwnerType InPointerToAttributeOwner, TIdentity_T<typename FNestedGetterTyped::template TFuncPtr<PayloadTypes...>> InFuncPtr, PayloadTypes&&... InputPayload)
+			: FInvalidateWidgetReasonAttribute(
+				FGetter::CreateStatic(
+					InvalidateWidgetReasonConverter,
+					FNestedGetterTyped::CreateStatic(InFuncPtr, Forward<PayloadTypes>(InputPayload)...),
+					InPointerToAttributeOwner
+				))
+		{
+		}
+
+		template<typename LambdaType, typename... PayloadTypes>
+		explicit TInvalidateWidgetReasonNestedAttribute(PointerToAttributeOwnerType InPointerToAttributeOwner, LambdaType&& InCallable, PayloadTypes&&... InputPayload)
+			: FInvalidateWidgetReasonAttribute(
+				CreateNestedGetterLambda(InPointerToAttributeOwner, Forward<LambdaType>(InCallable), Forward<PayloadTypes>(InputPayload)...)
+			)
+		{
+		}
+
+	private:
+		template<typename LambdaType, typename... PayloadTypes>
+		FGetter CreateNestedGetterLambda(PointerToAttributeOwnerType InPointerToAttributeOwner, LambdaType&& InCallable, PayloadTypes&&... InputPayload)
+		{
+			using LambdaOperatorType = decltype(&LambdaType::operator());
+			static_assert(std::is_same_v<TFunctionParameterType_T<TRemoveMemberPtr_T<LambdaOperatorType>, 1>, const AttributeOwnerType&>, "The second argument of the lambda needs to be 'const AttributeOwnerType&'");
+
+			if constexpr (std::is_same_v<TFunctionParameterType_T<TRemoveMemberPtr_T<LambdaOperatorType>, 0>, const SWidget&>)
+			{
+				return FGetter::CreateStatic(
+					InvalidateWidgetReasonConverter,
+					FNestedGetter::CreateLambda(Forward<LambdaType>(InCallable), Forward<PayloadTypes>(InputPayload)...),
+					InPointerToAttributeOwner
+				);
+			}
+			else
+			{
+				static_assert(std::is_same_v<TFunctionParameterType_T<TRemoveMemberPtr_T<LambdaOperatorType>, 0>, const WidgetType&>, "The first argument of the lambda needs to be either 'const SWidget&' or 'const WidgetType&'");
+				return FGetter::CreateStatic(
+					InvalidateWidgetReasonTypedConverter,
+					FNestedGetterTyped::CreateLambda(Forward<LambdaType>(InCallable), Forward<PayloadTypes>(InputPayload)...),
+					InPointerToAttributeOwner
+				);
+			}
+		}
+
+		static EInvalidateWidgetReason InvalidateWidgetReasonConverter(const SWidget& Widget, FNestedGetter Getter, PointerToAttributeOwnerType InPointerToAttributeOwner)
+		{
+			return Getter.Execute(Widget, static_cast<const WidgetType&>(Widget).*InPointerToAttributeOwner);
+		}
+		static EInvalidateWidgetReason InvalidateWidgetReasonTypedConverter(const SWidget& Widget, FNestedGetterTyped Getter, PointerToAttributeOwnerType InPointerToAttributeOwner)
+		{
+			return Getter.Execute(static_cast<const WidgetType&>(Widget), static_cast<const WidgetType&>(Widget).*InPointerToAttributeOwner);
+		}
+	};
+
 	// using "not checked" delegate to disable race detection. This delegate is used in a static (FSlateWidgetClassData) and will be destructed after the tls array used in the race detection code (which will result in a use after free).
 	using FAttributeValueChangedDelegate = TDelegate<void(SWidget&), FNotThreadSafeNotCheckedDelegateUserPolicy>;
 
+	/**
+	 * A delegate similar to FAttributeValueChangedDelegate but for nested attributes.
+	 * The main difference is that a reference to the object owning the nested attribute is also passed as an argument.
+	 */
+	template<typename PointerToAttributeOwnerType>
+	using TNestedAttributeValueChangedDelegate = TDelegate<void(SWidget&, TRemoveMemberPtr_T<PointerToAttributeOwnerType>&), FNotThreadSafeNotCheckedDelegateUserPolicy>;
+	/**
+	 * A delegate similar to FAttributeValueChangedDelegate but for nested attributes.
+	 * The main difference is that a reference to the object owning the nested attribute is also passed as an argument.
+	 * This templated version already converts the owning widget to the right type.
+	 */
+	template<typename PointerToAttributeOwnerType>
+	using TNestedAttributeValueChangedDelegateTyped = TDelegate<void(TMemberPtrOuter_T<PointerToAttributeOwnerType>&, TRemoveMemberPtr_T<PointerToAttributeOwnerType>&), FNotThreadSafeNotCheckedDelegateUserPolicy>;
+
+
 	/** */
 	enum class ECallbackOverrideType
 	{
@@ -303,9 +433,125 @@ public:
 			int32 AttributeIndex;
 		};
 
+		/**
+		 * Adjusted version of FAttributeEntry for Nested Attributes.
+		 * The only difference is the signature of OnValueChanged, as we need to pass the PointerToAttributeOwner for a nested attribute
+		 */
+		template<typename PointerToAttributeOwnerType>
+		struct TNestedAttributeEntry
+		{
+			static_assert(std::is_member_pointer_v<PointerToAttributeOwnerType>, "PointerToAttributeOwnerType needs to be a Pointer to Member."); \
+			using WidgetType = TMemberPtrOuter_T<PointerToAttributeOwnerType>;
+			using AttributeOwnerType = TRemoveMemberPtr_T<PointerToAttributeOwnerType>;
+
+			TNestedAttributeEntry(FAttributeEntry&& InEntry, PointerToAttributeOwnerType InPointerToAttributeOwner)
+				: Entry(MoveTemp(InEntry))
+				, PointerToAttributeOwner(InPointerToAttributeOwner)
+			{
+			}
+
+			/**
+			 * Update the attribute after the prerequisite.
+			 * The order is guaranteed but other attributes may be updated in between.
+			 * No order is guaranteed if the prerequisite or this property is updated manually.
+			 */
+			TNestedAttributeEntry& UpdatePrerequisite(FName Prerequisite)
+			{
+				Entry.UpdatePrerequisite(Prerequisite);
+				return *this;
+			}
+
+			/**
+			 * The attribute affect the visibility of the widget.
+			 * We only update the attributes that can change the visibility of the widget when the widget is collapsed.
+			 * Attributes that affect visibility must have the Visibility attribute as a Prerequisite or the Visibility attribute must have it as a Prerequisite.
+			 */
+			TNestedAttributeEntry& AffectVisibility()
+			{
+				Entry.AffectVisibility();
+				return *this;
+			}
+
+			/**
+			 * Notified when the attribute value changed.
+			 * It's preferable that you delay any action to the Tick or Paint function.
+			 * You are not allowed to make changes that would affect the SWidget ChildOrder or its Visibility.
+			 * It will not be called when the SWidget is in its construction phase.
+			 * @see SWidget::IsConstructed
+			 */
+			TNestedAttributeEntry& OnValueChanged(TNestedAttributeValueChangedDelegate<PointerToAttributeOwnerType> Callback)
+			{
+				if (Callback.IsBound())
+				{
+					Entry.OnValueChanged(FAttributeValueChangedDelegate::CreateStatic(OnValueChangedConverter, MoveTemp(Callback), PointerToAttributeOwner));
+				}
+				return *this;
+			}
+			TNestedAttributeEntry& OnValueChanged(TNestedAttributeValueChangedDelegateTyped<PointerToAttributeOwnerType> Callback)
+			{
+				if (Callback.IsBound())
+				{
+					Entry.OnValueChanged(FAttributeValueChangedDelegate::CreateStatic(OnValueChangedConverter, MoveTemp(Callback), PointerToAttributeOwner));
+				}
+				return *this;
+			}
+
+			template<typename... VarTypes>
+			TNestedAttributeEntry& OnValueChanged_Static(TIdentity_T<typename TNestedAttributeValueChangedDelegate<PointerToAttributeOwnerType>::template TFuncPtr<VarTypes...>> InFunc, VarTypes... Vars)
+			{
+				return OnValueChanged(TNestedAttributeValueChangedDelegate<PointerToAttributeOwnerType>::CreateStatic(InFunc, Vars...));
+			}
+			template<typename... VarTypes, typename U = std::enable_if<!std::is_same_v<WidgetType, SWidget>>::type>
+			TNestedAttributeEntry& OnValueChanged_Static(TIdentity_T<typename TNestedAttributeValueChangedDelegateTyped<PointerToAttributeOwnerType>::template TFuncPtr<VarTypes...>> InFunc, VarTypes... Vars)
+			{
+				return OnValueChanged(TNestedAttributeValueChangedDelegateTyped<PointerToAttributeOwnerType>::CreateStatic(InFunc, Vars...));
+			}
+			template<typename LambdaType>
+			TNestedAttributeEntry& OnValueChanged_Lambda(LambdaType&& InFunctor)
+			{
+				using LambdaOperatorType = decltype(&LambdaType::operator());
+				static_assert(std::is_same_v<TFunctionParameterType_T<TRemoveMemberPtr_T<LambdaOperatorType>, 1>, AttributeOwnerType&>, "The second argument of the lambda needs to be 'AttributeOwnerType&'");
+
+				if constexpr (std::is_same_v<TFunctionParameterType_T<TRemoveMemberPtr_T<LambdaOperatorType>, 0>, SWidget&>)
+				{
+					return OnValueChanged(TNestedAttributeValueChangedDelegate<PointerToAttributeOwnerType>::CreateLambda(MoveTemp(InFunctor)));
+				}
+				else
+				{
+					static_assert(std::is_same_v<TFunctionParameterType_T<TRemoveMemberPtr_T<LambdaOperatorType>, 0>, WidgetType&>, "The first argument of the lambda needs to be either 'SWidget&' or 'WidgetType&'");
+					return OnValueChanged(TNestedAttributeValueChangedDelegateTyped<PointerToAttributeOwnerType>::CreateLambda(MoveTemp(InFunctor)));
+				}
+			}
+
+		private:
+			static void OnValueChangedConverter(SWidget& Widget, TNestedAttributeValueChangedDelegate<PointerToAttributeOwnerType> InCallback, PointerToAttributeOwnerType InPointerToAttributeOwner)
+			{
+				InCallback.ExecuteIfBound(Widget, static_cast<WidgetType&>(Widget).*InPointerToAttributeOwner);
+			}
+			static void OnValueChangedConverter(SWidget& Widget, TNestedAttributeValueChangedDelegateTyped<PointerToAttributeOwnerType> InCallback, PointerToAttributeOwnerType InPointerToAttributeOwner)
+			{
+				InCallback.ExecuteIfBound(static_cast<WidgetType&>(Widget), static_cast<WidgetType&>(Widget).*InPointerToAttributeOwner);
+			}
+
+		private:
+			FAttributeEntry Entry;
+			PointerToAttributeOwnerType PointerToAttributeOwner;
+		};
+
 		SLATECORE_API FAttributeEntry AddMemberAttribute(FName AttributeName, OffsetType Offset, const FInvalidateWidgetReasonAttribute& ReasonGetter);
 		SLATECORE_API FAttributeEntry AddMemberAttribute(FName AttributeName, OffsetType Offset, FInvalidateWidgetReasonAttribute&& ReasonGetter);
 
+		template<typename PointerToAttributeOwnerType>
+		TNestedAttributeEntry<PointerToAttributeOwnerType> AddNestedMemberAttribute(PointerToAttributeOwnerType InPointerToAttributeOwner, FName AttributeName, OffsetType Offset, const FInvalidateWidgetReasonAttribute& ReasonGetter)
+		{
+			return TNestedAttributeEntry<PointerToAttributeOwnerType>(AddMemberAttribute(AttributeName, Offset, ReasonGetter), InPointerToAttributeOwner);
+		}
+		template<typename PointerToAttributeOwnerType>
+		TNestedAttributeEntry<PointerToAttributeOwnerType> AddNestedMemberAttribute(PointerToAttributeOwnerType InPointerToAttributeOwner, FName AttributeName, OffsetType Offset, FInvalidateWidgetReasonAttribute&& ReasonGetter)
+		{
+			return TNestedAttributeEntry<PointerToAttributeOwnerType>(AddMemberAttribute(AttributeName, Offset, MoveTemp(ReasonGetter)), InPointerToAttributeOwner);
+		}
+
 		SLATECORE_API FContainerInitializer AddContainer(FName ContainerName, OffsetType Offset);
 
 	public:
@@ -388,6 +634,29 @@ private:
 #define SLATE_ADD_MEMBER_ATTRIBUTE_DEFINITION(_Initializer, _Property, _Reason) \
 	SLATE_ADD_MEMBER_ATTRIBUTE_DEFINITION_WITH_NAME(_Initializer, GET_MEMBER_NAME_CHECKED(PrivateThisType, _Property), _Property, _Reason)
 
+/**
+ * Add a TSlateAttribute nested inside a member of the widget to the descriptor.
+ * @param _Initializer The FSlateAttributeInitializer from the PrivateRegisterAttributes function.
+ * @param _PointerToAttributeOwner	The Pointer to Member of the widget property owning _Property.
+ *									The widget's type is referred below to 'WidgetType' and the member type (owning the property) is referred to 'AttributeOwnerType'
+ * @param _Property The TSlateAttribute property
+ * @param _Reason The EInvalidationWidgetReason or a static function/lambda that takes (const SWidget&, const AttributeOwnerType&) or (const WidgetType&, const AttributeOwnerType&) as arguments, and that returns the invalidation reason.
+ */
+#define SLATE_ADD_NESTED_MEMBER_ATTRIBUTE_DEFINITION_WITH_NAME(_Initializer, _Name, _PointerToAttributeOwner, _Property, _Reason) \
+		static_assert(std::is_member_pointer_v<decltype(_PointerToAttributeOwner)>, "The _PointerToAttributeOwner variable needs to be a Pointer to Member."); \
+		{ \
+			using __AttributeOwnerType = TRemoveMemberPtr_T<decltype(_PointerToAttributeOwner)>; \
+			static_assert(decltype(__AttributeOwnerType::_Property)::template IsNestedAttributeValid<decltype(_PointerToAttributeOwner), decltype(&__AttributeOwnerType::_Property)>(), "_PointerToAttributeOwner needs to be defined inside a widget, and _Property needs to be defined inside _PointerToAttributeOwner"); \
+			static_assert(decltype(__AttributeOwnerType::_Property)::AttributeType == SlateAttributePrivate::ESlateAttributeType::Member, "The SlateProperty is not a TSlateNestedAttribute. Use SLATE_ADD_NESTED_MEMBER_ATTRIBUTE_DEFINITION"); \
+			static_assert(!decltype(__AttributeOwnerType::_Property)::HasDefinedInvalidationReason, "When implementing the SLATE_DECLARE_WIDGET pattern, use TSlateAttribute without the invalidation reason."); \
+			using __ReasonType = decltype(_Reason); \
+			static_assert(!std::is_same_v<__ReasonType, EInvalidateWidgetReason> || FSlateAttributeBase::IsInvalidateWidgetReasonSupported(_Reason), "The invalidation is not supported by the SlateAttribute system."); \
+		}\
+		_Initializer.AddNestedMemberAttribute(_PointerToAttributeOwner, _Name, SlateAttributePrivate::OffsetOfPointerToMember(_PointerToAttributeOwner) + STRUCT_OFFSET(TRemoveMemberPtr_T<decltype(_PointerToAttributeOwner)>, _Property), FSlateAttributeDescriptor::TInvalidateWidgetReasonNestedAttribute{_PointerToAttributeOwner, _Reason})
+
+#define SLATE_ADD_NESTED_MEMBER_ATTRIBUTE_DEFINITION(_Initializer, _PointerToAttributeOwner, _Property, _Reason) \
+	SLATE_ADD_NESTED_MEMBER_ATTRIBUTE_DEFINITION_WITH_NAME(_Initializer, GET_MEMBER_NAME_CHECKED(TRemoveMemberPtr_T<decltype(_PointerToAttributeOwner)>, _Property), _PointerToAttributeOwner, _Property, _Reason)
+
 #define SLATE_ADD_PANELCHILDREN_DEFINITION_WITH_NAME(_Initializer, _Name, _Container) \
 		_Initializer.AddContainer(_Name, STRUCT_OFFSET(PrivateThisType, _Container))
 
diff --git a/Engine/Source/Runtime/SlateCore/Public/Types/SlateNestedAttribute.h b/Engine/Source/Runtime/SlateCore/Public/Types/SlateNestedAttribute.h
new file mode 100644
--- /dev/null
+++ b/Engine/Source/Runtime/SlateCore/Public/Types/SlateNestedAttribute.h
@@ -0,0 +1,169 @@
+﻿// Copyright Epic Games, Inc. All Rights Reserved.
+
+#pragma once
+
+#include "SlateAttribute.h"
+#include "Attributes/SlateAttributeDefinition.inl"
+#include "Attributes/SlateAttributeNestedMember.inl"
+
+
+/**
+ * A base type that defines TSlateNestedAttributes.
+ * These nested Attributes are similar to regular TSlateAttributes but can be defined outside a widget as long as
+ * the type owning the TSlateNestedAttributes is inside a widget.
+ *
+ * Example implementation:
+ *  @code
+// ------ .h
+
+struct SWidgetWithNestedAttributes;
+
+struct FSimpleNestedAttributesOwner : public FSlateNestedAttributeContainer // Define the type owning our attributes
+{
+	// To register the Nested Attributes, we need a pointer to member from the widget to this struct
+	static void RegisterNestedAttributes(FSlateAttributeInitializer& AttributeInitializer, FSimpleNestedAttributesOwner SWidgetWithNestedAttributes::* InPointerToOwner);
+
+	FSimpleNestedAttributesOwner(SWidgetWithNestedAttributes& Owner);
+
+	void SetNestedIntAttribute(TAttribute<int> Attr);
+	int GetNestedIntAttr() const { return NestedIntAttribute.Get(); }
+	void SetNestedFloatAttribute(TAttribute<float> Attr);
+	float GetNestedFloatAttr() const { return NestedFloatAttribute.Get(); }
+
+private:
+	// Callbacks given to OnValueChanged. The first argument can either be SWidget& or the type of the widget (here SWidgetWithNestedAttributes&)
+	static void OnNestedIntAttributeChanged(SWidget& OwningWidget, FSimpleNestedAttributesOwner& NestedAttributesOwner);
+	static void OnNestedFloatAttributeChanged(SWidgetWithNestedAttributes& OwningWidget, FSimpleNestedAttributesOwner& NestedAttributesOwner);
+
+	SWidgetWithNestedAttributes& Owner; // We need to keep track of the owning widget to call Assign on the attribute
+	TSlateNestedAttribute<int> NestedIntAttribute; // <--- The nested attribute
+	TSlateNestedAttribute<float> NestedFloatAttribute; // <--- The nested attribute
+};
+
+struct SWidgetWithNestedAttributes : public SCompoundWidget
+{
+	SLATE_DECLARE_WIDGET(SWidgetWithNestedAttributes, SWidget)
+public:
+	SLATE_BEGIN_ARGS( SWidgetWithNestedAttributes )
+		{}
+	SLATE_END_ARGS()
+
+	SWidgetWithNestedAttributes()
+		: NestedAttributes(*this)
+	{}
+
+	void Construct( const FArguments& InArgs );
+
+	void SetNestedIntAttribute(TAttribute<int> Attr) { NestedAttributes.SetNestedIntAttribute(MoveTemp(Attr)); }
+	void SetNestedFloatAttribute(TAttribute<float> Attr) { NestedAttributes.SetNestedFloatAttribute(MoveTemp(Attr)); }
+	int GetNestedIntAttribute() const { return NestedAttributes.GetNestedIntAttr(); }
+	float GetNestedFloatAttribute() const { return NestedAttributes.GetNestedFloatAttr(); }
+private:
+	FSimpleNestedAttributesOwner NestedAttributes; // <-- The widget has a field to our custom struct that contains the attributes
+};
+
+
+// ------ .cpp
+
+void FSimpleNestedAttributesOwner::RegisterNestedAttributes(FSlateAttributeInitializer& AttributeInitializer, FSimpleNestedAttributesOwner SWidgetWithNestedAttributes::* InPointerToOwner)
+{
+	// Here we define the nested attributes. We need to know which member of the widget owns this type
+
+	SLATE_ADD_NESTED_MEMBER_ATTRIBUTE_DEFINITION(AttributeInitializer, InPointerToOwner, NestedIntAttribute, EInvalidateWidgetReason::Paint)
+		.OnValueChanged_Static(OnNestedIntAttributeChanged);
+
+	// The Invalidation reason can also be a lambda. The first argument can either be const SWidget& or the type of the widget (here const SWidgetWithNestedAttributes&)
+	SLATE_ADD_NESTED_MEMBER_ATTRIBUTE_DEFINITION(AttributeInitializer, InPointerToOwner, NestedFloatAttribute,
+		[](const SWidgetWithNestedAttributes& OwningWidget, const FSimpleNestedAttributesOwner& NestedAttributes)
+	{
+		UE_LOG(LogTemp, Warning, TEXT("In Get Invalidation Reason: value changed to: %f   (direct access:  %f)"), OwningWidget.GetNestedFloatAttribute(), NestedAttributes.GetNestedFloatAttr())
+		return EInvalidateWidgetReason::Paint;
+	})
+		.OnValueChanged_Static(OnNestedFloatAttributeChanged);
+}
+
+FSimpleNestedAttributesOwner::FSimpleNestedAttributesOwner(SWidgetWithNestedAttributes& Owner)
+	: Owner(Owner)
+	, NestedIntAttribute(Owner)
+	, NestedFloatAttribute(Owner)
+{}
+
+void FSimpleNestedAttributesOwner::SetNestedIntAttribute(TAttribute<int> Attr)
+{
+	NestedIntAttribute.Assign(Owner, MoveTemp(Attr));
+}
+
+void FSimpleNestedAttributesOwner::SetNestedFloatAttribute(TAttribute<float> Attr)
+{
+	NestedFloatAttribute.Assign(Owner, MoveTemp(Attr));
+}
+
+void FSimpleNestedAttributesOwner::OnNestedIntAttributeChanged(SWidget& OwningWidget, FSimpleNestedAttributesOwner& NestedAttributesOwner)
+{
+	// The OnValueChanged definition is slightly different from a regular TSlateAttribute, also providing with a reference to our container type
+
+	SWidgetWithNestedAttributes& Widget = static_cast<SWidgetWithNestedAttributes&>(OwningWidget);
+	UE_LOG(LogTemp, Warning, TEXT("OnNestedIntAttributeChanged: value changed to: %d   (direct access:  %d)"), Widget.GetNestedIntAttribute(), NestedAttributesOwner.GetNestedIntAttr())
+}
+
+void FSimpleNestedAttributesOwner::OnNestedFloatAttributeChanged(SWidgetWithNestedAttributes& OwningWidget, FSimpleNestedAttributesOwner& NestedAttributesOwner)
+{
+	UE_LOG(LogTemp, Warning, TEXT("OnNestedFloatAttributeChanged: value changed to: %f   (direct access:  %f)"), OwningWidget.GetNestedFloatAttribute(), NestedAttributesOwner.GetNestedFloatAttr())
+}
+
+
+SLATE_IMPLEMENT_WIDGET(SWidgetWithNestedAttributes)
+
+void SWidgetWithNestedAttributes::PrivateRegisterAttributes(FSlateAttributeInitializer& AttributeInitializer)
+{
+	// We need to register the nested attributes when registering the widget's attributes, passing the pointer to member
+	FSimpleNestedAttributesOwner::RegisterNestedAttributes(AttributeInitializer, &SWidgetWithNestedAttributes::NestedAttributes);
+}
+
+void SWidgetWithNestedAttributes::Construct(const FArguments& InArgs)
+{
+	// as needed
+}
+ * @endcode
+ */
+class FSlateNestedAttributeContainer
+{
+protected:
+
+	/**
+	 * A SlateAttribute that is member variable of a member variable of a SWidget.
+	 * @usage: TSlateNestedAttribute<int32> MyAttribute1; TSlateNestedAttribute<int32, EInvalidateWidgetReason::Paint> MyAttribute2; TSlateNestedAttribute<int32, EInvalidateWidgetReason::Paint, TSlateAttributeComparePredicate<>> MyAttribute3;
+	 */
+	template<typename InObjectType, EInvalidateWidgetReason InInvalidationReasonValue = EInvalidateWidgetReason::None, typename InComparePredicate = TSlateAttributeComparePredicate<>>
+	struct TSlateNestedAttribute : public SlateAttributePrivate::TSlateNestedMemberAttribute<
+			InObjectType,
+			typename std::conditional<InInvalidationReasonValue == EInvalidateWidgetReason::None, ::SlateAttributePrivate::FSlateAttributeNoInvalidationReason, TSlateAttributeInvalidationReason<InInvalidationReasonValue>>::type,
+			InComparePredicate>
+	{
+		using ::SlateAttributePrivate::TSlateNestedMemberAttribute<
+			InObjectType,
+			typename std::conditional<InInvalidationReasonValue == EInvalidateWidgetReason::None, ::SlateAttributePrivate::FSlateAttributeNoInvalidationReason, TSlateAttributeInvalidationReason<InInvalidationReasonValue>>::type,
+			InComparePredicate>::TSlateNestedMemberAttribute;
+	};
+
+	//~ Override for FText that use the TSlateAttributeFTextComparePredicate to compare FText
+	template<>
+	struct TSlateNestedAttribute<FText, EInvalidateWidgetReason::None> : public ::SlateAttributePrivate::TSlateNestedMemberAttribute<FText, ::SlateAttributePrivate::FSlateAttributeNoInvalidationReason, TSlateAttributeFTextComparePredicate>
+	{
+		using ::SlateAttributePrivate::TSlateNestedMemberAttribute<FText, ::SlateAttributePrivate::FSlateAttributeNoInvalidationReason, TSlateAttributeFTextComparePredicate>::TSlateNestedMemberAttribute;
+	};
+
+	//~ Override for FText that use the TSlateAttributeFTextComparePredicate to compare FText
+	template<EInvalidateWidgetReason InInvalidationReasonValue>
+	struct TSlateNestedAttribute<FText, InInvalidationReasonValue> : public ::SlateAttributePrivate::TSlateNestedMemberAttribute<
+		FText,
+		typename std::conditional<InInvalidationReasonValue == EInvalidateWidgetReason::None, ::SlateAttributePrivate::FSlateAttributeNoInvalidationReason, TSlateAttributeInvalidationReason<InInvalidationReasonValue>>::type,
+		TSlateAttributeFTextComparePredicate>
+	{
+		using ::SlateAttributePrivate::TSlateNestedMemberAttribute<
+			FText,
+			typename std::conditional<InInvalidationReasonValue == EInvalidateWidgetReason::None, ::SlateAttributePrivate::FSlateAttributeNoInvalidationReason, TSlateAttributeInvalidationReason<InInvalidationReasonValue>>::type,
+			TSlateAttributeFTextComparePredicate>::TSlateNestedMemberAttribute;
+	};
+};
+
