Slate Dynamic Invalidation - SMultilineText

#Synced-CL 49551883

diff --git a/Engine/Source/Runtime/Slate/Private/Widgets/Text/SMultiLineEditableText.cpp b/Engine/Source/Runtime/Slate/Private/Widgets/Text/SMultiLineEditableText.cpp
--- a/Engine/Source/Runtime/Slate/Private/Widgets/Text/SMultiLineEditableText.cpp
+++ b/Engine/Source/Runtime/Slate/Private/Widgets/Text/SMultiLineEditableText.cpp
@@ -17,21 +17,35 @@
 SLATE_IMPLEMENT_WIDGET(SMultiLineEditableText)
 void SMultiLineEditableText::PrivateRegisterAttributes(FSlateAttributeInitializer& AttributeInitializer)
 {
+	struct FInvalidation
+	{
+		static void SynchronizeTextStyle(SWidget& OwningWidget)
+		{
+			SMultiLineEditableText& MultiLineEditableText = static_cast<SMultiLineEditableText&>(OwningWidget);
+			MultiLineEditableText.SynchronizeTextStyle();
+		}
+	};
+
 	FSlateEditableTextLayout::RegisterNestedAttributes<&SMultiLineEditableText::SlateEditableTextLayout>(AttributeInitializer);
+
+	SLATE_ADD_MEMBER_ATTRIBUTE_DEFINITION(AttributeInitializer, FontOverrideAttribute, EInvalidateWidgetReason::Layout)
+		.OnValueChanged(FSlateAttributeDescriptor::FAttributeValueChangedDelegate::CreateStatic(FInvalidation::SynchronizeTextStyle));
+	SLATE_ADD_MEMBER_ATTRIBUTE_DEFINITION(AttributeInitializer, bIsReadOnlyAttribute, EInvalidateWidgetReason::Paint);
 }
 
 SMultiLineEditableText::SMultiLineEditableText()
 	: SlateEditableTextLayout(*this)
 	, bSelectAllTextWhenFocused(false)
-	, bIsReadOnly(false)
 	, AmountScrolledWhileRightMouseDown(0.0f)
 	, bIsSoftwareCursor(false)
+	, FontOverrideAttribute(*this)
+	, bIsReadOnlyAttribute(*this, false)
 {
 }
 
 void SMultiLineEditableText::Construct( const FArguments& InArgs )
 {
-	bIsReadOnly = InArgs._IsReadOnly;
+	SetIsReadOnly(InArgs._IsReadOnly);
 
 	OnIsTypedCharValid = InArgs._OnIsTypedCharValid;
 	OnTextChangedCallback = InArgs._OnTextChanged;
@@ -56,7 +70,7 @@ void SMultiLineEditableText::Construct( const FArguments& InArgs )
 
 	HScrollBar = InArgs._HScrollBar;
 
-	// Horizontal scroll is unecessary when auto wrap is enabled.
+	// Horizontal scroll is unnecessary when auto wrap is enabled.
 	// Otherwise if allowed, having both enabled will cause a wrap+scroll relayout feedback loop which can result in jittering
 	const bool bAutoWrapEnabled = InArgs._AutoWrapText.IsSet() && InArgs._AutoWrapText.Get();
 
@@ -73,19 +87,13 @@ void SMultiLineEditableText::Construct( const FArguments& InArgs )
 		VScrollBar->SetOnUserScrolled(FOnUserScrolled::CreateSP(this, &SMultiLineEditableText::OnVScrollBarMoved));
 	}
 
-	FTextBlockStyle TextStyle = *InArgs._TextStyle;
-	if (InArgs._Font.IsSet() || InArgs._Font.IsBound())
-	{
-		TextStyle.SetFont(InArgs._Font.Get());
-	}
-
 	TSharedPtr<ITextLayoutMarshaller> Marshaller = InArgs._Marshaller;
 	if (!Marshaller.IsValid())
 	{
 		Marshaller = FPlainTextLayoutMarshaller::Create();
 	}
 
-	SlateEditableTextLayout.Construct(InArgs._Text, TextStyle, InArgs._TextShapingMethod, InArgs._TextFlowDirection, InArgs._CreateSlateTextLayout, Marshaller.ToSharedRef(), Marshaller.ToSharedRef());
+	SlateEditableTextLayout.Construct(InArgs._Text, *InArgs._TextStyle, InArgs._TextShapingMethod, InArgs._TextFlowDirection, InArgs._CreateSlateTextLayout, Marshaller.ToSharedRef(), Marshaller.ToSharedRef());
 	SlateEditableTextLayout.SetHintText(InArgs._HintText);
 	SlateEditableTextLayout.SetSearchText(InArgs._SearchText);
 	SlateEditableTextLayout.SetTextWrapping(InArgs._WrapTextAt, InArgs._AutoWrapText, InArgs._WrappingPolicy);
@@ -100,6 +108,8 @@ void SMultiLineEditableText::Construct( const FArguments& InArgs )
 		SlateEditableTextLayout.SetOnAllFontFacesFinishLoadingDelegate(InArgs._OnAllFontFacesFinishLoading);
 	}
 
+	SetFont(InArgs._Font);
+
 	// build context menu extender
 	MenuExtender = MakeShareable(new FExtender);
 	MenuExtender->AddMenuExtension("EditText", EExtensionHook::Before, TSharedPtr<FUICommandList>(), InArgs._ContextMenuExtender);
@@ -117,9 +127,9 @@ void SMultiLineEditableText::GetTextLine(const int32 InLineIndex, FString& OutTe
 	SlateEditableTextLayout.GetTextLine(InLineIndex, OutTextLine);
 }
 
-void SMultiLineEditableText::SetText(const TAttribute< FText >& InText)
+void SMultiLineEditableText::SetText(TAttribute<FText> InText)
 {
-	SlateEditableTextLayout.SetText(InText);
+	SlateEditableTextLayout.SetText(MoveTemp(InText));
 }
 
 int32 SMultiLineEditableText::GetTextLineCount() const
@@ -137,9 +147,9 @@ FText SMultiLineEditableText::GetPlainText() const
 	return SlateEditableTextLayout.GetPlainText();
 }
 
-void SMultiLineEditableText::SetHintText(const TAttribute< FText >& InHintText)
+void SMultiLineEditableText::SetHintText(TAttribute<FText> InHintText)
 {
-	SlateEditableTextLayout.SetHintText(InHintText);
+	SlateEditableTextLayout.SetHintText(MoveTemp(InHintText));
 }
 
 FText SMultiLineEditableText::GetHintText() const
@@ -147,9 +157,9 @@ FText SMultiLineEditableText::GetHintText() const
 	return SlateEditableTextLayout.GetHintText();
 }
 
-void SMultiLineEditableText::SetSearchText(const TAttribute<FText>& InSearchText)
+void SMultiLineEditableText::SetSearchText(TAttribute<FText> InSearchText)
 {
-	SlateEditableTextLayout.SetSearchText(InSearchText);
+	SlateEditableTextLayout.SetSearchText(MoveTemp(InSearchText));
 }
 
 FText SMultiLineEditableText::GetSearchText() const
@@ -175,23 +185,21 @@ void SMultiLineEditableText::SetTextStyle(const FTextBlockStyle* InTextStyle)
 	}
 	else
 	{
-		FArguments Defaults;
+		const FArguments Defaults;
 		check(Defaults._TextStyle);
 		SlateEditableTextLayout.SetTextStyle(*Defaults._TextStyle);
 	}
 }
 
-void SMultiLineEditableText::SetFont(const TAttribute< FSlateFontInfo >& InNewFont)
+void SMultiLineEditableText::SetFont(TAttribute<FSlateFontInfo> InNewFont)
 {
-	FTextBlockStyle TextStyle = SlateEditableTextLayout.GetTextStyle();
-	TextStyle.SetFont(InNewFont.Get());
-	SlateEditableTextLayout.SetTextStyle(TextStyle);
+	bIsFontOverrideAttrSet = InNewFont.IsSet();
+	FontOverrideAttribute.Assign(*this, MoveTemp(InNewFont));
 }
 
 FSlateFontInfo SMultiLineEditableText::GetFont() const
 {
-	FTextBlockStyle TextStyle = SlateEditableTextLayout.GetTextStyle();
-	return TextStyle.Font;
+	return bIsFontOverrideAttrSet ? FontOverrideAttribute.Get() : SlateEditableTextLayout.GetTextStyle().Font;
 }
 
 void SMultiLineEditableText::SetTextShapingMethod(const TOptional<ETextShapingMethod>& InTextShapingMethod)
@@ -204,84 +212,84 @@ void SMultiLineEditableText::SetTextFlowDirection(const TOptional<ETextFlowDirec
 	SlateEditableTextLayout.SetTextFlowDirection(InTextFlowDirection);
 }
 
-void SMultiLineEditableText::SetWrapTextAt(const TAttribute<float>& InWrapTextAt)
+void SMultiLineEditableText::SetWrapTextAt(TAttribute<float> InWrapTextAt)
 {
-	SlateEditableTextLayout.SetWrapTextAt(InWrapTextAt);
+	SlateEditableTextLayout.SetWrapTextAt(MoveTemp(InWrapTextAt));
 }
 
-void SMultiLineEditableText::SetAutoWrapText(const TAttribute<bool>& InAutoWrapText)
+void SMultiLineEditableText::SetAutoWrapText(TAttribute<bool> InAutoWrapText)
 {
-	SlateEditableTextLayout.SetAutoWrapText(InAutoWrapText);
+	SlateEditableTextLayout.SetAutoWrapText(MoveTemp(InAutoWrapText));
 }
 
-void SMultiLineEditableText::SetWrappingPolicy(const TAttribute<ETextWrappingPolicy>& InWrappingPolicy)
+void SMultiLineEditableText::SetWrappingPolicy(TAttribute<ETextWrappingPolicy> InWrappingPolicy)
 {
-	SlateEditableTextLayout.SetWrappingPolicy(InWrappingPolicy);
+	SlateEditableTextLayout.SetWrappingPolicy(MoveTemp(InWrappingPolicy));
 }
 
-void SMultiLineEditableText::SetLineHeightPercentage(const TAttribute<float>& InLineHeightPercentage)
+void SMultiLineEditableText::SetLineHeightPercentage(TAttribute<float> InLineHeightPercentage)
 {
-	SlateEditableTextLayout.SetLineHeightPercentage(InLineHeightPercentage);
+	SlateEditableTextLayout.SetLineHeightPercentage(MoveTemp(InLineHeightPercentage));
 }
 
-void SMultiLineEditableText::SetApplyLineHeightToBottomLine(const TAttribute<bool>& InApplyLineHeightToBottomLine)
+void SMultiLineEditableText::SetApplyLineHeightToBottomLine(TAttribute<bool> InApplyLineHeightToBottomLine)
 {
-	SlateEditableTextLayout.SetApplyLineHeightToBottomLine(InApplyLineHeightToBottomLine);
+	SlateEditableTextLayout.SetApplyLineHeightToBottomLine(MoveTemp(InApplyLineHeightToBottomLine));
 }
 
-void SMultiLineEditableText::SetMargin(const TAttribute<FMargin>& InMargin)
+void SMultiLineEditableText::SetMargin(TAttribute<FMargin> InMargin)
 {
-	SlateEditableTextLayout.SetMargin(InMargin);
+	SlateEditableTextLayout.SetMargin(MoveTemp(InMargin));
 }
 
-void SMultiLineEditableText::SetJustification(const TAttribute<ETextJustify::Type>& InJustification)
+void SMultiLineEditableText::SetJustification(TAttribute<ETextJustify::Type> InJustification)
 {
-	SlateEditableTextLayout.SetJustification(InJustification);
+	SlateEditableTextLayout.SetJustification(MoveTemp(InJustification));
 }
 
-void SMultiLineEditableText::SetOverflowPolicy(TOptional<ETextOverflowPolicy> InOverflowPolicy)
+void SMultiLineEditableText::SetOverflowPolicy(const TOptional<ETextOverflowPolicy>& InOverflowPolicy)
 {
 	SlateEditableTextLayout.SetOverflowPolicy(InOverflowPolicy);
 }
 
-void SMultiLineEditableText::SetAllowContextMenu(const TAttribute< bool >& InAllowContextMenu)
+void SMultiLineEditableText::SetAllowContextMenu(TAttribute<bool> InAllowContextMenu)
 {
-	bAllowContextMenu = InAllowContextMenu;
+	bAllowContextMenu = MoveTemp(InAllowContextMenu);
 }
 
 void SMultiLineEditableText::SetVirtualKeyboardDismissAction(TAttribute< EVirtualKeyboardDismissAction > InVirtualKeyboardDismissAction)
 {
-	VirtualKeyboardDismissAction = InVirtualKeyboardDismissAction;
+	VirtualKeyboardDismissAction = MoveTemp(InVirtualKeyboardDismissAction);
 }
 
-void SMultiLineEditableText::SetIsReadOnly(const TAttribute<bool>& InIsReadOnly)
+void SMultiLineEditableText::SetIsReadOnly(TAttribute<bool> InIsReadOnly)
 {
-	bIsReadOnly = InIsReadOnly;
+	bIsReadOnlyAttribute.Assign(*this, MoveTemp(InIsReadOnly));
 }
 
-void SMultiLineEditableText::SetSelectAllTextWhenFocused(const TAttribute<bool>& InSelectAllTextWhenFocused)
+void SMultiLineEditableText::SetSelectAllTextWhenFocused(TAttribute<bool> InSelectAllTextWhenFocused)
 {
-	bSelectAllTextWhenFocused = InSelectAllTextWhenFocused;
+	bSelectAllTextWhenFocused = MoveTemp(InSelectAllTextWhenFocused);
 }
 
-void SMultiLineEditableText::SetSelectWordOnMouseDoubleClick(const TAttribute<bool>& InSelectWordOnMouseDoubleClick)
+void SMultiLineEditableText::SetSelectWordOnMouseDoubleClick(TAttribute<bool> InSelectWordOnMouseDoubleClick)
 {
-	bSelectWordOnMouseDoubleClick = InSelectWordOnMouseDoubleClick;
+	bSelectWordOnMouseDoubleClick = MoveTemp(InSelectWordOnMouseDoubleClick);
 }
 
-void SMultiLineEditableText::SetClearTextSelectionOnFocusLoss(const TAttribute<bool>& InClearTextSelectionOnFocusLoss)
+void SMultiLineEditableText::SetClearTextSelectionOnFocusLoss(TAttribute<bool> InClearTextSelectionOnFocusLoss)
 {
-	bClearTextSelectionOnFocusLoss = InClearTextSelectionOnFocusLoss;
+	bClearTextSelectionOnFocusLoss = MoveTemp(InClearTextSelectionOnFocusLoss);
 }
 
-void SMultiLineEditableText::SetRevertTextOnEscape(const TAttribute<bool>& InRevertTextOnEscape)
+void SMultiLineEditableText::SetRevertTextOnEscape(TAttribute<bool> InRevertTextOnEscape)
 {
-	bRevertTextOnEscape = InRevertTextOnEscape;
+	bRevertTextOnEscape = MoveTemp(InRevertTextOnEscape);
 }
 
-void SMultiLineEditableText::SetClearKeyboardFocusOnCommit(const TAttribute<bool>& InClearKeyboardFocusOnCommit)
+void SMultiLineEditableText::SetClearKeyboardFocusOnCommit(TAttribute<bool> InClearKeyboardFocusOnCommit)
 {
-	bClearKeyboardFocusOnCommit = InClearKeyboardFocusOnCommit;
+	bClearKeyboardFocusOnCommit = MoveTemp(InClearKeyboardFocusOnCommit);
 }
 
 void SMultiLineEditableText::SetFontFacesLoadingPaintPolicy(EFontFacesLoadingPaintPolicy InFontFacesLoadingPaintPolicy)
@@ -303,7 +311,7 @@ void SMultiLineEditableText::OnVScrollBarMoved(const float InScrollOffsetFractio
 
 bool SMultiLineEditableText::IsTextReadOnly() const
 {
-	return bIsReadOnly.Get(false);
+	return bIsReadOnlyAttribute.Get();
 }
 
 bool SMultiLineEditableText::IsTextPassword() const
@@ -460,12 +468,12 @@ void SMultiLineEditableText::OnCursorMoved(const FTextLocation& InLocation)
 	Invalidate(EInvalidateWidget::Layout);
 }
 
-float SMultiLineEditableText::UpdateAndClampHorizontalScrollBar(const float InViewOffset, const float InViewFraction, const EVisibility InVisiblityOverride)
+float SMultiLineEditableText::UpdateAndClampHorizontalScrollBar(const float InViewOffset, const float InViewFraction, const EVisibility InVisibilityOverride)
 {
 	if (HScrollBar.IsValid())
 	{
 		HScrollBar->SetState(InViewOffset, InViewFraction);
-		HScrollBar->SetUserVisibility(InVisiblityOverride);
+		HScrollBar->SetUserVisibility(InVisibilityOverride);
 		if (!HScrollBar->IsNeeded())
 		{
 			// We cannot scroll, so ensure that there is no offset
@@ -476,12 +484,12 @@ float SMultiLineEditableText::UpdateAndClampHorizontalScrollBar(const float InVi
 	return SlateEditableTextLayout.GetScrollOffset().X;
 }
 
-float SMultiLineEditableText::UpdateAndClampVerticalScrollBar(const float InViewOffset, const float InViewFraction, const EVisibility InVisiblityOverride)
+float SMultiLineEditableText::UpdateAndClampVerticalScrollBar(const float InViewOffset, const float InViewFraction, const EVisibility InVisibilityOverride)
 {
 	if (VScrollBar.IsValid())
 	{
 		VScrollBar->SetState(InViewOffset, InViewFraction);
-		VScrollBar->SetUserVisibility(InVisiblityOverride);
+		VScrollBar->SetUserVisibility(InVisibilityOverride);
 		if (!VScrollBar->IsNeeded())
 		{
 			// We cannot scroll, so ensure that there is no offset
@@ -561,7 +569,7 @@ void SMultiLineEditableText::InsertTextAtCursor(const FString& InString)
 
 void SMultiLineEditableText::InsertRunAtCursor(TSharedRef<IRun> InRun)
 {
-	SlateEditableTextLayout.InsertRunAtCursor(InRun);
+	SlateEditableTextLayout.InsertRunAtCursor(MoveTemp(InRun));
 }
 
 void SMultiLineEditableText::GoTo(const FTextLocation& NewLocation)
@@ -643,6 +651,30 @@ void SMultiLineEditableText::ForceScroll(int32 UserIndex, float ScrollAxisMagnit
 	OnMouseWheel(CachedGeom, FPointerEvent(UserIndex, 0, ScrollPos, ScrollPos, PressedKeys, EKeys::Invalid, ScrollAxisMagnitude, FModifierKeysState()));
 }
 
+void SMultiLineEditableText::SynchronizeTextStyle()
+{
+	// Has the style used for this editable text changed?
+	bool bTextStyleChanged = false;
+	FTextBlockStyle NewTextStyle = SlateEditableTextLayout.GetTextStyle();
+
+	// Sync from the font override
+	if (bIsFontOverrideAttrSet)
+	{
+		const FSlateFontInfo& NewFontInfo = FontOverrideAttribute.Get();
+		if (!NewTextStyle.Font.IsIdenticalTo(NewFontInfo))
+		{
+			NewTextStyle.Font = NewFontInfo;
+			bTextStyleChanged = true;
+		}
+	}
+
+	if (bTextStyleChanged)
+	{
+		SlateEditableTextLayout.SetTextStyle(NewTextStyle);
+		SlateEditableTextLayout.ForceRefreshTextLayout(SlateEditableTextLayout.GetEditableText());
+	}
+}
+
 void SMultiLineEditableText::Tick( const FGeometry& AllottedGeometry, const double InCurrentTime, const float InDeltaTime )
 {
 	SlateEditableTextLayout.Tick(AllottedGeometry, InCurrentTime, InDeltaTime);
@@ -683,6 +715,7 @@ int32 SMultiLineEditableText::OnPaint( const FPaintArgs& Args, const FGeometry&
 
 void SMultiLineEditableText::CacheDesiredSize(float LayoutScaleMultiplier)
 {
+	SynchronizeTextStyle();
 	SlateEditableTextLayout.CacheDesiredSize(LayoutScaleMultiplier);
 	SWidget::CacheDesiredSize(LayoutScaleMultiplier);
 }
@@ -865,8 +898,7 @@ bool SMultiLineEditableText::IsInteractable() const
 bool SMultiLineEditableText::ComputeVolatility() const
 {
 	return SWidget::ComputeVolatility()
-		|| HasKeyboardFocus()
-		|| bIsReadOnly.IsBound();
+		|| HasKeyboardFocus();
 }
 
 bool SMultiLineEditableText::IsRightClickScrolling() const
diff --git a/Engine/Source/Runtime/Slate/Public/Widgets/Text/SMultiLineEditableText.h b/Engine/Source/Runtime/Slate/Public/Widgets/Text/SMultiLineEditableText.h
--- a/Engine/Source/Runtime/Slate/Public/Widgets/Text/SMultiLineEditableText.h
+++ b/Engine/Source/Runtime/Slate/Public/Widgets/Text/SMultiLineEditableText.h
@@ -38,6 +38,12 @@ enum class ETextShapingMethod : uint8;
 class ITextLayoutMarshaller;
 class FSlateEditableTextLayout;
 
+template <>
+struct TWidgetTypeTraits<class SMultiLineEditableText>
+{
+	static constexpr bool SupportsInvalidation() { return true; }
+};
+
 /** An editable text widget that supports multiple lines and soft word-wrapping. */
 class SMultiLineEditableText : public SWidget, public ISlateEditableTextWidget
 {
@@ -49,7 +55,7 @@ public:
 	{
 	public:
 		FScopedEditableTextTransaction(TSharedPtr<SMultiLineEditableText> InText)
-			: Text(InText)
+			: Text(MoveTemp(InText))
 		{
 			Text->BeginEditTransaction();
 		}
@@ -120,7 +126,7 @@ public:
 		SLATE_ATTRIBUTE(float, WrapTextAt)
 
 		/** Whether to wrap text automatically based on the widget's computed horizontal space.  IMPORTANT: Using automatic wrapping can result
-		    in visual artifacts, as the the wrapped size will computed be at least one frame late!  Consider using WrapTextAt instead.  The initial 
+			in visual artifacts, as the computed wrapped size will be at least one frame late!  Consider using WrapTextAt instead.  The initial
 			desired size will not be clamped.  This works best in cases where the text block's size is not affecting other widget's layout. */
 		SLATE_ATTRIBUTE(bool, AutoWrapText)
 
@@ -234,7 +240,7 @@ public:
 		/** Which text flow direction should we use? (unset to use the default returned by GetDefaultTextFlowDirection) */
 		SLATE_ARGUMENT( TOptional<ETextFlowDirection>, TextFlowDirection )
 
-		/** Determines what happens to text that is clipped and doesnt fit within the clip rect for this widget */
+		/** Determines what happens to text that is clipped and doesn't fit within the clip rect for this widget */
 		SLATE_ARGUMENT(TOptional<ETextOverflowPolicy>, OverflowPolicy)
 
 		/** The paint policy to use for text if the font faces are still async loading. */
@@ -251,7 +257,7 @@ public:
 	/**
 	 * Sets the text for this text block
 	 */
-	SLATE_API void SetText(const TAttribute< FText >& InText);
+	SLATE_API void SetText(TAttribute<FText> InText);
 
 	/**
 	 * Returns the text string
@@ -284,13 +290,13 @@ public:
 	/**
 	 * Sets the text that appears when there is no text in the text box
 	 */
-	SLATE_API void SetHintText(const TAttribute< FText >& InHintText);
+	SLATE_API void SetHintText(TAttribute<FText> InHintText);
 
 	/** Get the text that appears when there is no text in the text box */
 	SLATE_API FText GetHintText() const;
 
 	/** Set the text that is currently being searched for (if any) */
-	SLATE_API void SetSearchText(const TAttribute<FText>& InSearchText);
+	SLATE_API void SetSearchText(TAttribute<FText> InSearchText);
 
 	/** Get the text that is currently being searched for (if any) */
 	SLATE_API FText GetSearchText() const;
@@ -305,7 +311,7 @@ public:
 	SLATE_API void SetTextStyle(const FTextBlockStyle* InTextStyle);
 
 	/** See attribute Font */
-	SLATE_API void SetFont(const TAttribute< FSlateFontInfo >& InNewFont);
+	SLATE_API void SetFont(TAttribute<FSlateFontInfo> InNewFont);
 	SLATE_API FSlateFontInfo GetFont() const;
 
 	/** See TextShapingMethod attribute */
@@ -315,44 +321,44 @@ public:
 	SLATE_API void SetTextFlowDirection(const TOptional<ETextFlowDirection>& InTextFlowDirection);
 
 	/** See WrapTextAt attribute */
-	SLATE_API void SetWrapTextAt(const TAttribute<float>& InWrapTextAt);
+	SLATE_API void SetWrapTextAt(TAttribute<float> InWrapTextAt);
 
 	/** See AutoWrapText attribute */
-	SLATE_API void SetAutoWrapText(const TAttribute<bool>& InAutoWrapText);
+	SLATE_API void SetAutoWrapText(TAttribute<bool> InAutoWrapText);
 
 	/** Set WrappingPolicy attribute */
-	SLATE_API void SetWrappingPolicy(const TAttribute<ETextWrappingPolicy>& InWrappingPolicy);
+	SLATE_API void SetWrappingPolicy(TAttribute<ETextWrappingPolicy> InWrappingPolicy);
 
 	/** See LineHeightPercentage attribute */
-	SLATE_API void SetLineHeightPercentage(const TAttribute<float>& InLineHeightPercentage);
+	SLATE_API void SetLineHeightPercentage(TAttribute<float> InLineHeightPercentage);
 
 	/** See ApplyLineHeightToBottomLine attribute */
-	SLATE_API void SetApplyLineHeightToBottomLine(const TAttribute<bool>& InApplyLineHeightToBottomLine);
+	SLATE_API void SetApplyLineHeightToBottomLine(TAttribute<bool> InApplyLineHeightToBottomLine);
 
 	/** See Margin attribute */
-	SLATE_API void SetMargin(const TAttribute<FMargin>& InMargin);
+	SLATE_API void SetMargin(TAttribute<FMargin> InMargin);
 
 	/** See Justification attribute */
-	SLATE_API void SetJustification(const TAttribute<ETextJustify::Type>& InJustification);
+	SLATE_API void SetJustification(TAttribute<ETextJustify::Type> InJustification);
 
 	/** Sets the overflow policy for this text block */
-	SLATE_API void SetOverflowPolicy(TOptional<ETextOverflowPolicy> InOverflowPolicy);
+	SLATE_API void SetOverflowPolicy(const TOptional<ETextOverflowPolicy>& InOverflowPolicy);
 
 	/** See the AllowContextMenu attribute */
-	SLATE_API void SetAllowContextMenu(const TAttribute< bool >& InAllowContextMenu);
+	SLATE_API void SetAllowContextMenu(TAttribute<bool> InAllowContextMenu);
 
 	/** Set the VirtualKeyboardDismissAction attribute */
-	SLATE_API void SetVirtualKeyboardDismissAction(TAttribute< EVirtualKeyboardDismissAction > InVirtualKeyboardDismissAction);
+	SLATE_API void SetVirtualKeyboardDismissAction(TAttribute<EVirtualKeyboardDismissAction> InVirtualKeyboardDismissAction);
 
 	/**
 	 * Sets whether to select word on the mouse double click
 	 *
 	 * @param  InSelectWordOnMouseDoubleClick		Select word on the mouse double click
 	 */
-	SLATE_API void SetSelectWordOnMouseDoubleClick(const TAttribute<bool>& InSelectWordOnMouseDoubleClick);
+	SLATE_API void SetSelectWordOnMouseDoubleClick(TAttribute<bool> InSelectWordOnMouseDoubleClick);
 	
 	/** Sets the ReadOnly attribute */
-	SLATE_API void SetIsReadOnly(const TAttribute< bool >& InIsReadOnly);
+	SLATE_API void SetIsReadOnly(TAttribute<bool> InIsReadOnly);
 
 	/** Get the number of Text Lines */
 	UE_DEPRECATED(5.8, "Use the const version instead.")
@@ -368,28 +374,28 @@ public:
 	 *
 	 * @param  InSelectAllTextWhenFocused	Select all text when the user clicks?
 	 */
-	SLATE_API void SetSelectAllTextWhenFocused(const TAttribute<bool>& InSelectAllTextWhenFocused);
+	SLATE_API void SetSelectAllTextWhenFocused(TAttribute<bool> InSelectAllTextWhenFocused);
 
 	/**
 	 * Sets whether to clear text selection when focus is lost
 	 *
 	 * @param  InClearTextSelectionOnFocusLoss	Clear text selection when focus is lost?
 	 */
-	SLATE_API void SetClearTextSelectionOnFocusLoss(const TAttribute<bool>& InClearTextSelectionOnFocusLoss);
+	SLATE_API void SetClearTextSelectionOnFocusLoss(TAttribute<bool> InClearTextSelectionOnFocusLoss);
 
 	/**
 	 * Sets whether to allow the user to back out of changes when they press the escape key
 	 *
 	 * @param  InRevertTextOnEscape			Allow the user to back out of changes?
 	 */
-	SLATE_API void SetRevertTextOnEscape(const TAttribute<bool>& InRevertTextOnEscape);
+	SLATE_API void SetRevertTextOnEscape(TAttribute<bool> InRevertTextOnEscape);
 
 	/**
 	 * Sets whether to clear keyboard focus when pressing enter to commit changes
 	 *
 	 * @param  InClearKeyboardFocusOnCommit		Clear keyboard focus when pressing enter?
 	 */
-	SLATE_API void SetClearKeyboardFocusOnCommit(const TAttribute<bool>& InClearKeyboardFocusOnCommit);
+	SLATE_API void SetClearKeyboardFocusOnCommit(TAttribute<bool> InClearKeyboardFocusOnCommit);
 
 	/** Query to see if any text is selected within the document */
 	SLATE_API bool AnyTextSelected() const;
@@ -443,7 +449,7 @@ public:
 	/** Get the run currently under the cursor, or null if there is no run currently under the cursor */
 	SLATE_API TSharedPtr<const IRun> GetRunUnderCursor() const;
 
-	/** Get the runs currently that are current selected, some of which may be only partially selected */
+	/** Get the runs that are currently selected, some of which may be only partially selected */
 	SLATE_API TArray<TSharedRef<const IRun>> GetSelectedRuns() const;
 
 	/** Get the interaction position of the cursor (where to insert, delete, etc, text from/to) */
@@ -458,7 +464,7 @@ public:
 	/** Get the vertical scroll bar widget */
 	SLATE_API TSharedPtr<const SScrollBar> GetVScrollBar() const;
 
-	/** Refresh this editable text immediately, rather than wait for the usual caching mechanisms to take affect on the text Tick */
+	/** Refresh this editable text immediately, rather than wait for the usual caching mechanisms to take effect on the text Tick */
 	SLATE_API void Refresh();
 
 	/**
@@ -468,7 +474,7 @@ public:
 	 */
 	void SetOnKeyCharHandler(FOnKeyChar InOnKeyCharHandler)
 	{
-		OnKeyCharHandler = InOnKeyCharHandler;
+		OnKeyCharHandler = MoveTemp(InOnKeyCharHandler);
 	}
 
 	/**
@@ -478,7 +484,7 @@ public:
 	 */
 	void SetOnKeyDownHandler(FOnKeyDown InOnKeyDownHandler)
 	{
-		OnKeyDownHandler = InOnKeyDownHandler;
+		OnKeyDownHandler = MoveTemp(InOnKeyDownHandler);
 	}
 
 	/** Sets the paint policy to use for text if the font faces are still async loading. */
@@ -488,6 +494,10 @@ public:
 	 */
 	SLATE_API void ForceScroll(int32 UserIndex, float ScrollAxisMagnitude);
 
+protected:
+	/** Synchronize the text style currently set (including from overrides) and update the text layout if required */
+	SLATE_API void SynchronizeTextStyle();
+
 protected:
 	//~ Begin SWidget Interface
 	SLATE_API virtual void Tick( const FGeometry& AllottedGeometry, const double InCurrentTime, const float InDeltaTime ) override;
@@ -550,8 +560,8 @@ protected:
 	SLATE_API virtual void OnTextChanged(const FText& InText) override;
 	SLATE_API virtual void OnTextCommitted(const FText& InText, const ETextCommit::Type InTextAction) override;
 	SLATE_API virtual void OnCursorMoved(const FTextLocation& InLocation) override;
-	SLATE_API virtual float UpdateAndClampHorizontalScrollBar(const float InViewOffset, const float InViewFraction, const EVisibility InVisiblityOverride) override;
-	SLATE_API virtual float UpdateAndClampVerticalScrollBar(const float InViewOffset, const float InViewFraction, const EVisibility InVisiblityOverride) override;
+	SLATE_API virtual float UpdateAndClampHorizontalScrollBar(const float InViewOffset, const float InViewFraction, const EVisibility InVisibilityOverride) override;
+	SLATE_API virtual float UpdateAndClampVerticalScrollBar(const float InViewOffset, const float InViewFraction, const EVisibility InVisibilityOverride) override;
 	//~ End ISlateEditableTextWidget Interface
 
 protected:
@@ -562,9 +572,22 @@ protected:
 	SLATE_API void EndEditTransaction();
 
 protected:
+	TSlateAttributeRef<FSlateFontInfo> GetFontOverrideAttribute() const
+	{
+		return TSlateAttributeRef<FSlateFontInfo>{SharedThis(this), FontOverrideAttribute};
+	}
+
+	TSlateAttributeRef<bool> GetIsReadOnlyAttribute() const
+	{
+		return TSlateAttributeRef<bool>{SharedThis(this), bIsReadOnlyAttribute};
+	}
+
+protected:
+#if WITH_EDITORONLY_DATA
 	/** The text layout that deals with the editable text */
-	UE_DEPRECATED(5.8, "Use SlateEditableTextLayout, this variablle will never be set.")
+	UE_DEPRECATED(5.8, "Use SlateEditableTextLayout, this variable will never be set.")
 	TUniquePtr<FSlateEditableTextLayout> EditableTextLayout;
+#endif
 	/** The text layout that deals with the editable text */
 	FSlateEditableTextLayout SlateEditableTextLayout;
 
@@ -583,14 +606,17 @@ protected:
 	/** True if any changes should be reverted if we receive an escape key */
 	TAttribute<bool> bRevertTextOnEscape;
 
-	/** True if we want the text control to lose focus on an text commit/revert events */
+	/** True if we want the text control to lose focus on a text commit/revert events */
 	TAttribute<bool> bClearKeyboardFocusOnCommit;
 
 	/** Whether the context menu can be opened */
 	TAttribute<bool> bAllowContextMenu;
 
+#if WITH_EDITORONLY_DATA
 	/** Sets whether this text box can actually be modified interactively by the user */
+	UE_DEPRECATED(5.8, "Use SetIsTextReadOnly / IsTextReadOnly / GetIsTextReadOnlyAttribute")
 	TAttribute<bool> bIsReadOnly;
+#endif
 
 	/** Delegate to call before a context menu is opened */
 	FOnContextMenuOpening OnContextMenuOpening;
@@ -654,6 +680,14 @@ protected:
 
 	/** The message action to take when the virtual keyboard is dismissed by the user */
 	TAttribute<EVirtualKeyboardDismissAction> VirtualKeyboardDismissAction;
+
+private:
+	/** The font used to draw the text */
+	TSlateAttribute<FSlateFontInfo> FontOverrideAttribute;
+	bool bIsFontOverrideAttrSet = false;
+
+	/** Sets whether this text box can actually be modified interactively by the user */
+	TSlateAttribute<bool> bIsReadOnlyAttribute;
 };
 
 #endif //WITH_FANCY_TEXT
