Backing out changelist 49941561 due to issues with `SEditableText` and numerical inputs

#Synced-CL 49551883

diff --git a/Engine/Plugins/Animation/DeformerGraph/Source/OptimusEditor/Private/Widgets/SOptimusShaderTextDocumentTextBox.cpp b/Engine/Plugins/Animation/DeformerGraph/Source/OptimusEditor/Private/Widgets/SOptimusShaderTextDocumentTextBox.cpp
--- a/Engine/Plugins/Animation/DeformerGraph/Source/OptimusEditor/Private/Widgets/SOptimusShaderTextDocumentTextBox.cpp
+++ b/Engine/Plugins/Animation/DeformerGraph/Source/OptimusEditor/Private/Widgets/SOptimusShaderTextDocumentTextBox.cpp
@@ -440,7 +440,7 @@ void SOptimusShaderTextDocumentTextBox::OnToggleComment() const
 	{
 		// Automatically go to the next line if there is no selection
 		int32 LineIndex = CursorLocation.GetLineIndex() + 1;
-		if (LineIndex < const_cast<const SMultiLineEditableText*>(Text.Get())->GetTextLineCount())
+		if (LineIndex < Text->GetTextLineCount())
 		{
 			FString Line;
 			Text->GetTextLine(LineIndex, Line);
diff --git a/Engine/Plugins/PCG/Source/PCGEditor/Private/Widgets/SPCGNodeSourceTextBox.cpp b/Engine/Plugins/PCG/Source/PCGEditor/Private/Widgets/SPCGNodeSourceTextBox.cpp
--- a/Engine/Plugins/PCG/Source/PCGEditor/Private/Widgets/SPCGNodeSourceTextBox.cpp
+++ b/Engine/Plugins/PCG/Source/PCGEditor/Private/Widgets/SPCGNodeSourceTextBox.cpp
@@ -436,7 +436,7 @@ void SPCGNodeSourceTextBox::OnToggleComment() const
 	{
 		// Automatically go to the next line if there is no selection
 		int32 LineIndex = CursorLocation.GetLineIndex() + 1;
-		if (LineIndex < const_cast<const SMultiLineEditableText*>(Text.Get())->GetTextLineCount())
+		if (LineIndex < Text->GetTextLineCount())
 		{
 			FString Line;
 			Text->GetTextLine(LineIndex, Line);
diff --git a/Engine/Source/Editor/InputBindingEditor/Private/Widgets/SChordEditor.cpp b/Engine/Source/Editor/InputBindingEditor/Private/Widgets/SChordEditor.cpp
--- a/Engine/Source/Editor/InputBindingEditor/Private/Widgets/SChordEditor.cpp
+++ b/Engine/Source/Editor/InputBindingEditor/Private/Widgets/SChordEditor.cpp
@@ -43,7 +43,7 @@ void SChordEditor::Construct( const FArguments& InArgs, TSharedPtr<FUICommandInf
 		.HintTextOpacity(0.25f)
 	);
 
-	SlateEditableTextLayout.LoadText();
+	EditableTextLayout->LoadText();
 }
 
 
@@ -64,17 +64,17 @@ FReply SChordEditor::OnKeyDown( const FGeometry& MyGeometry, const FKeyEvent& In
 			EditingInputChord.Key = Key;
 		}
 	
-		SlateEditableTextLayout.BeginEditTransaction();
+		EditableTextLayout->BeginEditTransation();
 
 		EditingInputChord.bCtrl = InKeyEvent.IsControlDown();
 		EditingInputChord.bAlt = InKeyEvent.IsAltDown();
 		EditingInputChord.bShift = InKeyEvent.IsShiftDown();
 		EditingInputChord.bCmd = InKeyEvent.IsCommandDown();
 
-		SlateEditableTextLayout.LoadText();
-		SlateEditableTextLayout.GoTo(ETextLocation::EndOfDocument);
+		EditableTextLayout->LoadText();
+		EditableTextLayout->GoTo(ETextLocation::EndOfDocument);
 
-		SlateEditableTextLayout.EndEditTransaction();
+		EditableTextLayout->EndEditTransaction();
 		
 		OnChordTyped( EditingInputChord );
 		
diff --git a/Engine/Source/Runtime/Slate/Private/Widgets/Input/SEditableText.cpp b/Engine/Source/Runtime/Slate/Private/Widgets/Input/SEditableText.cpp
--- a/Engine/Source/Runtime/Slate/Private/Widgets/Input/SEditableText.cpp
+++ b/Engine/Source/Runtime/Slate/Private/Widgets/Input/SEditableText.cpp
@@ -4,7 +4,6 @@
 #include "Framework/Text/TextEditHelper.h"
 #include "Framework/Text/PlainTextLayoutMarshaller.h"
 #include "Widgets/Text/SlateEditableTextLayout.h"
-#include "Widgets/Text/SlateTextBlockLayout.h"
 #include "Types/ReflectionMetadata.h"
 #include "Types/TrackedMetaData.h"
 
@@ -12,44 +11,7 @@
 #include "Widgets/Accessibility/SlateAccessibleWidgets.h"
 #endif
 
-SLATE_IMPLEMENT_WIDGET(SEditableText)
-void SEditableText::PrivateRegisterAttributes(FSlateAttributeInitializer& AttributeInitializer)
-{
-	struct FInvalidation
-	{
-		static void SynchronizeTextStyle(SWidget& OwningWidget)
-		{
-			SEditableText& EditableText = static_cast<SEditableText&>(OwningWidget);
-			EditableText.SynchronizeTextStyle();
-		}
-		static void UpdateMarshaller(SWidget& OwningWidget)
-		{
-			SEditableText& EditableText = static_cast<SEditableText&>(OwningWidget);
-			EditableText.PlainTextMarshaller->MakeDirty();
-		}
-	};
-
-	FSlateEditableTextLayout::RegisterNestedAttributes<&SEditableText::SlateEditableTextLayout>(AttributeInitializer);
-
-	SLATE_ADD_MEMBER_ATTRIBUTE_DEFINITION(AttributeInitializer, MinDesiredWidthAttribute, EInvalidateWidgetReason::Layout);
-	SLATE_ADD_MEMBER_ATTRIBUTE_DEFINITION(AttributeInitializer, bIsPasswordAttribute, EInvalidateWidgetReason::Layout)
-		.OnValueChanged(FSlateAttributeDescriptor::FAttributeValueChangedDelegate::CreateStatic(FInvalidation::UpdateMarshaller));
-	SLATE_ADD_MEMBER_ATTRIBUTE_DEFINITION(AttributeInitializer, FontOverrideAttribute, EInvalidateWidgetReason::Layout)
-		.OnValueChanged(FSlateAttributeDescriptor::FAttributeValueChangedDelegate::CreateStatic(FInvalidation::SynchronizeTextStyle));
-
-	SLATE_ADD_MEMBER_ATTRIBUTE_DEFINITION(AttributeInitializer, ColorAndOpacityAttribute, EInvalidateWidgetReason::Paint)
-		.OnValueChanged(FSlateAttributeDescriptor::FAttributeValueChangedDelegate::CreateStatic(FInvalidation::SynchronizeTextStyle));
-	SLATE_ADD_MEMBER_ATTRIBUTE_DEFINITION(AttributeInitializer, BackgroundImageSelectedAttribute, EInvalidateWidgetReason::Layout)
-		.OnValueChanged(FSlateAttributeDescriptor::FAttributeValueChangedDelegate::CreateStatic(FInvalidation::SynchronizeTextStyle));
-}
-
 SEditableText::SEditableText()
-	: SlateEditableTextLayout(*this)
-	, FontOverrideAttribute(*this)
-	, ColorAndOpacityAttribute(*this)
-	, BackgroundImageSelectedAttribute(*this)
-	, bIsPasswordAttribute(*this, false)
-	, MinDesiredWidthAttribute(*this, 0.0f)
 {
 #if WITH_ACCESSIBILITY
 	AccessibleBehavior = EAccessibleBehavior::Auto;
@@ -57,9 +19,15 @@ SEditableText::SEditableText()
 #endif
 }
 
+SEditableText::~SEditableText()
+{
+	// Needed to avoid "deletion of pointer to incomplete type 'FSlateEditableTextLayout'; no destructor called" error when using TUniquePtr
+}
+
 void SEditableText::Construct( const FArguments& InArgs )
 {
-	SetIsReadOnly(InArgs._IsReadOnly);
+	bIsReadOnly = InArgs._IsReadOnly;
+	bIsPassword = InArgs._IsPassword;
 
 	bIsCaretMovedWhenGainFocus = InArgs._IsCaretMovedWhenGainFocus;
 	bSelectAllTextWhenFocused = InArgs._SelectAllTextWhenFocused;
@@ -72,7 +40,7 @@ void SEditableText::Construct( const FArguments& InArgs )
 	OnTextChangedCallback = InArgs._OnTextChanged;
 	OnTextCommittedCallback = InArgs._OnTextCommitted;
 	OnCursorMovedWithSelectionCallback = InArgs._OnCursorMovedWithSelection;
-	SetMinDesiredWidth(InArgs._MinDesiredWidth);
+	MinDesiredWidth = InArgs._MinDesiredWidth;
 	bSelectAllTextOnCommit = InArgs._SelectAllTextOnCommit;
 	bSelectWordOnMouseDoubleClick = InArgs._SelectWordOnMouseDoubleClick;
 	VirtualKeyboardType = InArgs._VirtualKeyboardType;
@@ -83,6 +51,10 @@ void SEditableText::Construct( const FArguments& InArgs )
 	OnKeyDownHandler = InArgs._OnKeyDownHandler;
 	bEnableIntegratedKeyboard = InArgs._EnableIntegratedKeyboard;
 
+	Font = InArgs._Font;
+	ColorAndOpacity = InArgs._ColorAndOpacity;
+	BackgroundImageSelected = InArgs._BackgroundImageSelected;
+
 	// We use the given style when creating the text layout as it may not be safe to call the override delegates until we've finished being constructed
 	// The first call to SynchronizeTextStyle will apply the correct overrides, and that will happen before the first paint
 	check(InArgs._Style);
@@ -92,31 +64,26 @@ void SEditableText::Construct( const FArguments& InArgs )
 	TextStyle.HighlightShape = InArgs._Style->BackgroundImageSelected;
 
 	PlainTextMarshaller = FPlainTextLayoutMarshaller::Create();
+	PlainTextMarshaller->SetIsPassword(bIsPassword);
 
 	// We use a separate marshaller for the hint text, as that should never be displayed as a password
 	TSharedRef<FPlainTextLayoutMarshaller> HintTextMarshaller = FPlainTextLayoutMarshaller::Create();
 
-	SlateEditableTextLayout.Construct(InArgs._Text, TextStyle, InArgs._TextShapingMethod, InArgs._TextFlowDirection, FCreateSlateTextLayout(), PlainTextMarshaller.ToSharedRef(), MoveTemp(HintTextMarshaller));
-	SlateEditableTextLayout.SetCursorBrush(InArgs._CaretImage.IsSet() ? InArgs._CaretImage : &InArgs._Style->CaretImage);
-	SlateEditableTextLayout.SetCompositionBrush(InArgs._BackgroundImageComposing.IsSet() ? InArgs._BackgroundImageComposing : &InArgs._Style->BackgroundImageComposing);
-	SlateEditableTextLayout.SetDebugSourceInfo(TAttribute<FString>::Create(TAttribute<FString>::FGetter::CreateLambda([this]{ return FReflectionMetaData::GetWidgetDebugInfo(this); })));
-	SlateEditableTextLayout.SetOverflowPolicy(InArgs._OverflowPolicy);
-	SlateEditableTextLayout.SetFontFacesLoadingPaintPolicy(InArgs._FontFacesLoadingPaintPolicy);
+	EditableTextLayout = MakeUnique<FSlateEditableTextLayout>(*this, InArgs._Text, TextStyle, InArgs._TextShapingMethod, InArgs._TextFlowDirection, FCreateSlateTextLayout(), PlainTextMarshaller.ToSharedRef(), HintTextMarshaller);
+	EditableTextLayout->SetHintText(InArgs._HintText);
+	EditableTextLayout->SetHintTextOpacity(InArgs._HintTextOpacity);
+	EditableTextLayout->SetSearchText(InArgs._SearchText);
+	EditableTextLayout->SetCursorBrush(InArgs._CaretImage.IsSet() ? InArgs._CaretImage : &InArgs._Style->CaretImage);
+	EditableTextLayout->SetCompositionBrush(InArgs._BackgroundImageComposing.IsSet() ? InArgs._BackgroundImageComposing : &InArgs._Style->BackgroundImageComposing);
+	EditableTextLayout->SetDebugSourceInfo(TAttribute<FString>::Create(TAttribute<FString>::FGetter::CreateLambda([this]{ return FReflectionMetaData::GetWidgetDebugInfo(this); })));
+	EditableTextLayout->SetJustification(InArgs._Justification);
+	EditableTextLayout->SetOverflowPolicy(InArgs._OverflowPolicy);
+	EditableTextLayout->SetFontFacesLoadingPaintPolicy(InArgs._FontFacesLoadingPaintPolicy);
 	if (InArgs._OnAllFontFacesFinishLoading.IsBound())
 	{
-		SlateEditableTextLayout.SetOnAllFontFacesFinishLoadingDelegate(InArgs._OnAllFontFacesFinishLoading);
+		EditableTextLayout->SetOnAllFontFacesFinishLoadingDelegate(InArgs._OnAllFontFacesFinishLoading);
 	}
 
-	SetText(InArgs._Text);
-	SetIsPassword(InArgs._IsPassword);
-	SetHintText(InArgs._HintText);
-	SetHintTextOpacity(InArgs._HintTextOpacity);
-	SetSearchText(InArgs._SearchText);
-	SetFont(InArgs._Font);
-	SetBackgroundImageSelected(InArgs._BackgroundImageSelected);
-	SetColorAndOpacity(InArgs._ColorAndOpacity);
-	SetJustification(InArgs._Justification);
-
 	// build context menu extender
 	MenuExtender = MakeShareable(new FExtender());
 	MenuExtender->AddMenuExtension("EditText", EExtensionHook::Before, TSharedPtr<FUICommandList>(), InArgs._ContextMenuExtender);
@@ -124,64 +91,65 @@ void SEditableText::Construct( const FArguments& InArgs )
 	AddMetadata(MakeShared<FTrackedMetaData>(this, FName("EditableText")));
 }
 
-void SEditableText::SetText(TAttribute<FText> InNewText)
+void SEditableText::SetText( const TAttribute< FText >& InNewText )
 {
-	SlateEditableTextLayout.SetText(MoveTemp(InNewText));
+	EditableTextLayout->SetText(InNewText);
 }
 
 FText SEditableText::GetText() const
 {
-	return SlateEditableTextLayout.GetText();
+	return EditableTextLayout->GetText();
 }
 
 bool SEditableText::SetEditableText(const FText& InNewText)
 {
-	return SlateEditableTextLayout.SetEditableText(InNewText);
+	return EditableTextLayout->SetEditableText(InNewText);
 }
 
-void SEditableText::SetFont(TAttribute<FSlateFontInfo> InNewFont)
+void SEditableText::SetFont( const TAttribute< FSlateFontInfo >& InNewFont )
 {
-	bIsFontOverrideAttrSet = InNewFont.IsSet();
-	FontOverrideAttribute.Assign(*this, MoveTemp(InNewFont));
+	Font = InNewFont;
+
+	Invalidate(EInvalidateWidgetReason::Layout);
 }
 
 FSlateFontInfo SEditableText::GetFont() const
 {
-	return bIsFontOverrideAttrSet ? FontOverrideAttribute.Get() : SlateEditableTextLayout.GetTextStyle().Font;
+	return Font.Get();
 }
 
 void SEditableText::SetTextStyle( const FEditableTextStyle& InNewTextStyle )
 {
-	bIsFontOverrideAttrSet = true;
-	FontOverrideAttribute.Assign(*this, InNewTextStyle.Font);
-	bIsColorAndOpacityAttrSet = true;
-	ColorAndOpacityAttribute.Assign(*this, InNewTextStyle.ColorAndOpacity);
-	bIsBackgroundImageSelectedAttrSet = true;
-	BackgroundImageSelectedAttribute.Assign(*this, &InNewTextStyle.BackgroundImageSelected);
+	Font = InNewTextStyle.Font;
+	ColorAndOpacity = InNewTextStyle.ColorAndOpacity;
+	BackgroundImageSelected = &InNewTextStyle.BackgroundImageSelected;
+
+	Invalidate(EInvalidateWidgetReason::Layout);
 }
 
 void SEditableText::SetTextBlockStyle(const FTextBlockStyle* InTextStyle)
 {
 	if (InTextStyle)
 	{
-		SlateEditableTextLayout.SetTextStyle(*InTextStyle);
+		EditableTextLayout->SetTextStyle(*InTextStyle);
+		Invalidate(EInvalidateWidgetReason::Layout); //Using Layout as changing text block size can affect the size.
 	}
 }
 
 void SEditableText::Tick( const FGeometry& AllottedGeometry, const double InCurrentTime, const float InDeltaTime )
 {
-	SlateEditableTextLayout.Tick(AllottedGeometry, InCurrentTime, InDeltaTime);
+	EditableTextLayout->Tick(AllottedGeometry, InCurrentTime, InDeltaTime);
 }
 
 int32 SEditableText::OnPaint( const FPaintArgs& Args, const FGeometry& AllottedGeometry, const FSlateRect& MyCullingRect, FSlateWindowElementList& OutDrawElements, int32 LayerId, const FWidgetStyle& InWidgetStyle, bool bParentEnabled ) const
 {
-	const FTextBlockStyle& EditableTextStyle = SlateEditableTextLayout.GetTextStyle();
+	const FTextBlockStyle& EditableTextStyle = EditableTextLayout->GetTextStyle();
 	const FLinearColor ForegroundColor = EditableTextStyle.ColorAndOpacity.GetColor(InWidgetStyle);
 
 	FWidgetStyle TextWidgetStyle = FWidgetStyle(InWidgetStyle)
 		.SetForegroundColor(ForegroundColor);
 
-	LayerId = SlateEditableTextLayout.OnPaint(Args, AllottedGeometry, MyCullingRect, OutDrawElements, LayerId, TextWidgetStyle, ShouldBeEnabled(bParentEnabled));
+	LayerId = EditableTextLayout->OnPaint(Args, AllottedGeometry, MyCullingRect, OutDrawElements, LayerId, TextWidgetStyle, ShouldBeEnabled(bParentEnabled));
 
 	return LayerId;
 }
@@ -189,25 +157,25 @@ int32 SEditableText::OnPaint( const FPaintArgs& Args, const FGeometry& AllottedG
 void SEditableText::CacheDesiredSize(float LayoutScaleMultiplier)
 {
 	SynchronizeTextStyle();
-	SlateEditableTextLayout.CacheDesiredSize(LayoutScaleMultiplier);
+	EditableTextLayout->CacheDesiredSize(LayoutScaleMultiplier);
 	SWidget::CacheDesiredSize(LayoutScaleMultiplier);
 }
 
 FVector2D SEditableText::ComputeDesiredSize(float LayoutScaleMultiplier) const
 {
-	FVector2D TextLayoutSize = SlateEditableTextLayout.ComputeDesiredSize(LayoutScaleMultiplier);
-	TextLayoutSize.X = FMath::Max(TextLayoutSize.X, MinDesiredWidthAttribute.Get());
+	FVector2D TextLayoutSize = EditableTextLayout->ComputeDesiredSize(LayoutScaleMultiplier);
+	TextLayoutSize.X = FMath::Max(TextLayoutSize.X, MinDesiredWidth.Get());
 	return TextLayoutSize;
 }
 
 FChildren* SEditableText::GetChildren()
 {
-	return SlateEditableTextLayout.GetChildren();
+	return EditableTextLayout->GetChildren();
 }
 
 void SEditableText::OnArrangeChildren(const FGeometry& AllottedGeometry, FArrangedChildren& ArrangedChildren) const
 {
-	SlateEditableTextLayout.OnArrangeChildren(AllottedGeometry, ArrangedChildren);
+	EditableTextLayout->OnArrangeChildren(AllottedGeometry, ArrangedChildren);
 }
 
 FReply SEditableText::OnDragOver( const FGeometry& MyGeometry, const FDragDropEvent& DragDropEvent )
@@ -220,7 +188,7 @@ FReply SEditableText::OnDragOver( const FGeometry& MyGeometry, const FDragDropEv
 			return FReply::Handled();
 		}
 	}
-
+	
 	return FReply::Unhandled();
 }
 
@@ -231,11 +199,11 @@ FReply SEditableText::OnDrop( const FGeometry& MyGeometry, const FDragDropEvent&
 	{
 		if ( DragDropOp->HasText() )
 		{
-			SlateEditableTextLayout.SetText(FText::FromString(DragDropOp->GetText()));
+			EditableTextLayout->SetText(FText::FromString(DragDropOp->GetText()));
 			return FReply::Handled();
 		}
 	}
-
+	
 	return FReply::Unhandled();
 }
 
@@ -246,13 +214,13 @@ bool SEditableText::SupportsKeyboardFocus() const
 
 FReply SEditableText::OnFocusReceived( const FGeometry& MyGeometry, const FFocusEvent& InFocusEvent )
 {
-	SlateEditableTextLayout.HandleFocusReceived(InFocusEvent);
+	EditableTextLayout->HandleFocusReceived(InFocusEvent);
 	return FReply::Handled();
 }
 
 void SEditableText::OnFocusLost( const FFocusEvent& InFocusEvent )
 {
-	SlateEditableTextLayout.HandleFocusLost(InFocusEvent);
+	EditableTextLayout->HandleFocusLost(InFocusEvent);
 }
 
 FReply SEditableText::OnKeyChar( const FGeometry& MyGeometry, const FCharacterEvent& InCharacterEvent )
@@ -264,10 +232,10 @@ FReply SEditableText::OnKeyChar( const FGeometry& MyGeometry, const FCharacterEv
 	{
 		Reply = OnKeyCharHandler.Execute(MyGeometry, InCharacterEvent);
 	}
-
+	
 	if (!Reply.IsEventHandled())
 	{
-		Reply = SlateEditableTextLayout.HandleKeyChar(InCharacterEvent);
+		Reply = EditableTextLayout->HandleKeyChar(InCharacterEvent);
 	}
 
 	return Reply;
@@ -285,7 +253,7 @@ FReply SEditableText::OnKeyDown( const FGeometry& MyGeometry, const FKeyEvent& I
 
 	if (!Reply.IsEventHandled())
 	{
-		Reply = SlateEditableTextLayout.HandleKeyDown(InKeyEvent);
+		Reply = EditableTextLayout->HandleKeyDown(InKeyEvent);
 
 		if (!Reply.IsEventHandled())
 		{
@@ -298,27 +266,27 @@ FReply SEditableText::OnKeyDown( const FGeometry& MyGeometry, const FKeyEvent& I
 
 FReply SEditableText::OnKeyUp( const FGeometry& MyGeometry, const FKeyEvent& InKeyEvent )
 {
-	return SlateEditableTextLayout.HandleKeyUp(InKeyEvent);
+	return EditableTextLayout->HandleKeyUp(InKeyEvent);
 }
 
 FReply SEditableText::OnMouseButtonDown( const FGeometry& InMyGeometry, const FPointerEvent& InMouseEvent )
 {
-	return SlateEditableTextLayout.HandleMouseButtonDown(InMyGeometry, InMouseEvent);
+	return EditableTextLayout->HandleMouseButtonDown(InMyGeometry, InMouseEvent);
 }
 
 FReply SEditableText::OnMouseButtonUp( const FGeometry& InMyGeometry, const FPointerEvent& InMouseEvent )
 {
-	return SlateEditableTextLayout.HandleMouseButtonUp(InMyGeometry, InMouseEvent);
+	return EditableTextLayout->HandleMouseButtonUp(InMyGeometry, InMouseEvent);
 }
 
 FReply SEditableText::OnMouseMove( const FGeometry& InMyGeometry, const FPointerEvent& InMouseEvent )
 {
-	return SlateEditableTextLayout.HandleMouseMove(InMyGeometry, InMouseEvent);
+	return EditableTextLayout->HandleMouseMove(InMyGeometry, InMouseEvent);
 }
 
 FReply SEditableText::OnMouseButtonDoubleClick( const FGeometry& InMyGeometry, const FPointerEvent& InMouseEvent )
 {
-	return SlateEditableTextLayout.HandleMouseButtonDoubleClick(InMyGeometry, InMouseEvent);
+	return EditableTextLayout->HandleMouseButtonDoubleClick(InMyGeometry, InMouseEvent);
 }
 
 FCursorReply SEditableText::OnCursorQuery( const FGeometry& MyGeometry, const FPointerEvent& CursorEvent ) const
@@ -339,212 +307,215 @@ bool SEditableText::IsInteractable() const
 bool SEditableText::ComputeVolatility() const
 {
 	return SWidget::ComputeVolatility()
-		|| HasKeyboardFocus();
+		|| HasKeyboardFocus()
+		|| EditableTextLayout->ComputeVolatility()
+		|| Font.IsBound()
+		|| ColorAndOpacity.IsBound()
+		|| BackgroundImageSelected.IsBound()
+		|| bIsReadOnly.IsBound()
+		|| bIsPassword.IsBound()
+		|| MinDesiredWidth.IsBound();
 }
 
-void SEditableText::SetHintText(TAttribute<FText> InHintText )
+void SEditableText::SetHintText( const TAttribute< FText >& InHintText )
 {
-	SlateEditableTextLayout.SetHintText(MoveTemp(InHintText));
+	EditableTextLayout->SetHintText(InHintText);
 }
 
 FText SEditableText::GetHintText() const
 {
-	return SlateEditableTextLayout.GetHintText();
+	return EditableTextLayout->GetHintText();
 }
 
-void SEditableText::SetHintTextOpacity(TAttribute<float> InHintTextOpacity)
+void SEditableText::SetHintTextOpacity(const TAttribute<float>& InHintTextOpacity)
 {
-	SlateEditableTextLayout.SetHintTextOpacity(MoveTemp(InHintTextOpacity));
+	EditableTextLayout->SetHintTextOpacity(InHintTextOpacity);
 }
 
 float SEditableText::GetHintTextOpacity() const
 {
-	return SlateEditableTextLayout.GetHintTextOpacity();
+	return EditableTextLayout->GetHintTextOpacity();
 }
 
-void SEditableText::SetSearchText(TAttribute<FText> InSearchText)
+void SEditableText::SetSearchText(const TAttribute<FText>& InSearchText)
 {
-	SlateEditableTextLayout.SetSearchText(MoveTemp(InSearchText));
+	EditableTextLayout->SetSearchText(InSearchText);
 }
 
 FText SEditableText::GetSearchText() const
 {
-	return SlateEditableTextLayout.GetSearchText();
+	return EditableTextLayout->GetSearchText();
 }
 
-void SEditableText::SetIsReadOnly(TAttribute<bool> InIsReadOnly)
+void SEditableText::SetIsReadOnly( TAttribute< bool > InIsReadOnly )
 {
-	bIsReadOnly = MoveTemp(InIsReadOnly);
+	bIsReadOnly = InIsReadOnly;
 }
 
-void SEditableText::SetIsPassword(TAttribute<bool> InIsPassword)
+void SEditableText::SetIsPassword( TAttribute< bool > InIsPassword )
 {
-	PlainTextMarshaller->SetIsPassword(InIsPassword);
-	bIsPasswordAttribute.Assign(*this, MoveTemp(InIsPassword));
+	bIsPassword = InIsPassword;
+	PlainTextMarshaller->SetIsPassword(bIsPassword);
 }
 
 void SEditableText::SetColorAndOpacity(TAttribute<FSlateColor> Color)
 {
-	bIsColorAndOpacityAttrSet = Color.IsSet();
-	ColorAndOpacityAttribute.Assign(*this, MoveTemp(Color));
-}
-
-void SEditableText::SetBackgroundImageSelected(TAttribute<const FSlateBrush*> InBackgroundImageSelected)
-{
-	bIsBackgroundImageSelectedAttrSet = InBackgroundImageSelected.IsSet();
-	BackgroundImageSelectedAttribute.Assign(*this, MoveTemp(InBackgroundImageSelected));
+	ColorAndOpacity = Color;
 }
 
-void SEditableText::SetMinDesiredWidth(TAttribute<float> InMinDesiredWidth)
+void SEditableText::SetMinDesiredWidth(const TAttribute<float>& InMinDesiredWidth)
 {
-	MinDesiredWidthAttribute.Assign(*this, MoveTemp(InMinDesiredWidth));
+	MinDesiredWidth = InMinDesiredWidth;
 }
 
-void SEditableText::SetIsCaretMovedWhenGainFocus(TAttribute<bool> InIsCaretMovedWhenGainFocus)
+void SEditableText::SetIsCaretMovedWhenGainFocus(const TAttribute<bool>& InIsCaretMovedWhenGainFocus)
 {
-	bIsCaretMovedWhenGainFocus = MoveTemp(InIsCaretMovedWhenGainFocus);
+	bIsCaretMovedWhenGainFocus = InIsCaretMovedWhenGainFocus;
 }
 
-void SEditableText::SetSelectAllTextWhenFocused(TAttribute<bool> InSelectAllTextWhenFocused)
+void SEditableText::SetSelectAllTextWhenFocused(const TAttribute<bool>& InSelectAllTextWhenFocused)
 {
-	bSelectAllTextWhenFocused = MoveTemp(InSelectAllTextWhenFocused);
+	bSelectAllTextWhenFocused = InSelectAllTextWhenFocused;
 }
 
-void SEditableText::SetRevertTextOnEscape(TAttribute<bool> InRevertTextOnEscape)
+void SEditableText::SetRevertTextOnEscape(const TAttribute<bool>& InRevertTextOnEscape)
 {
-	bRevertTextOnEscape = MoveTemp(InRevertTextOnEscape);
+	bRevertTextOnEscape = InRevertTextOnEscape;
 }
 
-void SEditableText::SetClearKeyboardFocusOnCommit(TAttribute<bool> InClearKeyboardFocusOnCommit)
+void SEditableText::SetClearKeyboardFocusOnCommit(const TAttribute<bool>& InClearKeyboardFocusOnCommit)
 {
-	bClearKeyboardFocusOnCommit = MoveTemp(InClearKeyboardFocusOnCommit);
+	bClearKeyboardFocusOnCommit = InClearKeyboardFocusOnCommit;
 }
 
-void SEditableText::SetSelectAllTextOnCommit(TAttribute<bool> InSelectAllTextOnCommit)
+void SEditableText::SetSelectAllTextOnCommit(const TAttribute<bool>& InSelectAllTextOnCommit)
 {
-	bSelectAllTextOnCommit = MoveTemp(InSelectAllTextOnCommit);
+	bSelectAllTextOnCommit = InSelectAllTextOnCommit;
 }
 
-void SEditableText::SetSelectWordOnMouseDoubleClick(TAttribute<bool> InSelectWordOnMouseDoubleClick)
+void SEditableText::SetSelectWordOnMouseDoubleClick(const TAttribute<bool>& InSelectWordOnMouseDoubleClick)
 {
-	bSelectWordOnMouseDoubleClick = MoveTemp(InSelectWordOnMouseDoubleClick);
+	bSelectWordOnMouseDoubleClick = InSelectWordOnMouseDoubleClick;
 }
 
-void SEditableText::SetJustification(TAttribute<ETextJustify::Type> InJustification)
+void SEditableText::SetJustification(const TAttribute<ETextJustify::Type>& InJustification)
 {
-	SlateEditableTextLayout.SetJustification(MoveTemp(InJustification));
+	EditableTextLayout->SetJustification(InJustification);
 }
 
-void SEditableText::SetAllowContextMenu(TAttribute<bool> InAllowContextMenu)
+void SEditableText::SetAllowContextMenu(const TAttribute< bool >& InAllowContextMenu)
 {
-	bAllowContextMenu = MoveTemp(InAllowContextMenu);
+	bAllowContextMenu = InAllowContextMenu;
 }
 
-void SEditableText::SetEnableIntegratedKeyboard(TAttribute<bool> InEnableIntegratedKeyboard)
+void SEditableText::SetEnableIntegratedKeyboard(const TAttribute<bool>& InEnableIntegratedKeyboard)
 {
-	bEnableIntegratedKeyboard = MoveTemp(InEnableIntegratedKeyboard);
+	bEnableIntegratedKeyboard = InEnableIntegratedKeyboard;
 }
 
-void SEditableText::SetVirtualKeyboardDismissAction(TAttribute<EVirtualKeyboardDismissAction> InVirtualKeyboardDismissAction)
+void SEditableText::SetVirtualKeyboardDismissAction(TAttribute< EVirtualKeyboardDismissAction > InVirtualKeyboardDismissAction)
 {
-	VirtualKeyboardDismissAction = MoveTemp(InVirtualKeyboardDismissAction);
+	VirtualKeyboardDismissAction = InVirtualKeyboardDismissAction;
 }
 
 void SEditableText::SetTextShapingMethod(const TOptional<ETextShapingMethod>& InTextShapingMethod)
 {
-	SlateEditableTextLayout.SetTextShapingMethod(InTextShapingMethod);
+	EditableTextLayout->SetTextShapingMethod(InTextShapingMethod);
 }
 
 void SEditableText::SetTextFlowDirection(const TOptional<ETextFlowDirection>& InTextFlowDirection)
 {
-	SlateEditableTextLayout.SetTextFlowDirection(InTextFlowDirection);
+	EditableTextLayout->SetTextFlowDirection(InTextFlowDirection);
 }
 
-void SEditableText::SetOverflowPolicy(const TOptional<ETextOverflowPolicy>& InOverflowPolicy)
+void SEditableText::SetOverflowPolicy(TOptional<ETextOverflowPolicy> InOverflowPolicy)
 {
-	SlateEditableTextLayout.SetOverflowPolicy(InOverflowPolicy);
+	EditableTextLayout->SetOverflowPolicy(InOverflowPolicy);
 }
 
 void SEditableText::SetFontFacesLoadingPaintPolicy(EFontFacesLoadingPaintPolicy InFontFacesLoadingPaintPolicy)
 {
-	SlateEditableTextLayout.SetFontFacesLoadingPaintPolicy(InFontFacesLoadingPaintPolicy);
+	if (EditableTextLayout)
+	{
+		EditableTextLayout->SetFontFacesLoadingPaintPolicy(InFontFacesLoadingPaintPolicy);
+	}
 }
 
 bool SEditableText::AnyTextSelected() const
 {
-	return SlateEditableTextLayout.AnyTextSelected();
+	return EditableTextLayout->AnyTextSelected();
 }
 
 void SEditableText::SelectAllText()
 {
-	SlateEditableTextLayout.SelectAllText();
+	EditableTextLayout->SelectAllText();
 }
 
 void SEditableText::ClearSelection()
 {
-	SlateEditableTextLayout.ClearSelection();
+	EditableTextLayout->ClearSelection();
 }
 
 FText SEditableText::GetSelectedText() const
 {
-	return SlateEditableTextLayout.GetSelectedText();
+	return EditableTextLayout->GetSelectedText();
 }
 
 void SEditableText::GoTo(const FTextLocation& NewLocation)
 {
-	SlateEditableTextLayout.GoTo(NewLocation);
+	EditableTextLayout->GoTo(NewLocation);
 }
 
 void SEditableText::GoTo(const ETextLocation NewLocation)
 {
-	SlateEditableTextLayout.GoTo(NewLocation);
+	EditableTextLayout->GoTo(NewLocation);
 }
 
 void SEditableText::ScrollTo(const FTextLocation& NewLocation)
 {
-	SlateEditableTextLayout.ScrollTo(NewLocation);
+	EditableTextLayout->ScrollTo(NewLocation);
 }
 
 void SEditableText::ScrollTo(const ETextLocation NewLocation)
 {
-	SlateEditableTextLayout.ScrollTo(NewLocation);
+	EditableTextLayout->ScrollTo(NewLocation);
 }
 
 void SEditableText::BeginSearch(const FText& InSearchText, const ESearchCase::Type InSearchCase, const bool InReverse)
 {
-	SlateEditableTextLayout.BeginSearch(InSearchText, InSearchCase, InReverse);
+	EditableTextLayout->BeginSearch(InSearchText, InSearchCase, InReverse);
 }
 
 void SEditableText::AdvanceSearch(const bool InReverse)
 {
-	SlateEditableTextLayout.AdvanceSearch(InReverse);
+	EditableTextLayout->AdvanceSearch(InReverse);
 }
 
 void SEditableText::EnableTextInputMethodContext()
 {
-	SlateEditableTextLayout.EnableTextInputMethodContext();
+	EditableTextLayout->EnableTextInputMethodContext();
 }
 
 FTextSelection SEditableText::GetSelection() const
 {
-	return SlateEditableTextLayout.GetSelection();
+	return EditableTextLayout->GetSelection();
 }
 
 void SEditableText::SelectText(const FTextLocation& InSelectionStart, const FTextLocation& InCursorLocation)
 {
-	SlateEditableTextLayout.SelectText(InSelectionStart, InCursorLocation);
+	EditableTextLayout->SelectText(InSelectionStart, InCursorLocation);
 }
 
 void SEditableText::SynchronizeTextStyle()
 {
 	// Has the style used for this editable text changed?
 	bool bTextStyleChanged = false;
-	FTextBlockStyle NewTextStyle = SlateEditableTextLayout.GetTextStyle();
+	FTextBlockStyle NewTextStyle = EditableTextLayout->GetTextStyle();
 
 	// Sync from the font override
-	if (bIsFontOverrideAttrSet)
+	if (Font.IsSet())
 	{
-		const FSlateFontInfo& NewFontInfo = FontOverrideAttribute.Get();
+		const FSlateFontInfo& NewFontInfo = Font.Get();
 		if (!NewTextStyle.Font.IsIdenticalTo(NewFontInfo))
 		{
 			NewTextStyle.Font = NewFontInfo;
@@ -553,9 +524,9 @@ void SEditableText::SynchronizeTextStyle()
 	}
 
 	// Sync from the color override
-	if (bIsColorAndOpacityAttrSet)
+	if (ColorAndOpacity.IsSet())
 	{
-		const FSlateColor& NewColorAndOpacity = ColorAndOpacityAttribute.Get();
+		const FSlateColor& NewColorAndOpacity = ColorAndOpacity.Get();
 		if (NewTextStyle.ColorAndOpacity != NewColorAndOpacity)
 		{
 			NewTextStyle.ColorAndOpacity = NewColorAndOpacity;
@@ -564,9 +535,9 @@ void SEditableText::SynchronizeTextStyle()
 	}
 
 	// Sync from the highlight shape override
-	if (bIsBackgroundImageSelectedAttrSet)
+	if (BackgroundImageSelected.IsSet())
 	{
-		const FSlateBrush* NewSelectionBrush = BackgroundImageSelectedAttribute.Get();
+		const FSlateBrush* NewSelectionBrush = BackgroundImageSelected.Get();
 		if (NewSelectionBrush && NewTextStyle.HighlightShape != *NewSelectionBrush)
 		{
 			NewTextStyle.HighlightShape = *NewSelectionBrush;
@@ -576,8 +547,8 @@ void SEditableText::SynchronizeTextStyle()
 
 	if (bTextStyleChanged)
 	{
-		SlateEditableTextLayout.SetTextStyle(NewTextStyle);
-		SlateEditableTextLayout.ForceRefreshTextLayout(SlateEditableTextLayout.GetEditableText());
+		EditableTextLayout->SetTextStyle(NewTextStyle);
+		EditableTextLayout->ForceRefreshTextLayout(EditableTextLayout->GetEditableText());
 	}
 }
 
@@ -588,7 +559,7 @@ bool SEditableText::IsTextReadOnly() const
 
 bool SEditableText::IsTextPassword() const
 {
-	return bIsPasswordAttribute.Get();
+	return bIsPassword.Get(false);
 }
 
 bool SEditableText::IsMultiLineTextEdit() const
@@ -662,7 +633,7 @@ void SEditableText::EnsureActiveTick()
 	auto DoActiveTick = [this](double InCurrentTime, float InDeltaTime) -> EActiveTimerReturnType
 	{
 		// Continue if we still have focus, otherwise treat as a fire-and-forget Tick() request
-		const bool bShouldAppearFocused = HasKeyboardFocus() || SlateEditableTextLayout.HasActiveContextMenu();
+		const bool bShouldAppearFocused = HasKeyboardFocus() || EditableTextLayout->HasActiveContextMenu();
 		return (bShouldAppearFocused) ? EActiveTimerReturnType::Continue : EActiveTimerReturnType::Stop;
 	};
 
@@ -715,8 +686,8 @@ TSharedPtr<SWidget> SEditableText::BuildContextMenuContent() const
 	{
 		return OnContextMenuOpening.Execute();
 	}
-
-	return SlateEditableTextLayout.BuildDefaultContextMenu(MenuExtender);
+	
+	return EditableTextLayout->BuildDefaultContextMenu(MenuExtender);
 }
 
 void SEditableText::OnBeginTextEdit(const FText& InText)
@@ -737,41 +708,21 @@ void SEditableText::OnTextCommitted(const FText& InText, const ETextCommit::Type
 void SEditableText::OnCursorMoved(const FTextLocation& InLocation)
 {
 	FTextSelection TextSelection = GetSelection();
-
+	
 	OnCursorMovedWithSelectionCallback.ExecuteIfBound(InLocation, TextSelection);
 	Invalidate(EInvalidateWidgetReason::Layout);
 }
 
-float SEditableText::UpdateAndClampHorizontalScrollBar(const float InViewOffset, const float InViewFraction, const EVisibility InVisibilityOverride)
+float SEditableText::UpdateAndClampHorizontalScrollBar(const float InViewOffset, const float InViewFraction, const EVisibility InVisiblityOverride)
 {
-	return SlateEditableTextLayout.GetScrollOffset().X;
+	return EditableTextLayout->GetScrollOffset().X;
 }
 
-float SEditableText::UpdateAndClampVerticalScrollBar(const float InViewOffset, const float InViewFraction, const EVisibility InVisibilityOverride)
+float SEditableText::UpdateAndClampVerticalScrollBar(const float InViewOffset, const float InViewFraction, const EVisibility InVisiblityOverride)
 {
 	return 0.0f;
 }
 
-ETextJustify::Type SEditableText::GetJustification() const
-{
-	return SlateEditableTextLayout.GetJustification();
-}
-
-FSlateColor SEditableText::GetColorAndOpacity() const
-{
-	return ColorAndOpacityAttribute.Get();
-}
-
-const FSlateBrush* SEditableText::GetBackgroundImageSelected() const
-{
-	return BackgroundImageSelectedAttribute.Get();
-}
-
-float SEditableText::GetMinDesiredWidth() const
-{
-	return MinDesiredWidthAttribute.Get();
-}
-
 #if WITH_ACCESSIBILITY
 TSharedRef<FSlateAccessibleWidget> SEditableText::CreateAccessibleWidget()
 {
@@ -786,5 +737,8 @@ TOptional<FText> SEditableText::GetDefaultAccessibleText(EAccessibleType Accessi
 
 void SEditableText::ToggleVirtualKeyboard(bool bShow, const int32 UserIndex)
 {
-	SlateEditableTextLayout.ToggleVirtualKeyboard(bShow, UserIndex);
+	if (EditableTextLayout.IsValid())
+	{
+		EditableTextLayout->ToggleVirtualKeyboard(bShow, UserIndex);
+	}
 }
diff --git a/Engine/Source/Runtime/Slate/Private/Widgets/Text/SMultiLineEditableText.cpp b/Engine/Source/Runtime/Slate/Private/Widgets/Text/SMultiLineEditableText.cpp
--- a/Engine/Source/Runtime/Slate/Private/Widgets/Text/SMultiLineEditableText.cpp
+++ b/Engine/Source/Runtime/Slate/Private/Widgets/Text/SMultiLineEditableText.cpp
@@ -10,25 +10,22 @@
 #include "Framework/Text/TextEditHelper.h"
 #include "Framework/Text/PlainTextLayoutMarshaller.h"
 #include "Widgets/Text/SlateEditableTextLayout.h"
-#include "Widgets/Text/SlateTextBlockLayout.h"
 #include "Types/ReflectionMetadata.h"
 #include "Types/TrackedMetaData.h"
 
-SLATE_IMPLEMENT_WIDGET(SMultiLineEditableText)
-void SMultiLineEditableText::PrivateRegisterAttributes(FSlateAttributeInitializer& AttributeInitializer)
-{
-	FSlateEditableTextLayout::RegisterNestedAttributes<&SMultiLineEditableText::SlateEditableTextLayout>(AttributeInitializer);
-}
-
 SMultiLineEditableText::SMultiLineEditableText()
-	: SlateEditableTextLayout(*this)
-	, bSelectAllTextWhenFocused(false)
+	: bSelectAllTextWhenFocused(false)
 	, bIsReadOnly(false)
 	, AmountScrolledWhileRightMouseDown(0.0f)
 	, bIsSoftwareCursor(false)
 {
 }
 
+SMultiLineEditableText::~SMultiLineEditableText()
+{
+	// Needed to avoid "deletion of pointer to incomplete type 'FSlateEditableTextLayout'; no destructor called" error when using TUniquePtr
+}
+
 void SMultiLineEditableText::Construct( const FArguments& InArgs )
 {
 	bIsReadOnly = InArgs._IsReadOnly;
@@ -85,19 +82,19 @@ void SMultiLineEditableText::Construct( const FArguments& InArgs )
 		Marshaller = FPlainTextLayoutMarshaller::Create();
 	}
 
-	SlateEditableTextLayout.Construct(InArgs._Text, TextStyle, InArgs._TextShapingMethod, InArgs._TextFlowDirection, InArgs._CreateSlateTextLayout, Marshaller.ToSharedRef(), Marshaller.ToSharedRef());
-	SlateEditableTextLayout.SetHintText(InArgs._HintText);
-	SlateEditableTextLayout.SetSearchText(InArgs._SearchText);
-	SlateEditableTextLayout.SetTextWrapping(InArgs._WrapTextAt, InArgs._AutoWrapText, InArgs._WrappingPolicy);
-	SlateEditableTextLayout.SetMargin(InArgs._Margin);
-	SlateEditableTextLayout.SetJustification(InArgs._Justification);
-	SlateEditableTextLayout.SetLineHeightPercentage(InArgs._LineHeightPercentage);
-	SlateEditableTextLayout.SetDebugSourceInfo(TAttribute<FString>::Create(TAttribute<FString>::FGetter::CreateLambda([this]{ return FReflectionMetaData::GetWidgetDebugInfo(this); })));
-	SlateEditableTextLayout.SetOverflowPolicy(InArgs._OverflowPolicy);
-	SlateEditableTextLayout.SetFontFacesLoadingPaintPolicy(InArgs._FontFacesLoadingPaintPolicy);
+	EditableTextLayout = MakeUnique<FSlateEditableTextLayout>(*this, InArgs._Text, TextStyle, InArgs._TextShapingMethod, InArgs._TextFlowDirection, InArgs._CreateSlateTextLayout, Marshaller.ToSharedRef(), Marshaller.ToSharedRef());
+	EditableTextLayout->SetHintText(InArgs._HintText);
+	EditableTextLayout->SetSearchText(InArgs._SearchText);
+	EditableTextLayout->SetTextWrapping(InArgs._WrapTextAt, InArgs._AutoWrapText, InArgs._WrappingPolicy);
+	EditableTextLayout->SetMargin(InArgs._Margin);
+	EditableTextLayout->SetJustification(InArgs._Justification);
+	EditableTextLayout->SetLineHeightPercentage(InArgs._LineHeightPercentage);
+	EditableTextLayout->SetDebugSourceInfo(TAttribute<FString>::Create(TAttribute<FString>::FGetter::CreateLambda([this]{ return FReflectionMetaData::GetWidgetDebugInfo(this); })));
+	EditableTextLayout->SetOverflowPolicy(InArgs._OverflowPolicy);
+	EditableTextLayout->SetFontFacesLoadingPaintPolicy(InArgs._FontFacesLoadingPaintPolicy);
 	if (InArgs._OnAllFontFacesFinishLoading.IsBound())
 	{
-		SlateEditableTextLayout.SetOnAllFontFacesFinishLoadingDelegate(InArgs._OnAllFontFacesFinishLoading);
+		EditableTextLayout->SetOnAllFontFacesFinishLoadingDelegate(InArgs._OnAllFontFacesFinishLoading);
 	}
 
 	// build context menu extender
@@ -109,139 +106,139 @@ void SMultiLineEditableText::Construct( const FArguments& InArgs )
 
 void SMultiLineEditableText::GetCurrentTextLine(FString& OutTextLine) const
 {
-	SlateEditableTextLayout.GetCurrentTextLine(OutTextLine);
+	EditableTextLayout->GetCurrentTextLine(OutTextLine);
 }
 
 void SMultiLineEditableText::GetTextLine(const int32 InLineIndex, FString& OutTextLine) const
 {
-	SlateEditableTextLayout.GetTextLine(InLineIndex, OutTextLine);
+	EditableTextLayout->GetTextLine(InLineIndex, OutTextLine);
 }
 
 void SMultiLineEditableText::SetText(const TAttribute< FText >& InText)
 {
-	SlateEditableTextLayout.SetText(InText);
+	EditableTextLayout->SetText(InText);
 }
 
-int32 SMultiLineEditableText::GetTextLineCount() const
+int32 SMultiLineEditableText::GetTextLineCount()
 {
-	return SlateEditableTextLayout.GetTextLineCount();
+	return EditableTextLayout->GetTextLineCount();
 }
 
 FText SMultiLineEditableText::GetText() const
 {
-	return SlateEditableTextLayout.GetText();
+	return EditableTextLayout->GetText();
 }
 
 FText SMultiLineEditableText::GetPlainText() const
 {
-	return SlateEditableTextLayout.GetPlainText();
+	return EditableTextLayout->GetPlainText();
 }
 
 void SMultiLineEditableText::SetHintText(const TAttribute< FText >& InHintText)
 {
-	SlateEditableTextLayout.SetHintText(InHintText);
+	EditableTextLayout->SetHintText(InHintText);
 }
 
 FText SMultiLineEditableText::GetHintText() const
 {
-	return SlateEditableTextLayout.GetHintText();
+	return EditableTextLayout->GetHintText();
 }
 
 void SMultiLineEditableText::SetSearchText(const TAttribute<FText>& InSearchText)
 {
-	SlateEditableTextLayout.SetSearchText(InSearchText);
+	EditableTextLayout->SetSearchText(InSearchText);
 }
 
 FText SMultiLineEditableText::GetSearchText() const
 {
-	return SlateEditableTextLayout.GetSearchText();
+	return EditableTextLayout->GetSearchText();
 }
 
 int32 SMultiLineEditableText::GetSearchResultIndex() const
 {
-	return SlateEditableTextLayout.GetSearchResultIndex();
+	return EditableTextLayout->GetSearchResultIndex();
 }
 
 int32 SMultiLineEditableText::GetNumSearchResults() const
 {
-	return SlateEditableTextLayout.GetNumSearchResults();
+	return EditableTextLayout->GetNumSearchResults();
 }
 
 void SMultiLineEditableText::SetTextStyle(const FTextBlockStyle* InTextStyle)
 {
 	if (InTextStyle)
 	{
-		SlateEditableTextLayout.SetTextStyle(*InTextStyle);
+		EditableTextLayout->SetTextStyle(*InTextStyle);
 	}
 	else
 	{
 		FArguments Defaults;
 		check(Defaults._TextStyle);
-		SlateEditableTextLayout.SetTextStyle(*Defaults._TextStyle);
+		EditableTextLayout->SetTextStyle(*Defaults._TextStyle);
 	}
 }
 
 void SMultiLineEditableText::SetFont(const TAttribute< FSlateFontInfo >& InNewFont)
 {
-	FTextBlockStyle TextStyle = SlateEditableTextLayout.GetTextStyle();
+	FTextBlockStyle TextStyle = EditableTextLayout->GetTextStyle();
 	TextStyle.SetFont(InNewFont.Get());
-	SlateEditableTextLayout.SetTextStyle(TextStyle);
+	EditableTextLayout->SetTextStyle(TextStyle);
 }
 
 FSlateFontInfo SMultiLineEditableText::GetFont() const
 {
-	FTextBlockStyle TextStyle = SlateEditableTextLayout.GetTextStyle();
+	FTextBlockStyle TextStyle = EditableTextLayout->GetTextStyle();
 	return TextStyle.Font;
 }
 
 void SMultiLineEditableText::SetTextShapingMethod(const TOptional<ETextShapingMethod>& InTextShapingMethod)
 {
-	SlateEditableTextLayout.SetTextShapingMethod(InTextShapingMethod);
+	EditableTextLayout->SetTextShapingMethod(InTextShapingMethod);
 }
 
 void SMultiLineEditableText::SetTextFlowDirection(const TOptional<ETextFlowDirection>& InTextFlowDirection)
 {
-	SlateEditableTextLayout.SetTextFlowDirection(InTextFlowDirection);
+	EditableTextLayout->SetTextFlowDirection(InTextFlowDirection);
 }
 
 void SMultiLineEditableText::SetWrapTextAt(const TAttribute<float>& InWrapTextAt)
 {
-	SlateEditableTextLayout.SetWrapTextAt(InWrapTextAt);
+	EditableTextLayout->SetWrapTextAt(InWrapTextAt);
 }
 
 void SMultiLineEditableText::SetAutoWrapText(const TAttribute<bool>& InAutoWrapText)
 {
-	SlateEditableTextLayout.SetAutoWrapText(InAutoWrapText);
+	EditableTextLayout->SetAutoWrapText(InAutoWrapText);
 }
 
 void SMultiLineEditableText::SetWrappingPolicy(const TAttribute<ETextWrappingPolicy>& InWrappingPolicy)
 {
-	SlateEditableTextLayout.SetWrappingPolicy(InWrappingPolicy);
+	EditableTextLayout->SetWrappingPolicy(InWrappingPolicy);
 }
 
 void SMultiLineEditableText::SetLineHeightPercentage(const TAttribute<float>& InLineHeightPercentage)
 {
-	SlateEditableTextLayout.SetLineHeightPercentage(InLineHeightPercentage);
+	EditableTextLayout->SetLineHeightPercentage(InLineHeightPercentage);
 }
 
 void SMultiLineEditableText::SetApplyLineHeightToBottomLine(const TAttribute<bool>& InApplyLineHeightToBottomLine)
 {
-	SlateEditableTextLayout.SetApplyLineHeightToBottomLine(InApplyLineHeightToBottomLine);
+	EditableTextLayout->SetApplyLineHeightToBottomLine(InApplyLineHeightToBottomLine);
 }
 
 void SMultiLineEditableText::SetMargin(const TAttribute<FMargin>& InMargin)
 {
-	SlateEditableTextLayout.SetMargin(InMargin);
+	EditableTextLayout->SetMargin(InMargin);
 }
 
 void SMultiLineEditableText::SetJustification(const TAttribute<ETextJustify::Type>& InJustification)
 {
-	SlateEditableTextLayout.SetJustification(InJustification);
+	EditableTextLayout->SetJustification(InJustification);
 }
 
 void SMultiLineEditableText::SetOverflowPolicy(TOptional<ETextOverflowPolicy> InOverflowPolicy)
 {
-	SlateEditableTextLayout.SetOverflowPolicy(InOverflowPolicy);
+	EditableTextLayout->SetOverflowPolicy(InOverflowPolicy);
 }
 
 void SMultiLineEditableText::SetAllowContextMenu(const TAttribute< bool >& InAllowContextMenu)
@@ -286,18 +283,21 @@ void SMultiLineEditableText::SetClearKeyboardFocusOnCommit(const TAttribute<bool
 
 void SMultiLineEditableText::SetFontFacesLoadingPaintPolicy(EFontFacesLoadingPaintPolicy InFontFacesLoadingPaintPolicy)
 {
-	SlateEditableTextLayout.SetFontFacesLoadingPaintPolicy(InFontFacesLoadingPaintPolicy);
+	if (EditableTextLayout)
+	{
+		EditableTextLayout->SetFontFacesLoadingPaintPolicy(InFontFacesLoadingPaintPolicy);
+	}
 }
 
 void SMultiLineEditableText::OnHScrollBarMoved(const float InScrollOffsetFraction)
 {
-	SlateEditableTextLayout.SetHorizontalScrollFraction(InScrollOffsetFraction);
+	EditableTextLayout->SetHorizontalScrollFraction(InScrollOffsetFraction);
 	OnHScrollBarUserScrolled.ExecuteIfBound(InScrollOffsetFraction);
 }
 
 void SMultiLineEditableText::OnVScrollBarMoved(const float InScrollOffsetFraction)
 {
-	SlateEditableTextLayout.SetVerticalScrollFraction(InScrollOffsetFraction);
+	EditableTextLayout->SetVerticalScrollFraction(InScrollOffsetFraction);
 	OnVScrollBarUserScrolled.ExecuteIfBound(InScrollOffsetFraction);
 }
 
@@ -382,7 +382,7 @@ void SMultiLineEditableText::EnsureActiveTick()
 	auto DoActiveTick = [this](double InCurrentTime, float InDeltaTime) -> EActiveTimerReturnType
 	{
 		// Continue if we still have focus, otherwise treat as a fire-and-forget Tick() request
-		const bool bShouldAppearFocused = HasKeyboardFocus() || SlateEditableTextLayout.HasActiveContextMenu();
+		const bool bShouldAppearFocused = HasKeyboardFocus() || EditableTextLayout->HasActiveContextMenu();
 		return (bShouldAppearFocused) ? EActiveTimerReturnType::Continue : EActiveTimerReturnType::Stop;
 	};
 
@@ -436,7 +436,7 @@ TSharedPtr<SWidget> SMultiLineEditableText::BuildContextMenuContent() const
 		return OnContextMenuOpening.Execute();
 	}
 
-	return SlateEditableTextLayout.BuildDefaultContextMenu(MenuExtender);
+	return EditableTextLayout->BuildDefaultContextMenu(MenuExtender);
 }
 
 void SMultiLineEditableText::OnBeginTextEdit(const FText& InText)
@@ -473,7 +473,7 @@ float SMultiLineEditableText::UpdateAndClampHorizontalScrollBar(const float InVi
 		}
 	}
 
-	return SlateEditableTextLayout.GetScrollOffset().X;
+	return EditableTextLayout->GetScrollOffset().X;
 }
 
 float SMultiLineEditableText::UpdateAndClampVerticalScrollBar(const float InViewOffset, const float InViewFraction, const EVisibility InVisiblityOverride)
@@ -489,134 +489,134 @@ float SMultiLineEditableText::UpdateAndClampVerticalScrollBar(const float InView
 		}
 	}
 
-	return SlateEditableTextLayout.GetScrollOffset().Y;
+	return EditableTextLayout->GetScrollOffset().Y;
 }
 
 void SMultiLineEditableText::BeginEditTransaction()
 {
-	SlateEditableTextLayout.BeginEditTransaction();
+	EditableTextLayout->BeginEditTransation();
 }
 
 void SMultiLineEditableText::EndEditTransaction()
 {
-	SlateEditableTextLayout.EndEditTransaction();
+	EditableTextLayout->EndEditTransaction();
 }
 
 FReply SMultiLineEditableText::OnFocusReceived( const FGeometry& MyGeometry, const FFocusEvent& InFocusEvent )
 {
-	SlateEditableTextLayout.HandleFocusReceived(InFocusEvent);
+	EditableTextLayout->HandleFocusReceived(InFocusEvent);
 	return FReply::Handled();
 }
 
 void SMultiLineEditableText::OnFocusLost( const FFocusEvent& InFocusEvent )
 {
 	bIsSoftwareCursor = false;
-	SlateEditableTextLayout.HandleFocusLost(InFocusEvent);
+	EditableTextLayout->HandleFocusLost(InFocusEvent);
 }
 
 bool SMultiLineEditableText::AnyTextSelected() const
 {
-	return SlateEditableTextLayout.AnyTextSelected();
+	return EditableTextLayout->AnyTextSelected();
 }
 
 void SMultiLineEditableText::SelectAllText()
 {
-	SlateEditableTextLayout.SelectAllText();
+	EditableTextLayout->SelectAllText();
 }
 
 void SMultiLineEditableText::SelectText(const FTextLocation& InSelectionStart, const FTextLocation& InCursorLocation)
 {
-	SlateEditableTextLayout.SelectText(InSelectionStart, InCursorLocation);
+	EditableTextLayout->SelectText(InSelectionStart, InCursorLocation);
 }
 
 void SMultiLineEditableText::ClearSelection()
 {
-	SlateEditableTextLayout.ClearSelection();
+	EditableTextLayout->ClearSelection();
 }
 
 FText SMultiLineEditableText::GetSelectedText() const
 {
-	return SlateEditableTextLayout.GetSelectedText();
+	return EditableTextLayout->GetSelectedText();
 }
 
 FTextSelection SMultiLineEditableText::GetSelection() const
 {
-	return SlateEditableTextLayout.GetSelection();
+	return EditableTextLayout->GetSelection();
 }
 
 void SMultiLineEditableText::DeleteSelectedText()
 {
-	SlateEditableTextLayout.DeleteSelectedText();
+	EditableTextLayout->DeleteSelectedText();
 }
 
 void SMultiLineEditableText::InsertTextAtCursor(const FText& InText)
 {
-	SlateEditableTextLayout.InsertTextAtCursor(InText.ToString());
+	EditableTextLayout->InsertTextAtCursor(InText.ToString());
 }
 
 void SMultiLineEditableText::InsertTextAtCursor(const FString& InString)
 {
-	SlateEditableTextLayout.InsertTextAtCursor(InString);
+	EditableTextLayout->InsertTextAtCursor(InString);
 }
 
 void SMultiLineEditableText::InsertRunAtCursor(TSharedRef<IRun> InRun)
 {
-	SlateEditableTextLayout.InsertRunAtCursor(InRun);
+	EditableTextLayout->InsertRunAtCursor(InRun);
 }
 
 void SMultiLineEditableText::GoTo(const FTextLocation& NewLocation)
 {
-	SlateEditableTextLayout.GoTo(NewLocation);
+	EditableTextLayout->GoTo(NewLocation);
 }
 
 void SMultiLineEditableText::GoTo(const ETextLocation NewLocation)
 {
-	SlateEditableTextLayout.GoTo(NewLocation);
+	EditableTextLayout->GoTo(NewLocation);
 }
 
 void SMultiLineEditableText::ScrollTo(const FTextLocation& NewLocation)
 {
-	SlateEditableTextLayout.ScrollTo(NewLocation);
+	EditableTextLayout->ScrollTo(NewLocation);
 }
 
 void SMultiLineEditableText::ScrollTo(const ETextLocation NewLocation)
 {
-	SlateEditableTextLayout.ScrollTo(NewLocation);
+	EditableTextLayout->ScrollTo(NewLocation);
 }
 
 void SMultiLineEditableText::ApplyToSelection(const FRunInfo& InRunInfo, const FTextBlockStyle& InStyle)
 {
-	SlateEditableTextLayout.ApplyToSelection(InRunInfo, InStyle);
+	EditableTextLayout->ApplyToSelection(InRunInfo, InStyle);
 }
 
 void SMultiLineEditableText::BeginSearch(const FText& InSearchText, const ESearchCase::Type InSearchCase, const bool InReverse)
 {
-	SlateEditableTextLayout.BeginSearch(InSearchText, InSearchCase, InReverse);
+	EditableTextLayout->BeginSearch(InSearchText, InSearchCase, InReverse);
 }
 
 void SMultiLineEditableText::AdvanceSearch(const bool InReverse)
 {
-	SlateEditableTextLayout.AdvanceSearch(InReverse);
+	EditableTextLayout->AdvanceSearch(InReverse);
 }
 
 TSharedPtr<const IRun> SMultiLineEditableText::GetRunUnderCursor() const
 {
-	return SlateEditableTextLayout.GetRunUnderCursor();
+	return EditableTextLayout->GetRunUnderCursor();
 }
 
 TArray<TSharedRef<const IRun>> SMultiLineEditableText::GetSelectedRuns() const
 {
-	return SlateEditableTextLayout.GetSelectedRuns();
+	return EditableTextLayout->GetSelectedRuns();
 }
 
 FTextLocation SMultiLineEditableText::GetCursorLocation() const
 {
-	return SlateEditableTextLayout.GetCursorLocation();
+	return EditableTextLayout->GetCursorLocation();
 }
 
 TCHAR SMultiLineEditableText::GetCharacterAt(const FTextLocation& Location) const
 {
-	return SlateEditableTextLayout.GetCharacterAt(Location);
+	return EditableTextLayout->GetCharacterAt(Location);
 }
 
 TSharedPtr<const SScrollBar> SMultiLineEditableText::GetHScrollBar() const
@@ -631,7 +631,7 @@ TSharedPtr<const SScrollBar> SMultiLineEditableText::GetVScrollBar() const
 
 void SMultiLineEditableText::Refresh()
 {
-	SlateEditableTextLayout.Refresh();
+	EditableTextLayout->Refresh();
 }
 
 void SMultiLineEditableText::ForceScroll(int32 UserIndex, float ScrollAxisMagnitude)
@@ -645,25 +645,25 @@ void SMultiLineEditableText::ForceScroll(int32 UserIndex, float ScrollAxisMagnit
 
 void SMultiLineEditableText::Tick( const FGeometry& AllottedGeometry, const double InCurrentTime, const float InDeltaTime )
 {
-	SlateEditableTextLayout.Tick(AllottedGeometry, InCurrentTime, InDeltaTime);
+	EditableTextLayout->Tick(AllottedGeometry, InCurrentTime, InDeltaTime);
 }
 
 int32 SMultiLineEditableText::OnPaint( const FPaintArgs& Args, const FGeometry& AllottedGeometry, const FSlateRect& MyCullingRect, FSlateWindowElementList& OutDrawElements, int32 LayerId, const FWidgetStyle& InWidgetStyle, bool bParentEnabled ) const
 {
-	const FTextBlockStyle& EditableTextStyle = SlateEditableTextLayout.GetTextStyle();
+	const FTextBlockStyle& EditableTextStyle = EditableTextLayout->GetTextStyle();
 	const FLinearColor ForegroundColor = EditableTextStyle.ColorAndOpacity.GetColor(InWidgetStyle);
 
 	FWidgetStyle TextWidgetStyle = FWidgetStyle(InWidgetStyle)
 		.SetForegroundColor(ForegroundColor);
 
-	const bool bAutoWrap = SlateEditableTextLayout.GetAutoWrapText();
-	if (bAutoWrap && SlateEditableTextLayout.GetComputedWrappingWidth() == 0)
+	const bool bAutoWrap = EditableTextLayout->GetAutoWrapText();
+	if (bAutoWrap && EditableTextLayout->GetComputedWrappingWidth() == 0)
 	{
 		QUICK_SCOPE_CYCLE_COUNTER(STAT_Slate_AutoWrapRecompute)
 		const_cast<SMultiLineEditableText*>(this)->CacheDesiredSize(GetPrepassLayoutScaleMultiplier());
 	}
 
-	LayerId = SlateEditableTextLayout.OnPaint(Args, AllottedGeometry, MyCullingRect, OutDrawElements, LayerId, TextWidgetStyle, ShouldBeEnabled(bParentEnabled));
+	LayerId = EditableTextLayout->OnPaint(Args, AllottedGeometry, MyCullingRect, OutDrawElements, LayerId, TextWidgetStyle, ShouldBeEnabled(bParentEnabled));
 
 	if (bIsSoftwareCursor)
 	{
@@ -683,23 +683,23 @@ int32 SMultiLineEditableText::OnPaint( const FPaintArgs& Args, const FGeometry&
 
 void SMultiLineEditableText::CacheDesiredSize(float LayoutScaleMultiplier)
 {
-	SlateEditableTextLayout.CacheDesiredSize(LayoutScaleMultiplier);
+	EditableTextLayout->CacheDesiredSize(LayoutScaleMultiplier);
 	SWidget::CacheDesiredSize(LayoutScaleMultiplier);
 }
 
 FVector2D SMultiLineEditableText::ComputeDesiredSize( float LayoutScaleMultiplier ) const
 {
-	return SlateEditableTextLayout.ComputeDesiredSize(LayoutScaleMultiplier);
+	return EditableTextLayout->ComputeDesiredSize(LayoutScaleMultiplier);
 }
 
 FChildren* SMultiLineEditableText::GetChildren()
 {
-	return SlateEditableTextLayout.GetChildren();
+	return EditableTextLayout->GetChildren();
 }
 
 void SMultiLineEditableText::OnArrangeChildren( const FGeometry& AllottedGeometry, FArrangedChildren& ArrangedChildren ) const
 {
-	SlateEditableTextLayout.OnArrangeChildren(AllottedGeometry, ArrangedChildren);
+	EditableTextLayout->OnArrangeChildren(AllottedGeometry, ArrangedChildren);
 }
 
 bool SMultiLineEditableText::SupportsKeyboardFocus() const
@@ -719,7 +719,7 @@ FReply SMultiLineEditableText::OnKeyChar( const FGeometry& MyGeometry,const FCha
 
 	if (!Reply.IsEventHandled())
 	{
-		Reply = SlateEditableTextLayout.HandleKeyChar(InCharacterEvent);
+		Reply = EditableTextLayout->HandleKeyChar(InCharacterEvent);
 	}
 
 	return Reply;
@@ -737,7 +737,7 @@ FReply SMultiLineEditableText::OnKeyDown( const FGeometry& MyGeometry, const FKe
 
 	if (!Reply.IsEventHandled())
 	{
-		Reply = SlateEditableTextLayout.HandleKeyDown(InKeyEvent);
+		Reply = EditableTextLayout->HandleKeyDown(InKeyEvent);
 	}
 
 	return Reply;
@@ -745,20 +745,20 @@ FReply SMultiLineEditableText::OnKeyDown( const FGeometry& MyGeometry, const FKe
 
 FReply SMultiLineEditableText::OnKeyUp( const FGeometry& MyGeometry, const FKeyEvent& InKeyEvent )
 {
-	return SlateEditableTextLayout.HandleKeyUp(InKeyEvent);
+	return EditableTextLayout->HandleKeyUp(InKeyEvent);
 }
 
-FReply SMultiLineEditableText::OnMouseButtonDown( const FGeometry& MyGeometry, const FPointerEvent& MouseEvent )
+FReply SMultiLineEditableText::OnMouseButtonDown( const FGeometry& MyGeometry, const FPointerEvent& MouseEvent ) 
 {
 	if (MouseEvent.GetEffectingButton() == EKeys::RightMouseButton)
 	{
 		AmountScrolledWhileRightMouseDown = 0.0f;
 	}
 
-	return SlateEditableTextLayout.HandleMouseButtonDown(MyGeometry, MouseEvent);
+	return EditableTextLayout->HandleMouseButtonDown(MyGeometry, MouseEvent);
 }
 
-FReply SMultiLineEditableText::OnMouseButtonUp( const FGeometry& MyGeometry, const FPointerEvent& MouseEvent )
+FReply SMultiLineEditableText::OnMouseButtonUp( const FGeometry& MyGeometry, const FPointerEvent& MouseEvent ) 
 {
 	if (MouseEvent.GetEffectingButton() == EKeys::RightMouseButton)
 	{
@@ -774,7 +774,7 @@ FReply SMultiLineEditableText::OnMouseButtonUp( const FGeometry& MyGeometry, con
 		}
 	}
 
-	return SlateEditableTextLayout.HandleMouseButtonUp(MyGeometry, MouseEvent);
+	return EditableTextLayout->HandleMouseButtonUp(MyGeometry, MouseEvent);
 }
 
 FReply SMultiLineEditableText::OnMouseMove( const FGeometry& MyGeometry, const FPointerEvent& MouseEvent )
@@ -789,11 +789,11 @@ FReply SMultiLineEditableText::OnMouseMove( const FGeometry& MyGeometry, const F
 
 		if (IsRightClickScrolling())
 		{
-			const FVector2f PreviousScrollOffset = SlateEditableTextLayout.GetScrollOffset();
+			const FVector2f PreviousScrollOffset = EditableTextLayout->GetScrollOffset();
 
 			FVector2f NewScrollOffset = PreviousScrollOffset;
 			NewScrollOffset.Y -= ScrollByAmount;
-			SlateEditableTextLayout.SetScrollOffset(NewScrollOffset, MyGeometry);
+			EditableTextLayout->SetScrollOffset(NewScrollOffset, MyGeometry);
 
 			if (!bIsSoftwareCursor)
 			{
@@ -803,7 +803,7 @@ FReply SMultiLineEditableText::OnMouseMove( const FGeometry& MyGeometry, const F
 
 			if (PreviousScrollOffset.Y != NewScrollOffset.Y)
 			{
-				const float ScrollMax = SlateEditableTextLayout.GetSize().Y - MyGeometry.GetLocalSize().Y;
+				const float ScrollMax = EditableTextLayout->GetSize().Y - MyGeometry.GetLocalSize().Y;
 				const float ScrollbarOffset = (ScrollMax != 0.0f) ? NewScrollOffset.Y / ScrollMax : 0.0f;
 				OnVScrollBarUserScrolled.ExecuteIfBound(ScrollbarOffset);
 				SoftwareCursorPosition.Y += (PreviousScrollOffset.Y - NewScrollOffset.Y);
@@ -813,7 +813,7 @@ FReply SMultiLineEditableText::OnMouseMove( const FGeometry& MyGeometry, const F
 		}
 	}
 
-	return SlateEditableTextLayout.HandleMouseMove(MyGeometry, MouseEvent);
+	return EditableTextLayout->HandleMouseMove(MyGeometry, MouseEvent);
 }
 
 FReply SMultiLineEditableText::OnMouseWheel( const FGeometry& MyGeometry, const FPointerEvent& MouseEvent )
@@ -822,15 +822,15 @@ FReply SMultiLineEditableText::OnMouseWheel( const FGeometry& MyGeometry, const
 	{
 		const float ScrollAmount = -MouseEvent.GetWheelDelta() * GetGlobalScrollAmount();
 
-		const FVector2f PreviousScrollOffset = SlateEditableTextLayout.GetScrollOffset();
-
+		const FVector2f PreviousScrollOffset = EditableTextLayout->GetScrollOffset();
+		
 		FVector2f NewScrollOffset = PreviousScrollOffset;
 		NewScrollOffset.Y += ScrollAmount;
-		SlateEditableTextLayout.SetScrollOffset(NewScrollOffset, MyGeometry);
+		EditableTextLayout->SetScrollOffset(NewScrollOffset, MyGeometry);
 
 		if (PreviousScrollOffset.Y != NewScrollOffset.Y)
 		{
-			const float ScrollMax = SlateEditableTextLayout.GetSize().Y - MyGeometry.GetLocalSize().Y;
+			const float ScrollMax = EditableTextLayout->GetSize().Y - MyGeometry.GetLocalSize().Y;
 			const float ScrollbarOffset = (ScrollMax != 0.0f) ? NewScrollOffset.Y / ScrollMax : 0.0f;
 			OnVScrollBarUserScrolled.ExecuteIfBound(ScrollbarOffset);
 			return FReply::Handled();
@@ -842,7 +842,7 @@ FReply SMultiLineEditableText::OnMouseWheel( const FGeometry& MyGeometry, const
 
 FReply SMultiLineEditableText::OnMouseButtonDoubleClick(const FGeometry& MyGeometry, const FPointerEvent& MouseEvent)
 {
-	return SlateEditableTextLayout.HandleMouseButtonDoubleClick(MyGeometry, MouseEvent);
+	return EditableTextLayout->HandleMouseButtonDoubleClick(MyGeometry, MouseEvent);
 }
 
 FCursorReply SMultiLineEditableText::OnCursorQuery( const FGeometry& MyGeometry, const FPointerEvent& CursorEvent ) const
@@ -866,6 +866,7 @@ bool SMultiLineEditableText::ComputeVolatility() const
 {
 	return SWidget::ComputeVolatility()
 		|| HasKeyboardFocus()
+		|| EditableTextLayout->ComputeVolatility()
 		|| bIsReadOnly.IsBound();
 }
 
diff --git a/Engine/Source/Runtime/Slate/Private/Widgets/Text/SlateEditableTextLayout.cpp b/Engine/Source/Runtime/Slate/Private/Widgets/Text/SlateEditableTextLayout.cpp
--- a/Engine/Source/Runtime/Slate/Private/Widgets/Text/SlateEditableTextLayout.cpp
+++ b/Engine/Source/Runtime/Slate/Private/Widgets/Text/SlateEditableTextLayout.cpp
@@ -21,7 +21,7 @@
 
 /**
  * Ensure that text transactions are always completed.
- * i.e. never forget to call EndEditTransaction();
+ * i.e. never forget to call FinishChangingText();
  */
 struct FScopedEditableTextTransaction
 {
@@ -29,7 +29,7 @@ public:
 	explicit FScopedEditableTextTransaction(FSlateEditableTextLayout& InEditableTextLayout)
 		: EditableTextLayout(&InEditableTextLayout)
 	{
-		EditableTextLayout->BeginEditTransaction();
+		EditableTextLayout->BeginEditTransation();
 	}
 
 	~FScopedEditableTextTransaction()
@@ -65,7 +65,7 @@ bool IsCharAllowed(const TCHAR InChar)
 
 }
 
-void FSlateEditableTextLayout::Construct(TAttribute<FText> InInitialText, FTextBlockStyle InTextStyle, const TOptional<ETextShapingMethod>& InTextShapingMethod, const TOptional<ETextFlowDirection>& InTextFlowDirection, const FCreateSlateTextLayout& InCreateSlateTextLayout, TSharedRef<ITextLayoutMarshaller> InTextMarshaller, TSharedRef<ITextLayoutMarshaller> InHintTextMarshaller)
+FSlateEditableTextLayout::FSlateEditableTextLayout(ISlateEditableTextWidget& InOwnerWidget, const TAttribute<FText>& InInitialText, FTextBlockStyle InTextStyle, const TOptional<ETextShapingMethod> InTextShapingMethod, const TOptional<ETextFlowDirection> InTextFlowDirection, const FCreateSlateTextLayout& InCreateSlateTextLayout, TSharedRef<ITextLayoutMarshaller> InTextMarshaller, TSharedRef<ITextLayoutMarshaller> InHintTextMarshaller)
 {
 	CreateSlateTextLayout = InCreateSlateTextLayout;
 	if (!CreateSlateTextLayout.IsBound())
@@ -76,28 +76,34 @@ void FSlateEditableTextLayout::Construct(TAttribute<FText> InInitialText, FTextB
 			});
 	}
 
-	const TSharedRef<SWidget> OwnerSlateWidget = OwnerWidget->GetSlateWidget();
-
-	Marshaller = MoveTemp(InTextMarshaller);
-	HintMarshaller = MoveTemp(InHintTextMarshaller);
-	TextStyle = MoveTemp(InTextStyle);
+	OwnerWidget = &InOwnerWidget;
+	Marshaller = InTextMarshaller;
+	HintMarshaller = InHintTextMarshaller;
+	TextStyle = InTextStyle;
 	HintTextStyle = TextStyle;
-	TextLayout = CreateSlateTextLayout.Execute(&OwnerSlateWidget.Get(), TextStyle);
+	TextLayout = CreateSlateTextLayout.Execute(&InOwnerWidget.GetSlateWidget().Get(), TextStyle);
+
+	WrapTextAt = 0.0f;
+	AutoWrapText = false;
+	WrappingPolicy = ETextWrappingPolicy::DefaultWrapping;
 
+	Margin = FMargin(0);
+	Justification = ETextJustify::Left;
+	LineHeightPercentage = 1.0f;
 	DebugSourceInfo = FString();
 
 	SearchCase = ESearchCase::IgnoreCase;
 
 	GraphemeBreakIterator = GSlateUseSharedBreakIterator ? FBreakIterator::GetCharacterBoundaryIterator() : FBreakIterator::CreateCharacterBoundaryIterator();
 
-	BoundText.Assign(OwnerSlateWidget.Get(), MoveTemp(InInitialText), ESlateAttributeBindAction::Initialize);
+	BoundText = InInitialText;
 
 	// Set the initial text - the same as SetText, but doesn't call OnTextChanged as that can cause a crash during construction
 	{
 		const bool bIsPassword = OwnerWidget->IsTextPassword();
 		TextLayout->SetIsPassword(bIsPassword);
 
-		const FText& InitialTextToSet = BoundText.Get();
+		const FText& InitialTextToSet = BoundText.Get(FText::GetEmpty());
 		SetEditableText(InitialTextToSet, true);
 
 		// Update the cached BoundText value to prevent it triggering another SetEditableText update again next Tick
@@ -137,14 +143,13 @@ void FSlateEditableTextLayout::Construct(TAttribute<FText> InInitialText, FTextB
 	bTextChangedByVirtualKeyboard = false;
 	bTextCommittedByVirtualKeyboard = false;
 	bSelectionChangedExternally = false;
-	bMoveCursorToDocumentEnd = false;
 	VirtualKeyboardTextCommitType = ETextCommit::Default;
 
 	CachedSize = FVector2f::ZeroVector;
 
 	auto ExecuteDeleteAction = [this]()
 	{
-		BeginEditTransaction();
+		BeginEditTransation();
 		DeleteSelectedText();
 		EndEditTransaction();
 	};
@@ -215,21 +220,44 @@ FSlateEditableTextLayout::~FSlateEditableTextLayout()
 	}
 }
 
-void FSlateEditableTextLayout::SetText(TAttribute<FText> InText)
+void FSlateEditableTextLayout::SetText(const TAttribute<FText>& InText)
 {
-	bMoveCursorToDocumentEnd = false;
+	const FText PreviousText = BoundText.Get(FText::GetEmpty());
+	BoundText = InText;
+	const FText NewText = BoundText.Get(FText::GetEmpty());
+
+	// We need to force an update if the text doesn't match the editable text, as the editable 
+	// text may not match the current bound text since it may have been changed by the user
+	const FText EditableText = GetEditableText();
+	const bool bForceRefresh = !EditableText.ToString().Equals(NewText.ToString(), ESearchCase::CaseSensitive);
 
-	// this will call OnBoundTextChanged right away
-	BoundText.Assign(OwnerWidget->GetSlateWidget().Get(), MoveTemp(InText), FText::GetEmpty(), ESlateAttributeBindAction::Initialize);
+	// Only emit the "text changed" event if the text has actually been changed
+	const bool bHasTextChanged = OwnerWidget->GetSlateWidget()->HasAnyUserFocus().IsSet()
+		? !NewText.ToString().Equals(EditableText.ToString(), ESearchCase::CaseSensitive)
+		: !NewText.ToString().Equals(PreviousText.ToString(), ESearchCase::CaseSensitive);
 
-	// Make sure we move the cursor to the end of the new text if we had keyboard focus
-	if (bMoveCursorToDocumentEnd && OwnerWidget->GetSlateWidget()->HasAnyUserFocus().IsSet() && !bWasFocusedByLastMouseDown)
+	if (RefreshImpl(&NewText, bForceRefresh))
 	{
-		JumpTo(ETextLocation::EndOfDocument, ECursorAction::MoveCursor);
+		// Make sure we move the cursor to the end of the new text if we had keyboard focus
+		if (OwnerWidget->GetSlateWidget()->HasAnyUserFocus().IsSet() && !bWasFocusedByLastMouseDown)
+		{
+			JumpTo(ETextLocation::EndOfDocument, ECursorAction::MoveCursor);
+		}
+
+		// Let outsiders know that the text content has been changed
+		if (bHasTextChanged)
+		{
+			OwnerWidget->OnTextChanged(NewText);
+		}
+	}
+
+	if (bHasTextChanged || BoundText.IsBound())
+	{
+		OwnerWidget->GetSlateWidget()->Invalidate(EInvalidateWidget::LayoutAndVolatility);
 	}
 }
 
-int32 FSlateEditableTextLayout::GetTextLineCount() const
+int32 FSlateEditableTextLayout::GetTextLineCount()
 {
 	return TextLayout->GetLineCount();
 }
@@ -238,17 +266,15 @@ FText FSlateEditableTextLayout::GetText() const
 {
 	SLATE_CROSS_THREAD_CHECK();
 
-	return BoundText.Get();
+	return BoundText.Get(FText::GetEmpty());
 }
 
-void FSlateEditableTextLayout::SetHintText(TAttribute<FText> InHintText)
+void FSlateEditableTextLayout::SetHintText(const TAttribute<FText>& InHintText)
 {
-	const TSharedRef<SWidget> OwnerSlateWidget = OwnerWidget->GetSlateWidget();
-
-	HintText.Assign(OwnerSlateWidget.Get(), MoveTemp(InHintText));
+	HintText = InHintText;
 
 	// If we have hint text that is either non-empty or bound to a delegate, we'll also need to make the hint text layout
-	if (HintText.IsBound(OwnerSlateWidget.Get()) || !HintText.Get().IsEmpty())
+	if (HintText.IsBound() || !HintText.Get(FText::GetEmpty()).IsEmpty())
 	{
 		HintTextLayout = MakeUnique<FSlateTextBlockLayout>(OwnerWidget->GetSlateWidgetPtr().Get(), HintTextStyle, TextLayout->GetTextShapingMethod(), TextLayout->GetTextFlowDirection(), CreateSlateTextLayout, HintMarshaller.ToSharedRef(), nullptr);
 		HintTextLayout->SetDebugSourceInfo(DebugSourceInfo);
@@ -263,12 +289,18 @@ void FSlateEditableTextLayout::SetHintText(TAttribute<FText> InHintText)
 
 FText FSlateEditableTextLayout::GetHintText() const
 {
-	return HintText.Get();
+	return HintText.Get(FText::GetEmpty());
 }
 
-void FSlateEditableTextLayout::SetSearchText(TAttribute<FText> InSearchText)
+void FSlateEditableTextLayout::SetSearchText(const TAttribute<FText>& InSearchText)
 {
-	BoundSearchText.Assign(OwnerWidget->GetSlateWidget().Get(), MoveTemp(InSearchText));
+	const FText& SearchTextToSet = InSearchText.Get(FText::GetEmpty());
+
+	BoundSearchText = InSearchText;
+	BoundSearchTextLastTick = FTextSnapshot(SearchTextToSet);
+
+	BeginSearch(SearchTextToSet);
+	OwnerWidget->GetSlateWidget()->Invalidate(EInvalidateWidget::LayoutAndVolatility);
 }
 
 FText FSlateEditableTextLayout::GetSearchText() const
@@ -293,8 +325,6 @@ void FSlateEditableTextLayout::SetTextStyle(const FTextBlockStyle& InTextStyle)
 
 	TextLayout->SetDefaultTextStyle(TextStyle);
 	Marshaller->MakeDirty(); // will regenerate the text using the new default style
-
-	OwnerWidget->GetSlateWidget()->Invalidate(EInvalidateWidget::Layout); //Using Layout as changing text block size can affect the size.
 }
 
 const FTextBlockStyle& FSlateEditableTextLayout::GetTextStyle() const
@@ -302,14 +332,14 @@ const FTextBlockStyle& FSlateEditableTextLayout::GetTextStyle() const
 	return TextStyle;
 }
 
-void FSlateEditableTextLayout::SetHintTextOpacity(TAttribute<float> InOpacity)
+void FSlateEditableTextLayout::SetHintTextOpacity(const TAttribute<float>& InOpacity)
 {
-	HintTextOpacity.Assign(OwnerWidget->GetSlateWidget().Get(), MoveTemp(InOpacity));
+	HintTextOpacity = InOpacity;
 }
 
 float FSlateEditableTextLayout::GetHintTextOpacity() const
 {
-	return HintTextOpacity.Get();
+	return HintTextOpacity.Get(0.35f);
 }
 
 void FSlateEditableTextLayout::SetCursorBrush(const TAttribute<const FSlateBrush*>& InCursorBrush)
@@ -437,76 +467,76 @@ FTextSelection FSlateEditableTextLayout::GetSelection() const
 
 void FSlateEditableTextLayout::SetTextShapingMethod(const TOptional<ETextShapingMethod>& InTextShapingMethod)
 {
-	const ETextShapingMethod NewShapingMethod = InTextShapingMethod.Get(GetDefaultTextShapingMethod());
-	if (NewShapingMethod != TextLayout->GetTextShapingMethod())
-	{
-		TextLayout->SetTextShapingMethod((InTextShapingMethod.IsSet()) ? InTextShapingMethod.GetValue() : GetDefaultTextShapingMethod());
-
-		OwnerWidget->GetSlateWidget()->Invalidate(EInvalidateWidget::Layout);
-	}
+	TextLayout->SetTextShapingMethod((InTextShapingMethod.IsSet()) ? InTextShapingMethod.GetValue() : GetDefaultTextShapingMethod());
 }
 
 void FSlateEditableTextLayout::SetTextFlowDirection(const TOptional<ETextFlowDirection>& InTextFlowDirection)
 {
-	const ETextFlowDirection NewFlowDirection = InTextFlowDirection.Get(GetDefaultTextFlowDirection());
-	if (NewFlowDirection != TextLayout->GetTextFlowDirection())
-	{
-		TextLayout->SetTextFlowDirection(NewFlowDirection);
-
-		OwnerWidget->GetSlateWidget()->Invalidate(EInvalidateWidget::Layout);
-	}
+	TextLayout->SetTextFlowDirection((InTextFlowDirection.IsSet()) ? InTextFlowDirection.GetValue() : GetDefaultTextFlowDirection());
 }
 
-void FSlateEditableTextLayout::SetTextWrapping(TAttribute<float> InWrapTextAt, TAttribute<bool> InAutoWrapText, TAttribute<ETextWrappingPolicy> InWrappingPolicy)
+void FSlateEditableTextLayout::SetTextWrapping(const TAttribute<float>& InWrapTextAt, const TAttribute<bool>& InAutoWrapText, const TAttribute<ETextWrappingPolicy>& InWrappingPolicy)
 {
-	const TSharedRef<SWidget> OwnerSlateWidget = OwnerWidget->GetSlateWidget();
+	WrapTextAt = InWrapTextAt;
+	AutoWrapText = InAutoWrapText;
+	WrappingPolicy = InWrappingPolicy;
 
-	WrapTextAt.Assign(OwnerSlateWidget.Get(), MoveTemp(InWrapTextAt));
-	AutoWrapText.Assign(OwnerSlateWidget.Get(), MoveTemp(InAutoWrapText));
-	WrappingPolicy.Assign(OwnerSlateWidget.Get(), MoveTemp(InWrappingPolicy));
+	OwnerWidget->GetSlateWidget()->Invalidate(EInvalidateWidget::LayoutAndVolatility);
 }
 
-void FSlateEditableTextLayout::SetWrapTextAt(TAttribute<float> InWrapTextAt)
+void FSlateEditableTextLayout::SetWrapTextAt(const TAttribute<float>& InWrapTextAt)
 {
-	WrapTextAt.Assign(OwnerWidget->GetSlateWidget().Get(), MoveTemp(InWrapTextAt));
-}
+	WrapTextAt = InWrapTextAt;
 
-void FSlateEditableTextLayout::SetAutoWrapText(TAttribute<bool> InAutoWrapText)
-{
-	AutoWrapText.Assign(OwnerWidget->GetSlateWidget().Get(), MoveTemp(InAutoWrapText));
+	OwnerWidget->GetSlateWidget()->Invalidate(EInvalidateWidget::LayoutAndVolatility);
 }
 
-void FSlateEditableTextLayout::SetWrappingPolicy(TAttribute<ETextWrappingPolicy> InWrappingPolicy)
+void FSlateEditableTextLayout::SetAutoWrapText(const TAttribute<bool>& InAutoWrapText)
 {
-	WrappingPolicy.Assign(OwnerWidget->GetSlateWidget().Get(), MoveTemp(InWrappingPolicy));
+	AutoWrapText = InAutoWrapText;
+
+	OwnerWidget->GetSlateWidget()->Invalidate(EInvalidateWidget::LayoutAndVolatility);
 }
 
-void FSlateEditableTextLayout::SetMargin(TAttribute<FMargin> InMargin)
+void FSlateEditableTextLayout::SetWrappingPolicy(const TAttribute<ETextWrappingPolicy>& InWrappingPolicy)
 {
-	Margin.Assign(OwnerWidget->GetSlateWidget().Get(), MoveTemp(InMargin));
+	WrappingPolicy = InWrappingPolicy;
+
+	OwnerWidget->GetSlateWidget()->Invalidate(EInvalidateWidget::LayoutAndVolatility);
 }
 
-void FSlateEditableTextLayout::SetJustification(TAttribute<ETextJustify::Type> InJustification)
+void FSlateEditableTextLayout::SetMargin(const TAttribute<FMargin>& InMargin)
 {
-	Justification.Assign(OwnerWidget->GetSlateWidget().Get(), MoveTemp(InJustification));
+	Margin = InMargin;
+
+	OwnerWidget->GetSlateWidget()->Invalidate(EInvalidateWidget::LayoutAndVolatility);
 }
 
-ETextJustify::Type FSlateEditableTextLayout::GetJustification() const
+void FSlateEditableTextLayout::SetJustification(const TAttribute<ETextJustify::Type>& InJustification)
 {
-	return Justification.Get();
+	if (!Justification.IdenticalTo(InJustification))
+	{
+		Justification = InJustification;
+
+		OwnerWidget->GetSlateWidget()->Invalidate(EInvalidateWidget::LayoutAndVolatility);
+	}
 }
 
-void FSlateEditableTextLayout::SetLineHeightPercentage(TAttribute<float> InLineHeightPercentage)
+void FSlateEditableTextLayout::SetLineHeightPercentage(const TAttribute<float>& InLineHeightPercentage)
 {
-	LineHeightPercentage.Assign(OwnerWidget->GetSlateWidget().Get(), MoveTemp(InLineHeightPercentage));
+	LineHeightPercentage = InLineHeightPercentage;
+
+	OwnerWidget->GetSlateWidget()->Invalidate(EInvalidateWidget::LayoutAndVolatility);
 }
 
-void FSlateEditableTextLayout::SetApplyLineHeightToBottomLine(TAttribute<bool> InApplyLineHeightToBottomLine)
+void FSlateEditableTextLayout::SetApplyLineHeightToBottomLine(const TAttribute<bool>& InApplyLineHeightToBottomLine)
 {
-	ApplyLineHeightToBottomLine.Assign(OwnerWidget->GetSlateWidget().Get(), MoveTemp(InApplyLineHeightToBottomLine));
+	ApplyLineHeightToBottomLine = InApplyLineHeightToBottomLine;
+
+	OwnerWidget->GetSlateWidget()->Invalidate(EInvalidateWidget::LayoutAndVolatility);
 }
 
-void FSlateEditableTextLayout::SetOverflowPolicy(const TOptional<ETextOverflowPolicy>& InOverflowPolicy)
+void FSlateEditableTextLayout::SetOverflowPolicy(TOptional<ETextOverflowPolicy> InOverflowPolicy)
 {
 	if(OverflowPolicyOverride != InOverflowPolicy)
 	{
@@ -573,7 +603,7 @@ void FSlateEditableTextLayout::EnableTextInputMethodContext()
 
 bool FSlateEditableTextLayout::Refresh()
 {
-	const FText& TextToSet = BoundText.Get();
+	const FText& TextToSet = BoundText.Get(FText::GetEmpty());
 	return RefreshImpl(&TextToSet);
 }
 
@@ -592,7 +622,7 @@ bool FSlateEditableTextLayout::RefreshImpl(const FText* InTextToSet, const bool
 		if (bForce || !BoundTextLastTick.IsDisplayStringEqualTo(*InTextToSet))
 		{
 			// The source text has changed, so update the internal editable text
-			bHasSetText = SetEditableText(*InTextToSet, bForce);
+			bHasSetText = SetEditableText(*InTextToSet, true);
 		}
 
 		// Update this even if the text is lexically identical, as it will update the pointer compared by IdenticalTo for the next Tick
@@ -862,7 +892,7 @@ bool FSlateEditableTextLayout::HandleFocusReceived(const FFocusEvent& InFocusEve
 	// Force update of last cursor interaction time to ensure the caret blinks at the correct frequency regardless of how focus is set.
 	CursorInfo.UpdateLastCursorInteractionTime();
 
-	OwnerWidget->OnBeginTextEdit(BoundText.Get());
+	OwnerWidget->OnBeginTextEdit(BoundText.Get(FText::GetEmpty()));
 	return true;
 }
 
@@ -3149,7 +3179,7 @@ bool FSlateEditableTextLayout::HasTextChangedFromOriginal() const
 	return bHasChanged;
 }
 
-void FSlateEditableTextLayout::BeginEditTransaction()
+void FSlateEditableTextLayout::BeginEditTransation()
 {
 	NumTransactionsOpened += 1;
 	
@@ -3381,36 +3411,37 @@ void FSlateEditableTextLayout::Redo()
 
 void FSlateEditableTextLayout::SaveText(const FText& TextToSave)
 {
-	const TSharedRef<SWidget> OwnerSlateWidget = OwnerWidget->GetSlateWidget();
 	// Don't set text if the text attribute has a 'getter' binding on it, otherwise we'd blow away
 	// that binding.  If there is a getter binding, then we'll assume it will provide us with
 	// updated text after we've fired our 'text changed' callbacks
-	if (!BoundText.IsBound(OwnerSlateWidget.Get()))
+	if (!BoundText.IsBound())
 	{
-		BoundText.Assign(OwnerSlateWidget.Get(), TextToSave);
+		BoundText.Set(TextToSave);
 	}
 }
 
 void FSlateEditableTextLayout::LoadText()
 {
-	const TSharedRef<SWidget> OwnerSlateWidget = OwnerWidget->GetSlateWidget();
 	// We only need to do this if we're bound to a delegate, otherwise the text layout will already be up-to-date
 	// either from Construct, or a call to SetText
-	if (BoundText.IsBound(OwnerSlateWidget.Get()))
-	{
-		BoundText.UpdateNow(OwnerSlateWidget.Get());
-		// If we arrive here when the text is commited and `BoundText` is bound,
-		// the new text might have been refused in which case `BoundText.UpdateNow` would have not triggerred `OnBoundTextChanged`
-		// If this was not called, the Edited text would now not match the `BoundText` which is a problem, so we call it manually.
-		// Calling it twice would not be a problem (even though not ideal) as nothing would trigger now that the content is identical.
-		OnBoundTextChanged(OwnerSlateWidget.Get(), *this);
+	if (BoundText.IsBound())
+	{
+		SetText(BoundText);
 		TextLayout->UpdateIfNeeded();
 	}
 }
 
 bool FSlateEditableTextLayout::ComputeVolatility() const
 {
-	return false;
+	return BoundText.IsBound()
+		|| HintText.IsBound()
+		|| BoundSearchText.IsBound()
+		|| WrapTextAt.IsBound()
+		|| AutoWrapText.IsBound()
+		|| WrappingPolicy.IsBound()
+		|| Margin.IsBound()
+		|| Justification.IsBound()
+		|| LineHeightPercentage.IsBound();
 }
 
 void FSlateEditableTextLayout::UpdateTextChangedByVirtualKeyboard()
@@ -3478,7 +3509,7 @@ void FSlateEditableTextLayout::Tick(const FGeometry& AllottedGeometry, const dou
 
 	// Update the search before we process the next PositionToScrollIntoView
 	{
-		const FText& SearchTextToSet = BoundSearchText.Get();
+		const FText& SearchTextToSet = BoundSearchText.Get(FText::GetEmpty());
 		if (!BoundSearchTextLastTick.IdenticalTo(SearchTextToSet))
 		{
 			// The pointer used by the bound text has changed, however the text may still be the same - check that now
@@ -3567,9 +3598,9 @@ void FSlateEditableTextLayout::Tick(const FGeometry& AllottedGeometry, const dou
 	TextLayout->SetVisibleRegion(AllottedGeometry.Size, FVector2D(ScrollOffset) * TextLayout->GetScale());
 }
 
-int32 FSlateEditableTextLayout::OnPaint(const FPaintArgs& Args, const FGeometry& AllottedGeometry, const FSlateRect& MyCullingRect, FSlateWindowElementList& OutDrawElements, int32 LayerId, const FWidgetStyle& InWidgetStyle, bool bParentEnabled) const
+int32 FSlateEditableTextLayout::OnPaint(const FPaintArgs& Args, const FGeometry& AllottedGeometry, const FSlateRect& MyCullingRect, FSlateWindowElementList& OutDrawElements, int32 LayerId, const FWidgetStyle& InWidgetStyle, bool bParentEnabled)
 {
-	// Update the auto-wrap size now that we have computed paint geometry; won't take effect until text frame
+	// Update the auto-wrap size now that we have computed paint geometry; won't take affect until text frame
 	// Note: This is done here rather than in Tick(), because Tick() doesn't get called while resizing windows, but OnPaint() does
 	CachedSize = FVector2f(AllottedGeometry.GetLocalSize());
 
@@ -3783,13 +3814,11 @@ void FSlateEditableTextLayout::FVirtualKeyboardEntry::SetTextFromVirtualKeyboard
 {
 	check(IsInGameThread());
 
-	const TSharedRef<SWidget> OwnerSlateWidget = OwnerLayout->OwnerWidget->GetSlateWidget();
-
 	// Only set the text if the text attribute doesn't have a getter binding (otherwise it would be blown away).
 	// If it is bound, we'll assume that OnTextCommitted will handle the update.
-	if (!OwnerLayout->BoundText.IsBound(OwnerSlateWidget.Get()))
+	if (!OwnerLayout->BoundText.IsBound())
 	{
-		OwnerLayout->BoundText.Assign(OwnerSlateWidget.Get(), InNewText);
+		OwnerLayout->BoundText.Set(InNewText);
 	}
 
 	// Update the internal editable text
@@ -4155,7 +4184,7 @@ void FSlateEditableTextLayout::FTextInputMethodContext::BeginComposition()
 	{
 		bIsComposing = true;
 
-		OwnerLayout->BeginEditTransaction();
+		OwnerLayout->BeginEditTransation();
 		OwnerLayout->UpdateCursorHighlight();
 	}
 }
@@ -4214,39 +4243,4 @@ void FSlateEditableTextLayout::SetFontFacesLoadingPaintPolicy(EFontFacesLoadingP
 	{
 		TextLayout->SetFontFacesLoadingPaintPolicy(InFontFacesLoadingPaintPolicy);
 	}
-}
-
-void FSlateEditableTextLayout::OnBoundTextChanged(SWidget& Widget, FSlateEditableTextLayout& SlateTextLayout)
-{
-	const FText NewText = SlateTextLayout.BoundText.Get();
-
-	// We need to force an update if the text doesn't match the editable text, as the editable
-	// text may not match the current bound text since it may have been changed by the user
-	const FText EditableText = SlateTextLayout.GetEditableText();
-	const bool bForceRefresh = !EditableText.ToString().Equals(NewText.ToString(), ESearchCase::CaseSensitive);
-
-	// Only emit the "text changed" event if the text has actually been changed
-	const bool bHasTextChanged = Widget.HasAnyUserFocus().IsSet()
-		? !NewText.ToString().Equals(EditableText.ToString(), ESearchCase::CaseSensitive)
-		: !NewText.ToString().Equals(SlateTextLayout.PreviousBoundText.ToString(), ESearchCase::CaseSensitive);
-
-	if (SlateTextLayout.RefreshImpl(&NewText, bForceRefresh))
-	{
-		SlateTextLayout.bMoveCursorToDocumentEnd = true;
-
-		// Let outsiders know that the text content has been changed
-		if (bHasTextChanged)
-		{
-			SlateTextLayout.OwnerWidget->OnTextChanged(NewText);
-		}
-	}
-
-	SlateTextLayout.PreviousBoundText = SlateTextLayout.BoundText.Get();
-}
-
-void FSlateEditableTextLayout::OnBoundSearchTextChanged(SWidget& Widget, FSlateEditableTextLayout& SlateTextLayout)
-{
-	const FText& SearchTextToSet = SlateTextLayout.BoundSearchText.Get();
-	SlateTextLayout.BoundSearchTextLastTick = FTextSnapshot(SearchTextToSet);
-	SlateTextLayout.BeginSearch(SearchTextToSet);
-}
+}
\ No newline at end of file
diff --git a/Engine/Source/Runtime/Slate/Public/Widgets/Input/SEditableText.h b/Engine/Source/Runtime/Slate/Public/Widgets/Input/SEditableText.h
--- a/Engine/Source/Runtime/Slate/Public/Widgets/Input/SEditableText.h
+++ b/Engine/Source/Runtime/Slate/Public/Widgets/Input/SEditableText.h
@@ -19,13 +19,13 @@
 #include "Framework/Text/TextLayout.h"
 #include "Framework/Text/SlateAsyncFontLoadingTypes.h"
 #include "Widgets/Text/ISlateEditableTextWidget.h"
-#include "Widgets/Text/SlateEditableTextLayout.h"
 
 class FActiveTimerHandle;
 class FArrangedChildren;
 class FChildren;
 class FPaintArgs;
 class FPlainTextLayoutMarshaller;
+class FSlateEditableTextLayout;
 class FSlateWindowElementList;
 class IBreakIterator;
 struct FTextLocation;
@@ -34,18 +34,11 @@ enum class ETextShapingMethod : uint8;
 
 DECLARE_DELEGATE_TwoParams(FOnCursorMovedWithSelection, const FTextLocation&, const FTextSelection&);
 
-template <>
-struct TWidgetTypeTraits<class SEditableText>
-{
-	static constexpr bool SupportsInvalidation() { return true; }
-};
-
 /**
  * Editable text widget
  */
 class SEditableText : public SWidget, public ISlateEditableTextWidget
 {
-	SLATE_DECLARE_WIDGET_API(SEditableText, SWidget, SLATE_API)
 public:
 	SLATE_BEGIN_ARGS( SEditableText )
 		: _Text()
@@ -210,6 +203,9 @@ public:
 	/** Constructor */
 	SLATE_API SEditableText();
 
+	/** Destructor */
+	SLATE_API ~SEditableText();
+
 	/**
 	 * Construct this widget
 	 *
@@ -222,7 +218,7 @@ public:
 	 *
 	 * @param  InNewText  The new text
 	 */
-	SLATE_API void SetText(TAttribute<FText> InNewText);
+	SLATE_API void SetText( const TAttribute< FText >& InNewText );
 	
 	/**
 	 * Returns the text string
@@ -242,50 +238,47 @@ public:
 	SLATE_API bool SetEditableText(const FText& InNewText);
 
 	/** See the HintText attribute */
-	SLATE_API void SetHintText(TAttribute<FText> InHintText);
+	SLATE_API void SetHintText( const TAttribute< FText >& InHintText );
 	
 	/** Get the text that appears when there is no text in the text box */
 	SLATE_API FText GetHintText() const;
 	
 	/** Sets the hint text opacity */
-	SLATE_API void SetHintTextOpacity(TAttribute<float> InHintTextOpacity);
+	SLATE_API void SetHintTextOpacity(const TAttribute<float>& InHintTextOpacity);
 	
 	/** Get the hint text opacity */
 	SLATE_API float GetHintTextOpacity() const;
 
 	/** Set the text that is currently being searched for (if any) */
-	SLATE_API void SetSearchText(TAttribute<FText> InSearchText);
+	SLATE_API void SetSearchText(const TAttribute<FText>& InSearchText);
 
 	/** Get the text that is currently being searched for (if any) */
 	SLATE_API FText GetSearchText() const;
 
 	/** See the IsReadOnly attribute */
-	SLATE_API void SetIsReadOnly(TAttribute<bool> InIsReadOnly);
+	SLATE_API void SetIsReadOnly( TAttribute< bool > InIsReadOnly );
 	
 	/** See the IsPassword attribute */
-	SLATE_API void SetIsPassword(TAttribute<bool> InIsPassword);
+	SLATE_API void SetIsPassword( TAttribute< bool > InIsPassword );
 	
 	/** See the ColorAndOpacity attribute */
 	SLATE_API void SetColorAndOpacity(TAttribute<FSlateColor> Color);
 
-	/** See the BackgroundImageSelected attribute */
-	SLATE_API void SetBackgroundImageSelected(TAttribute<const FSlateBrush*> InBackgroundImageSelected);
-
 	/** See the AllowContextMenu attribute */
-	SLATE_API void SetAllowContextMenu(TAttribute<bool> InAllowContextMenu);
+	SLATE_API void SetAllowContextMenu(const TAttribute< bool >& InAllowContextMenu);
 
 	/** See the EnableIntegratedKeyboard attribute */
-	SLATE_API void SetEnableIntegratedKeyboard(TAttribute<bool> InEnableIntegratedKeyboard);
+	SLATE_API void SetEnableIntegratedKeyboard(const TAttribute< bool >& InEnableIntegratedKeyboard);
 
 	/** Set the VirtualKeyboardDismissAction attribute */
-	SLATE_API void SetVirtualKeyboardDismissAction(TAttribute<EVirtualKeyboardDismissAction> InVirtualKeyboardDismissAction);
+	SLATE_API void SetVirtualKeyboardDismissAction(TAttribute< EVirtualKeyboardDismissAction > InVirtualKeyboardDismissAction);
 
 	/**
 	 * Sets the font used to draw the text
 	 *
 	 * @param  InNewFont	The new font to use
 	 */
-	SLATE_API void SetFont(TAttribute<FSlateFontInfo> InNewFont);
+	SLATE_API void SetFont(const TAttribute< FSlateFontInfo >& InNewFont);
 
 	/** Gets the font used to draw the text. */
 	SLATE_API FSlateFontInfo GetFont() const;
@@ -293,7 +286,7 @@ public:
 	/**
 	 * Sets the text style used to draw the text
 	 *
-	 * @param  InNewTextStyle	The new text style to use
+	 * @param  NewTextStyle	The new text style to use
 	 */
 	SLATE_API void SetTextStyle( const FEditableTextStyle& InNewTextStyle );
 
@@ -310,53 +303,52 @@ public:
 	 *
 	 * @param  InMinDesiredWidth	The minimum width
 	 */
-	SLATE_API void SetMinDesiredWidth(TAttribute<float> InMinDesiredWidth);
+	SLATE_API void SetMinDesiredWidth(const TAttribute<float>& InMinDesiredWidth);
 
 	/**
 	 * Workaround as we loose focus when the auto completion closes.
 	 *
 	 * @param  InIsCaretMovedWhenGainFocus	Workaround
 	 */
-	SLATE_API void SetIsCaretMovedWhenGainFocus(TAttribute<bool> InIsCaretMovedWhenGainFocus);
+	SLATE_API void SetIsCaretMovedWhenGainFocus(const TAttribute<bool>& InIsCaretMovedWhenGainFocus);
 
 	/**
 	 * Sets whether to select all text when the user clicks to give focus on the widget
 	 *
 	 * @param  InSelectAllTextWhenFocused	Select all text when the user clicks?
 	 */
-	SLATE_API void SetSelectAllTextWhenFocused(TAttribute<bool> InSelectAllTextWhenFocused);
+	SLATE_API void SetSelectAllTextWhenFocused(const TAttribute<bool>& InSelectAllTextWhenFocused);
 
 	/**
 	 * Sets whether to allow the user to back out of changes when they press the escape key
 	 *
 	 * @param  InRevertTextOnEscape			Allow the user to back out of changes?
 	 */
-	SLATE_API void SetRevertTextOnEscape(TAttribute<bool> InRevertTextOnEscape);
+	SLATE_API void SetRevertTextOnEscape(const TAttribute<bool>& InRevertTextOnEscape);
 
 	/**
 	 * Sets whether to clear keyboard focus when pressing enter to commit changes
 	 *
 	 * @param  InClearKeyboardFocusOnCommit		Clear keyboard focus when pressing enter?
 	 */
-	SLATE_API void SetClearKeyboardFocusOnCommit(TAttribute<bool> InClearKeyboardFocusOnCommit);
+	SLATE_API void SetClearKeyboardFocusOnCommit(const TAttribute<bool>& InClearKeyboardFocusOnCommit);
 
 	/**
 	 * Sets whether to select all text when pressing enter to commit changes
 	 *
 	 * @param  InSelectAllTextOnCommit		Select all text when pressing enter?
 	 */
-	SLATE_API void SetSelectAllTextOnCommit(TAttribute<bool> InSelectAllTextOnCommit);
+	SLATE_API void SetSelectAllTextOnCommit(const TAttribute<bool>& InSelectAllTextOnCommit);
 
 	/**
 	 * Sets whether to select word on the mouse double click
 	 *
 	 * @param  InSelectWordOnMouseDoubleClick		Select word on the mouse double click
 	 */
-	SLATE_API void SetSelectWordOnMouseDoubleClick(TAttribute<bool> InSelectWordOnMouseDoubleClick);
+	SLATE_API void SetSelectWordOnMouseDoubleClick(const TAttribute<bool>& InSelectWordOnMouseDoubleClick);
 
 	/** See Justification attribute */
-	SLATE_API void SetJustification(TAttribute<ETextJustify::Type> InJustification);
-
+	SLATE_API void SetJustification(const TAttribute<ETextJustify::Type>& InJustification);
 	/**
 	 * Sets the OnKeyCharHandler to provide first chance handling of the OnKeyChar event
 	 *
@@ -364,7 +356,7 @@ public:
 	 */
 	void SetOnKeyCharHandler(FOnKeyChar InOnKeyCharHandler)
 	{
-		OnKeyCharHandler = MoveTemp(InOnKeyCharHandler);
+		OnKeyCharHandler = InOnKeyCharHandler;
 	}
 
 	/**
@@ -374,7 +366,7 @@ public:
 	 */
 	void SetOnKeyDownHandler(FOnKeyDown InOnKeyDownHandler)
 	{
-		OnKeyDownHandler = MoveTemp(InOnKeyDownHandler);
+		OnKeyDownHandler = InOnKeyDownHandler;
 	}
 
 	/** See TextShapingMethod attribute */
@@ -384,7 +376,7 @@ public:
 	SLATE_API void SetTextFlowDirection(const TOptional<ETextFlowDirection>& InTextFlowDirection);
 
 	/** Sets the overflow policy for this text block */
-	SLATE_API void SetOverflowPolicy(const TOptional<ETextOverflowPolicy>& InOverflowPolicy);
+	SLATE_API void SetOverflowPolicy(TOptional<ETextOverflowPolicy> InOverflowPolicy);
 
 	/** Sets the paint policy to use for text if the font faces are still async loading. */
 	SLATE_API void SetFontFacesLoadingPaintPolicy(EFontFacesLoadingPaintPolicy InFontFacesLoadingPaintPolicy);
@@ -451,7 +443,7 @@ protected:
 	SLATE_API virtual FReply OnMouseMove( const FGeometry& InMyGeometry, const FPointerEvent& InMouseEvent ) override;
 	SLATE_API virtual FReply OnMouseButtonDoubleClick( const FGeometry& InMyGeometry, const FPointerEvent& InMouseEvent ) override;
 	SLATE_API virtual FCursorReply OnCursorQuery( const FGeometry& MyGeometry, const FPointerEvent& CursorEvent ) const override;
-	SLATE_API virtual const FSlateBrush* GetFocusBrush() const override;
+	SLATE_API virtual const FSlateBrush* GetFocusBrush() const;
 	SLATE_API virtual bool IsInteractable() const override;
 	SLATE_API virtual bool ComputeVolatility() const override;
 #if WITH_ACCESSIBILITY
@@ -495,75 +487,31 @@ protected:
 	SLATE_API virtual void OnTextChanged(const FText& InText) override;
 	SLATE_API virtual void OnTextCommitted(const FText& InText, const ETextCommit::Type InTextAction) override;
 	SLATE_API virtual void OnCursorMoved(const FTextLocation& InLocation) override;
-	SLATE_API virtual float UpdateAndClampHorizontalScrollBar(const float InViewOffset, const float InViewFraction, const EVisibility InVisibilityOverride) override;
-	SLATE_API virtual float UpdateAndClampVerticalScrollBar(const float InViewOffset, const float InViewFraction, const EVisibility InVisibilityOverride) override;
+	SLATE_API virtual float UpdateAndClampHorizontalScrollBar(const float InViewOffset, const float InViewFraction, const EVisibility InVisiblityOverride) override;
+	SLATE_API virtual float UpdateAndClampVerticalScrollBar(const float InViewOffset, const float InViewFraction, const EVisibility InVisiblityOverride) override;
 	//~ End ISlateEditableTextWidget Interface
 
-protected:
-	SLATE_API ETextJustify::Type GetJustification() const;
-
-	TSlateAttributeRef<FSlateFontInfo> GetFontOverrideAttribute() const
-	{
-		return TSlateAttributeRef<FSlateFontInfo>{SharedThis(this), FontOverrideAttribute};
-	}
-
-	SLATE_API FSlateColor GetColorAndOpacity() const;
-	TSlateAttributeRef<FSlateColor> GetColorAndOpacityAttribute() const
-	{
-		return TSlateAttributeRef<FSlateColor>{SharedThis(this), ColorAndOpacityAttribute};
-	}
-
-	SLATE_API const FSlateBrush* GetBackgroundImageSelected() const;
-	TSlateAttributeRef<const FSlateBrush*> GetBackgroundImageSelectedAttribute() const
-	{
-		return TSlateAttributeRef<const FSlateBrush*>{SharedThis(this), BackgroundImageSelectedAttribute};
-	}
-
-	TSlateAttributeRef<bool> GetIsPasswordAttribute() const
-	{
-		return TSlateAttributeRef<bool>{SharedThis(this), bIsPasswordAttribute};
-	}
-
-	SLATE_API float GetMinDesiredWidth() const;
-	TSlateAttributeRef<float> GetMinDesiredWidthAttribute() const
-	{
-		return TSlateAttributeRef<float>{SharedThis(this), MinDesiredWidthAttribute};
-	}
-
 protected:
 	/** Text marshaller used by the editable text layout */
 	TSharedPtr<FPlainTextLayoutMarshaller> PlainTextMarshaller;
 
-#if WITH_EDITORONLY_DATA
 	/** The text layout that deals with the editable text */
-	UE_DEPRECATED(5.8, "Use SlateEditableTextLayout, this variable will never be set.")
 	TUniquePtr<FSlateEditableTextLayout> EditableTextLayout;
-#endif
-	/** The text layout that deals with the editable text */
-	FSlateEditableTextLayout SlateEditableTextLayout;
 
-#if WITH_EDITORONLY_DATA
 	/** The font used to draw the text */
-	UE_DEPRECATED(5.8, "Use SetFont / GetFont / GetFontOverrideAttribute")
-	TSlateDeprecatedTAttribute<FSlateFontInfo> Font;
+	TAttribute<FSlateFontInfo> Font;
 
 	/** Text color and opacity */
-	UE_DEPRECATED(5.8, "Use SetColorAndOpacity / GetColorAndOpacity / GetColorAndOpacityAttribute")
-	TSlateDeprecatedTAttribute<FSlateColor> ColorAndOpacity;
+	TAttribute<FSlateColor> ColorAndOpacity;
 
 	/** Background image for the selected text */
-	UE_DEPRECATED(5.8, "Use SetBackgroundImageSelected / GetBackgroundImageSelected / GetBackgroundImageSelectedAttribute")
-	TSlateDeprecatedTAttribute<const FSlateBrush*> BackgroundImageSelected;
-#endif
+	TAttribute<const FSlateBrush*> BackgroundImageSelected;
 
 	/** Sets whether this text box can actually be modified interactively by the user */
 	TAttribute<bool> bIsReadOnly;
 
-#if WITH_EDITORONLY_DATA
 	/** Sets whether this text box is for storing a password */
-	UE_DEPRECATED(5.8, "Use SetIsPassword / IsTextPassword / GetIsPasswordAttribute")
-	TSlateDeprecatedTAttribute<bool> bIsPassword;
-#endif
+	TAttribute<bool> bIsPassword;
 
 	/** Workaround as we loose focus when the auto completion closes. */
 	TAttribute<bool> bIsCaretMovedWhenGainFocus;
@@ -607,11 +555,8 @@ protected:
 	/** Called whenever the user's selection and/or cursor position of the text has changed. */
 	FOnCursorMovedWithSelection OnCursorMovedWithSelectionCallback;
 
-#if WITH_EDITORONLY_DATA
 	/** Prevents the editable text from being smaller than desired in certain cases (e.g. when it is empty) */
-	UE_DEPRECATED(5.8, "Use SetMinDesiredWidth / GetMinDesiredWidth / GetMinDesiredWidthAttribute")
-	TSlateDeprecatedTAttribute<float> MinDesiredWidth;
-#endif
+	TAttribute<float> MinDesiredWidth;
 
 	/** Menu extender for right-click context menu */
 	TSharedPtr<FExtender> MenuExtender;
@@ -631,31 +576,12 @@ protected:
 	/** The type of virtual keyboard to use for editing this text on mobile */
 	TAttribute<EKeyboardType> VirtualKeyboardType;
 
+	/** Additional options used by the virtual keyboard summoned by this widget */
+	FVirtualKeyboardOptions VirtualKeyboardOptions;
+
 	/** The type of event that will trigger the display of the virtual keyboard */
 	TAttribute<EVirtualKeyboardTrigger> VirtualKeyboardTrigger;
 
 	/** The message action to take when the virtual keyboard is dismissed by the user */
 	TAttribute<EVirtualKeyboardDismissAction> VirtualKeyboardDismissAction;
-
-	/** Additional options used by the virtual keyboard summoned by this widget */
-	FVirtualKeyboardOptions VirtualKeyboardOptions;
-
-private:
-	/** The font used to draw the text */
-	TSlateAttribute<FSlateFontInfo> FontOverrideAttribute;
-	bool bIsFontOverrideAttrSet = false;
-
-	/** Text color and opacity */
-	TSlateAttribute<FSlateColor> ColorAndOpacityAttribute;
-	bool bIsColorAndOpacityAttrSet = false;
-
-	/** Background image for the selected text */
-	TSlateAttribute<const FSlateBrush*> BackgroundImageSelectedAttribute;
-	bool bIsBackgroundImageSelectedAttrSet = false;
-
-	/** Sets whether this text box is for storing a password */
-	TSlateAttribute<bool> bIsPasswordAttribute;
-
-	/** Prevents the editable text from being smaller than desired in certain cases (e.g. when it is empty) */
-	TSlateAttribute<float> MinDesiredWidthAttribute;
 };
diff --git a/Engine/Source/Runtime/Slate/Public/Widgets/Text/SMultiLineEditableText.h b/Engine/Source/Runtime/Slate/Public/Widgets/Text/SMultiLineEditableText.h
--- a/Engine/Source/Runtime/Slate/Public/Widgets/Text/SMultiLineEditableText.h
+++ b/Engine/Source/Runtime/Slate/Public/Widgets/Text/SMultiLineEditableText.h
@@ -3,7 +3,6 @@
 #pragma once
 
 #include "CoreMinimal.h"
-#include "SlateEditableTextLayout.h"
 #include "Misc/Attribute.h"
 #include "SlateGlobals.h"
 #include "Layout/Margin.h"
@@ -29,6 +28,7 @@
 class FActiveTimerHandle;
 class FArrangedChildren;
 class FPaintArgs;
+class FSlateEditableTextLayout;
 class FSlateWindowElementList;
 class ITextLayoutMarshaller;
 enum class ETextShapingMethod : uint8;
@@ -41,7 +41,6 @@ class FSlateEditableTextLayout;
 /** An editable text widget that supports multiple lines and soft word-wrapping. */
 class SMultiLineEditableText : public SWidget, public ISlateEditableTextWidget
 {
-	SLATE_DECLARE_WIDGET_API(SMultiLineEditableText, SWidget, SLATE_API)
 public:
 
 	/** Used to merge multiple text edit transactions within a scope */
@@ -245,6 +244,7 @@ public:
 	SLATE_END_ARGS()
 
 	SLATE_API SMultiLineEditableText();
+	SLATE_API ~SMultiLineEditableText();
 
 	SLATE_API void Construct( const FArguments& InArgs );
 
@@ -355,13 +355,7 @@ public:
 	SLATE_API void SetIsReadOnly(const TAttribute< bool >& InIsReadOnly);
 
 	/** Get the number of Text Lines */
-	UE_DEPRECATED(5.8, "Use the const version instead.")
-	int32 GetTextLineCount()
-	{
-		return const_cast<const SMultiLineEditableText*>(this)->GetTextLineCount();
-	}
-	/** Get the number of Text Lines */
-	SLATE_API int32 GetTextLineCount() const;
+	SLATE_API int32 GetTextLineCount();
 
 	/**
 	 * Sets whether to select all text when the user clicks to give focus on the widget
@@ -563,10 +557,7 @@ protected:
 
 protected:
 	/** The text layout that deals with the editable text */
-	UE_DEPRECATED(5.8, "Use SlateEditableTextLayout, this variablle will never be set.")
 	TUniquePtr<FSlateEditableTextLayout> EditableTextLayout;
-	/** The text layout that deals with the editable text */
-	FSlateEditableTextLayout SlateEditableTextLayout;
 
 	/** Whether to allow multi-line text */
 	TAttribute<bool> bAllowMultiLine;
diff --git a/Engine/Source/Runtime/Slate/Public/Widgets/Text/SlateEditableTextLayout.h b/Engine/Source/Runtime/Slate/Public/Widgets/Text/SlateEditableTextLayout.h
--- a/Engine/Source/Runtime/Slate/Public/Widgets/Text/SlateEditableTextLayout.h
+++ b/Engine/Source/Runtime/Slate/Public/Widgets/Text/SlateEditableTextLayout.h
@@ -19,7 +19,6 @@
 #include "Widgets/Text/SlateEditableTextTypes.h"
 #include "Framework/Text/SlateTextLayoutFactory.h"
 #include "GenericPlatform/ITextInputMethodSystem.h"
-#include "Types/SlateNestedAttribute.h"
 
 class FArrangedChildren;
 class FExtender;
@@ -34,85 +33,19 @@ enum class ETextShapingMethod : uint8;
 enum class EFontFacesLoadingPaintPolicy : uint8;
 
 /** Class to handle the cached layout of SEditableText/SMultiLineEditableText by proxying around a FTextLayout */
-class FSlateEditableTextLayout : public FSlateNestedAttributeContainer
+class FSlateEditableTextLayout
 {
 public:
-	template<auto InPointerToAttributeOwner>
-	static void RegisterNestedAttributes(FSlateAttributeInitializer& AttributeInitializer)
-	{
-		SLATE_ADD_NESTED_MEMBER_ATTRIBUTE_DEFINITION(AttributeInitializer, InPointerToAttributeOwner, BoundText, EInvalidateWidgetReason::Layout)
-			.OnValueChanged_Static(OnBoundTextChanged);
-		SLATE_ADD_NESTED_MEMBER_ATTRIBUTE_DEFINITION(AttributeInitializer, InPointerToAttributeOwner, HintText, EInvalidateWidgetReason::Layout);
-		SLATE_ADD_NESTED_MEMBER_ATTRIBUTE_DEFINITION(AttributeInitializer, InPointerToAttributeOwner, HintTextOpacity, EInvalidateWidgetReason::Paint);
-		SLATE_ADD_NESTED_MEMBER_ATTRIBUTE_DEFINITION(AttributeInitializer, InPointerToAttributeOwner, BoundSearchText, EInvalidateWidgetReason::Layout)
-			.OnValueChanged_Static(OnBoundSearchTextChanged);
-
-		SLATE_ADD_NESTED_MEMBER_ATTRIBUTE_DEFINITION(AttributeInitializer, InPointerToAttributeOwner, WrapTextAt, EInvalidateWidgetReason::Layout);
-		SLATE_ADD_NESTED_MEMBER_ATTRIBUTE_DEFINITION(AttributeInitializer, InPointerToAttributeOwner, AutoWrapText, EInvalidateWidgetReason::Layout);
-		SLATE_ADD_NESTED_MEMBER_ATTRIBUTE_DEFINITION(AttributeInitializer, InPointerToAttributeOwner, WrappingPolicy, EInvalidateWidgetReason::Layout);
-
-		SLATE_ADD_NESTED_MEMBER_ATTRIBUTE_DEFINITION(AttributeInitializer, InPointerToAttributeOwner, Margin, EInvalidateWidgetReason::Layout);
-
-		SLATE_ADD_NESTED_MEMBER_ATTRIBUTE_DEFINITION(AttributeInitializer, InPointerToAttributeOwner, Justification, EInvalidateWidgetReason::Layout);
-		SLATE_ADD_NESTED_MEMBER_ATTRIBUTE_DEFINITION(AttributeInitializer, InPointerToAttributeOwner, LineHeightPercentage, EInvalidateWidgetReason::Layout);
-		SLATE_ADD_NESTED_MEMBER_ATTRIBUTE_DEFINITION(AttributeInitializer, InPointerToAttributeOwner, ApplyLineHeightToBottomLine, EInvalidateWidgetReason::Layout);
-	}
-
-	template<typename RealWidgetType>
-	UE_DEPRECATED(5.8, "Use the other base constructor and call Construct during widget construction")
-	FSlateEditableTextLayout(RealWidgetType& InOwnerWidget, TAttribute<FText> InInitialText, FTextBlockStyle InTextStyle, const TOptional<ETextShapingMethod> InTextShapingMethod, const TOptional<ETextFlowDirection> InTextFlowDirection, const FCreateSlateTextLayout& InCreateSlateTextLayout, TSharedRef<ITextLayoutMarshaller> InTextMarshaller, TSharedRef<ITextLayoutMarshaller> InHintTextMarshaller)
-		: FSlateEditableTextLayout(InOwnerWidget)
-	{
-		Construct(MoveTemp(InInitialText), MoveTemp(InTextStyle), InTextShapingMethod, InTextFlowDirection, InCreateSlateTextLayout, MoveTemp(InTextMarshaller), MoveTemp(InHintTextMarshaller));
-	}
-
-
-	template<typename RealWidgetType>
-	FSlateEditableTextLayout(RealWidgetType& InOwnerWidget)
-		: OwnerWidget(&InOwnerWidget)
-		, BoundText(InOwnerWidget, FText::GetEmpty())
-		, bWasPasswordLastTick(false)
-		, HintText(InOwnerWidget, FText::GetEmpty())
-		, HintTextOpacity(InOwnerWidget, 0.35f)
-		, BoundSearchText(InOwnerWidget, FText::GetEmpty())
-		, SearchCase(ESearchCase::IgnoreCase)
-		, CurrentSearchResultIndex(0)
-		, WrapTextAt(InOwnerWidget, 0.0f)
-		, AutoWrapText(InOwnerWidget, false)
-		, WrappingPolicy(InOwnerWidget, ETextWrappingPolicy::DefaultWrapping)
-		, Margin(InOwnerWidget, FMargin(0))
-		, Justification(InOwnerWidget, ETextJustify::Left)
-		, LineHeightPercentage(InOwnerWidget, 1.0f)
-		, ApplyLineHeightToBottomLine(InOwnerWidget, false)
-		, bHasRegisteredTextInputMethodContext(false)
-		, PreferredCursorScreenOffsetInLine(0)
-		, CurrentUndoLevel(0)
-		, NumTransactionsOpened(0)
-		, bIsDragSelecting(false)
-		, bWasFocusedByLastMouseDown(false)
-		, bHasDragSelectedSinceFocused(false)
-		, bTextChangedByVirtualKeyboard(false)
-		, bTextCommittedByVirtualKeyboard(false)
-		, VirtualKeyboardTextCommitType()
-		, bSelectionChangedExternally(false)
-		, ExternalSelectionStart(0)
-		, ExternalSelectionEnd(0)
-	{
-		static_assert(std::is_base_of_v<SWidget, RealWidgetType>, "RealWidgetType must be a Widget type.");
-		static_assert(std::is_base_of_v<ISlateEditableTextWidget, RealWidgetType>, "RealWidgetType must derive from ISlateEditableTextWidget.");
-	}
-
+	SLATE_API FSlateEditableTextLayout(ISlateEditableTextWidget& InOwnerWidget, const TAttribute<FText>& InInitialText, FTextBlockStyle InTextStyle, const TOptional<ETextShapingMethod> InTextShapingMethod, const TOptional<ETextFlowDirection> InTextFlowDirection, const FCreateSlateTextLayout& InCreateSlateTextLayout, TSharedRef<ITextLayoutMarshaller> InTextMarshaller, TSharedRef<ITextLayoutMarshaller> InHintTextMarshaller);
 	SLATE_API ~FSlateEditableTextLayout();
 
-	SLATE_API void Construct(TAttribute<FText> InInitialText, FTextBlockStyle InTextStyle, const TOptional<ETextShapingMethod>& InTextShapingMethod, const TOptional<ETextFlowDirection>& InTextFlowDirection, const FCreateSlateTextLayout& InCreateSlateTextLayout, TSharedRef<ITextLayoutMarshaller> InTextMarshaller, TSharedRef<ITextLayoutMarshaller> InHintTextMarshaller);
-
-	SLATE_API void SetText(TAttribute<FText> InText);
+	SLATE_API void SetText(const TAttribute<FText>& InText);
 	SLATE_API FText GetText() const;
 
-	SLATE_API void SetHintText(TAttribute<FText> InHintText);
+	SLATE_API void SetHintText(const TAttribute<FText>& InHintText);
 	SLATE_API FText GetHintText() const;
 
-	SLATE_API void SetSearchText(TAttribute<FText> InSearchText);
+	SLATE_API void SetSearchText(const TAttribute<FText>& InSearchText);
 	SLATE_API FText GetSearchText() const;
 	
 	/** Get the index of the search result (0 if none) */
@@ -124,17 +57,11 @@ public:
 	SLATE_API void SetTextStyle(const FTextBlockStyle& InTextStyle);
 	SLATE_API const FTextBlockStyle& GetTextStyle() const;
 	
-	SLATE_API void SetHintTextOpacity(TAttribute<float> InOpacity);
+	SLATE_API void SetHintTextOpacity(const TAttribute<float>& InOpacity);
 	SLATE_API float GetHintTextOpacity() const;
 
 	/** Get the number of Text Lines */
-	UE_DEPRECATED(5.8, "Use the const version instead.")
-	int32 GetTextLineCount()
-	{
-		return const_cast<const FSlateEditableTextLayout*>(this)->GetTextLineCount();
-	}
-	/** Get the number of Text Lines */
-	SLATE_API int32 GetTextLineCount() const;
+	SLATE_API int32 GetTextLineCount();
 
 	/** Set the brush to use when drawing the cursor */
 	SLATE_API void SetCursorBrush(const TAttribute<const FSlateBrush*>& InCursorBrush);
@@ -175,33 +102,31 @@ public:
 	SLATE_API void SetTextFlowDirection(const TOptional<ETextFlowDirection>& InTextFlowDirection);
 
 	/** Set the wrapping to use for this document */
-	SLATE_API void SetTextWrapping(TAttribute<float> InWrapTextAt, TAttribute<bool> InAutoWrapText, TAttribute<ETextWrappingPolicy> InWrappingPolicy);
+	SLATE_API void SetTextWrapping(const TAttribute<float>& InWrapTextAt, const TAttribute<bool>& InAutoWrapText, const TAttribute<ETextWrappingPolicy>& InWrappingPolicy);
 
 	/** Set whether text wraps onto a new line when it's length exceeds this width; if this value is zero or negative, no wrapping occurs */
-	SLATE_API void SetWrapTextAt(TAttribute<float> InWrapTextAt);
+	SLATE_API void SetWrapTextAt(const TAttribute<float>& InWrapTextAt);
 
 	/** Set whether to wrap text automatically based on the widget's computed horizontal space */
-	SLATE_API void SetAutoWrapText(TAttribute<bool> InAutoWrapText);
+	SLATE_API void SetAutoWrapText(const TAttribute<bool>& InAutoWrapText);
 
 	/** Set the wrapping policy to use */
-	SLATE_API void SetWrappingPolicy(TAttribute<ETextWrappingPolicy> InWrappingPolicy);
+	SLATE_API void SetWrappingPolicy(const TAttribute<ETextWrappingPolicy>& InWrappingPolicy);
 
 	/** Set the amount of blank space left around the edges of text area */
-	SLATE_API void SetMargin(TAttribute<FMargin> InMargin);
+	SLATE_API void SetMargin(const TAttribute<FMargin>& InMargin);
 
 	/** Set how the text should be aligned with the margin */
-	SLATE_API void SetJustification(TAttribute<ETextJustify::Type> InJustification);
-	/** Get how the text should be aligned with the margin */
-	SLATE_API ETextJustify::Type GetJustification() const;
+	SLATE_API void SetJustification(const TAttribute<ETextJustify::Type>& InJustification);
 
 	/** Set the amount to scale each lines height by */
-	SLATE_API void SetLineHeightPercentage(TAttribute<float> InLineHeightPercentage);
+	SLATE_API void SetLineHeightPercentage(const TAttribute<float>& InLineHeightPercentage);
 
 	/** Set whether to leave extra space below the last line due to line height */
-	SLATE_API void SetApplyLineHeightToBottomLine(TAttribute<bool> InApplyLineHeightToBottomLine);
+	SLATE_API void SetApplyLineHeightToBottomLine(const TAttribute<bool>& InApplyLineHeightToBottomLine);
 
 	/** Set the text overflow policy that should be used to determine what happens to clipped text */
-	SLATE_API void SetOverflowPolicy(const TOptional<ETextOverflowPolicy>& InOverflowPolicy);
+	SLATE_API void SetOverflowPolicy(TOptional<ETextOverflowPolicy> InOverflowPolicy);
 
 	/** Set the information used to help identify who owns this text layout in the case of an error */
 	SLATE_API void SetDebugSourceInfo(const TAttribute<FString>& InDebugSourceInfo);
@@ -215,7 +140,7 @@ public:
 	/** Register and activate the IME context for this text layout */
 	SLATE_API void EnableTextInputMethodContext();
 
-	/** Refresh this editable text immediately, rather than wait for the usual caching mechanisms to take effect on the text Tick */
+	/** Refresh this editable text immediately, rather than wait for the usual caching mechanisms to take affect on the text Tick */
 	SLATE_API bool Refresh();
 
 	/** Force the text layout to be updated from the marshaller */
@@ -386,7 +311,7 @@ public:
 	/**
 	 * Given a location and a Direction to offset, return a new location.
 	 *
-	 * @param CurrentLocation    Cursor location from which to offset
+	 * @param Location    Cursor location from which to offset
 	 * @param Direction   Positive means right, negative means left.
 	 */
 	SLATE_API FTextLocation TranslatedLocation(const FTextLocation& CurrentLocation, int8 Direction) const;
@@ -429,13 +354,8 @@ public:
 	/** Returns whether the current text varies from the original */
 	SLATE_API bool HasTextChangedFromOriginal() const;
 
-	UE_DEPRECATED(5.8, "Use the properly spelled version BeginEditTransaction")
-	void BeginEditTransation()
-	{
-		BeginEditTransaction();
-	}
 	/** Called to begin an undoable editable text transaction */
-	SLATE_API void BeginEditTransaction();
+	SLATE_API void BeginEditTransation();
 
 	/** Called to end an undoable editable text transaction */
 	SLATE_API void EndEditTransaction();
@@ -465,17 +385,11 @@ public:
 
 	SLATE_API void LoadText();
 
-	UE_DEPRECATED(5.8, "ComputeVolatility always returns false with the addition of TSlateNestedAttributes.")
 	SLATE_API bool ComputeVolatility() const;
 
 	SLATE_API void Tick(const FGeometry& AllottedGeometry, const double InCurrentTime, const float InDeltaTime);
 
-	UE_DEPRECATED(5.8, "Use the const version instead.")
-	int32 OnPaint(const FPaintArgs& Args, const FGeometry& AllottedGeometry, const FSlateRect& MyCullingRect, FSlateWindowElementList& OutDrawElements, int32 LayerId, const FWidgetStyle& InWidgetStyle, bool bParentEnabled)
-	{
-		return const_cast<const FSlateEditableTextLayout*>(this)->OnPaint(Args, AllottedGeometry, MyCullingRect, OutDrawElements, LayerId, InWidgetStyle, bParentEnabled);
-	}
-	SLATE_API int32 OnPaint(const FPaintArgs& Args, const FGeometry& AllottedGeometry, const FSlateRect& MyCullingRect, FSlateWindowElementList& OutDrawElements, int32 LayerId, const FWidgetStyle& InWidgetStyle, bool bParentEnabled) const;
+	SLATE_API int32 OnPaint(const FPaintArgs& Args, const FGeometry& AllottedGeometry, const FSlateRect& MyCullingRect, FSlateWindowElementList& OutDrawElements, int32 LayerId, const FWidgetStyle& InWidgetStyle, bool bParentEnabled);
 
 	SLATE_API void CacheDesiredSize(float LayoutScaleMultiplier);
 
@@ -515,11 +429,6 @@ public:
 	/** Sets the paint policy to use for text if the font faces are still async loading. */
 	SLATE_API void SetFontFacesLoadingPaintPolicy(EFontFacesLoadingPaintPolicy InFontFacesLoadingPaintPolicy);
 private:
-	/** Function called when the BoundText attribute changes. Will call the Widget's OnTextChanged */
-	SLATE_API static void OnBoundTextChanged(SWidget& Widget, FSlateEditableTextLayout& SlateTextLayout);
-	/** Function called when the BoundSearchText attribute changes */
-	SLATE_API static void OnBoundSearchTextChanged(SWidget& Widget, FSlateEditableTextLayout& SlateTextLayout);
-
 	/** Insert the given text at the current cursor position, correctly taking into account new line characters */
 	SLATE_API void InsertTextAtCursorImpl(const FString& InString);
 
@@ -647,11 +556,10 @@ private:
 	FTextBlockStyle TextStyle;
 
 	/** Style used to draw the hint text (only valid when HintTextLayout is set) */
-	mutable FTextBlockStyle HintTextStyle;
+	FTextBlockStyle HintTextStyle;
 
 	/** The text displayed in this text block */
-	TSlateNestedAttribute<FText> BoundText;
-	FText PreviousBoundText;
+	TAttribute<FText> BoundText;
 
 	/** The state of BoundText last Tick() (used to allow updates when the text is changed) */
 	FTextSnapshot BoundTextLastTick;
@@ -660,13 +568,13 @@ private:
 	bool bWasPasswordLastTick;
 
 	/** The text that appears when there is no text in the text box */
-	TSlateNestedAttribute<FText> HintText;
+	TAttribute<FText> HintText;
 	
-	/** Opacity of the hint text */
-	TSlateNestedAttribute<float> HintTextOpacity;
+	/** Override opacity of the hint text */
+	TAttribute<float> HintTextOpacity;
 
 	/** The text to be searched for */
-	TSlateNestedAttribute<FText> BoundSearchText;
+	TAttribute<FText> BoundSearchText;
 
 	/** The state of BoundSearchText last Tick() (used to allow updates when the text is changed) */
 	FTextSnapshot BoundSearchTextLastTick;
@@ -684,25 +592,25 @@ private:
 	int32 CurrentSearchResultIndex;
 
 	/** Whether text wraps onto a new line when it's length exceeds this width; if this value is zero or negative, no wrapping occurs. */
-	TSlateNestedAttribute<float> WrapTextAt;
+	TAttribute<float> WrapTextAt;
 
 	/** True if we're wrapping text automatically based on the computed horizontal space for this widget */
-	TSlateNestedAttribute<bool> AutoWrapText;
+	TAttribute<bool> AutoWrapText;
 
 	/** The wrapping policy we're using */
-	TSlateNestedAttribute<ETextWrappingPolicy> WrappingPolicy;
+	TAttribute<ETextWrappingPolicy> WrappingPolicy;
 
 	/** The amount of blank space left around the edges of text area */
-	TSlateNestedAttribute<FMargin> Margin;
+	TAttribute<FMargin> Margin;
 
 	/** How the text should be aligned with the margin */
-	TSlateNestedAttribute<ETextJustify::Type> Justification;
+	TAttribute<ETextJustify::Type> Justification;
 
 	/** The amount to scale each lines height by */
-	TSlateNestedAttribute<float> LineHeightPercentage;
+	TAttribute<float> LineHeightPercentage;
 
 	/** Whether to leave extra space below the last line due to line height */
-	TSlateNestedAttribute<bool> ApplyLineHeightToBottomLine;
+	TAttribute<bool> ApplyLineHeightToBottomLine;
 
 	/** The information used to help identify who owns this text layout in the case of an error */
 	TAttribute<FString> DebugSourceInfo;
@@ -755,7 +663,7 @@ private:
 	/** Current undo state level that we've rolled back to, or INDEX_NONE if we haven't undone. Used for 'Redo'. */
 	int32 CurrentUndoLevel;
 
-	/** Undo state that will be pushed if text is actually changed between calls to BeginEditTransaction() and EndEditTransaction() */
+	/** Undo state that will be pushed if text is actually changed between calls to BeginEditTransation() and EndEditTransaction() */
 	TOptional<SlateEditableTextTypes::FUndoState> StateBeforeChangingText;
 
 	/** Track the number transactions that are opened */
@@ -789,7 +697,7 @@ private:
 	TOptional<ETextOverflowPolicy> OverflowPolicyOverride;
 
 	/** The last known size of the widget from the previous OnPaint, used to recalculate wrapping */
-	mutable FVector2f CachedSize;
+	FVector2f CachedSize;
 
 	/** A list commands to execute if a user presses the corresponding key-binding in the text box */
 	TSharedPtr<FUICommandList> UICommandList;
@@ -800,9 +708,6 @@ private:
 	/** Whether the cursor position has been changed externally*/
 	bool bSelectionChangedExternally;
 
-	/** Whether the cursor position would have moved in OnBoundTextChanged */
-	bool bMoveCursorToDocumentEnd;
-
 	/** The boundaries of the external selection */
 	int ExternalSelectionStart;
 	int ExternalSelectionEnd;
