PR #13890: Add Control + Shift + Z as an alternate InputChord for Redo

[RN] minor editor

#Synced-CL 50003544

diff --git a/Engine/Source/Runtime/Slate/Private/Framework/Commands/GenericCommands.cpp b/Engine/Source/Runtime/Slate/Private/Framework/Commands/GenericCommands.cpp
--- a/Engine/Source/Runtime/Slate/Private/Framework/Commands/GenericCommands.cpp
+++ b/Engine/Source/Runtime/Slate/Private/Framework/Commands/GenericCommands.cpp
@@ -17,7 +17,7 @@ void FGenericCommands::RegisterCommands()
 	UI_COMMAND(Duplicate, "Duplicate", "Duplicate selection", EUserInterfaceActionType::Button, FInputChord(EModifierKey::Control, EKeys::D))
 #endif
 	UI_COMMAND(Undo, "Undo", "Undo last action", EUserInterfaceActionType::Button, FInputChord(EModifierKey::Control, EKeys::Z))
-	UI_COMMAND(Redo, "Redo", "Redo last action", EUserInterfaceActionType::Button, FInputChord(EModifierKey::Control, EKeys::Y))
+	UI_COMMAND(Redo, "Redo", "Redo last action", EUserInterfaceActionType::Button, FInputChord(EModifierKey::Control, EKeys::Y), FInputChord(EModifierKey::Control | EModifierKey::Shift, EKeys::Z))
 
 	UI_COMMAND(Delete, "Delete", "Delete current selection", EUserInterfaceActionType::Button, FInputChord(EKeys::Delete), FInputChord(EKeys::BackSpace))
 
