Minor fixes and adjustments for Editable Texts in preparation for Dynamic Invalidation related changes

#Synced-CL 49551883

diff --git a/Engine/Plugins/Animation/DeformerGraph/Source/OptimusEditor/Private/Widgets/SOptimusShaderTextDocumentTextBox.cpp b/Engine/Plugins/Animation/DeformerGraph/Source/OptimusEditor/Private/Widgets/SOptimusShaderTextDocumentTextBox.cpp
--- a/Engine/Plugins/Animation/DeformerGraph/Source/OptimusEditor/Private/Widgets/SOptimusShaderTextDocumentTextBox.cpp
+++ b/Engine/Plugins/Animation/DeformerGraph/Source/OptimusEditor/Private/Widgets/SOptimusShaderTextDocumentTextBox.cpp
@@ -440,7 +440,7 @@ void SOptimusShaderTextDocumentTextBox::OnToggleComment() const
 	{
 		// Automatically go to the next line if there is no selection
 		int32 LineIndex = CursorLocation.GetLineIndex() + 1;
-		if (LineIndex < Text->GetTextLineCount())
+		if (LineIndex < const_cast<const SMultiLineEditableText*>(Text.Get())->GetTextLineCount())
 		{
 			FString Line;
 			Text->GetTextLine(LineIndex, Line);
diff --git a/Engine/Plugins/PCG/Source/PCGEditor/Private/Widgets/SPCGNodeSourceTextBox.cpp b/Engine/Plugins/PCG/Source/PCGEditor/Private/Widgets/SPCGNodeSourceTextBox.cpp
--- a/Engine/Plugins/PCG/Source/PCGEditor/Private/Widgets/SPCGNodeSourceTextBox.cpp
+++ b/Engine/Plugins/PCG/Source/PCGEditor/Private/Widgets/SPCGNodeSourceTextBox.cpp
@@ -436,7 +436,7 @@ void SPCGNodeSourceTextBox::OnToggleComment() const
 	{
 		// Automatically go to the next line if there is no selection
 		int32 LineIndex = CursorLocation.GetLineIndex() + 1;
-		if (LineIndex < Text->GetTextLineCount())
+		if (LineIndex < const_cast<const SMultiLineEditableText*>(Text.Get())->GetTextLineCount())
 		{
 			FString Line;
 			Text->GetTextLine(LineIndex, Line);
diff --git a/Engine/Source/Editor/InputBindingEditor/Private/Widgets/SChordEditor.cpp b/Engine/Source/Editor/InputBindingEditor/Private/Widgets/SChordEditor.cpp
--- a/Engine/Source/Editor/InputBindingEditor/Private/Widgets/SChordEditor.cpp
+++ b/Engine/Source/Editor/InputBindingEditor/Private/Widgets/SChordEditor.cpp
@@ -11,7 +11,6 @@
 #include "Internationalization/Internationalization.h"
 #include "Misc/AssertionMacros.h"
 #include "Styling/CoreStyle.h"
-#include "Templates/UniquePtr.h"
 #include "Templates/UnrealTemplate.h"
 #include "UObject/NameTypes.h"
 #include "Widgets/Text/SlateEditableTextLayout.h"
@@ -43,7 +42,7 @@ void SChordEditor::Construct( const FArguments& InArgs, TSharedPtr<FUICommandInf
 		.HintTextOpacity(0.25f)
 	);
 
-	EditableTextLayout->LoadText();
+	SlateEditableTextLayout.LoadText();
 }
 
 
@@ -64,17 +63,17 @@ FReply SChordEditor::OnKeyDown( const FGeometry& MyGeometry, const FKeyEvent& In
 			EditingInputChord.Key = Key;
 		}
 	
-		EditableTextLayout->BeginEditTransation();
+		SlateEditableTextLayout.BeginEditTransaction();
 
 		EditingInputChord.bCtrl = InKeyEvent.IsControlDown();
 		EditingInputChord.bAlt = InKeyEvent.IsAltDown();
 		EditingInputChord.bShift = InKeyEvent.IsShiftDown();
 		EditingInputChord.bCmd = InKeyEvent.IsCommandDown();
 
-		EditableTextLayout->LoadText();
-		EditableTextLayout->GoTo(ETextLocation::EndOfDocument);
+		SlateEditableTextLayout.LoadText();
+		SlateEditableTextLayout.GoTo(ETextLocation::EndOfDocument);
 
-		EditableTextLayout->EndEditTransaction();
+		SlateEditableTextLayout.EndEditTransaction();
 		
 		OnChordTyped( EditingInputChord );
 		
diff --git a/Engine/Source/Runtime/Slate/Private/Widgets/Input/SEditableText.cpp b/Engine/Source/Runtime/Slate/Private/Widgets/Input/SEditableText.cpp
--- a/Engine/Source/Runtime/Slate/Private/Widgets/Input/SEditableText.cpp
+++ b/Engine/Source/Runtime/Slate/Private/Widgets/Input/SEditableText.cpp
@@ -4,6 +4,7 @@
 #include "Framework/Text/TextEditHelper.h"
 #include "Framework/Text/PlainTextLayoutMarshaller.h"
 #include "Widgets/Text/SlateEditableTextLayout.h"
+#include "Widgets/Text/SlateTextBlockLayout.h"
 #include "Types/ReflectionMetadata.h"
 #include "Types/TrackedMetaData.h"
 
@@ -12,6 +13,7 @@
 #endif
 
 SEditableText::SEditableText()
+	: SlateEditableTextLayout(*this)
 {
 #if WITH_ACCESSIBILITY
 	AccessibleBehavior = EAccessibleBehavior::Auto;
@@ -19,15 +21,9 @@ SEditableText::SEditableText()
 #endif
 }
 
-SEditableText::~SEditableText()
-{
-	// Needed to avoid "deletion of pointer to incomplete type 'FSlateEditableTextLayout'; no destructor called" error when using TUniquePtr
-}
-
 void SEditableText::Construct( const FArguments& InArgs )
 {
-	bIsReadOnly = InArgs._IsReadOnly;
-	bIsPassword = InArgs._IsPassword;
+	SetIsReadOnly(InArgs._IsReadOnly);
 
 	bIsCaretMovedWhenGainFocus = InArgs._IsCaretMovedWhenGainFocus;
 	bSelectAllTextWhenFocused = InArgs._SelectAllTextWhenFocused;
@@ -40,7 +36,7 @@ void SEditableText::Construct( const FArguments& InArgs )
 	OnTextChangedCallback = InArgs._OnTextChanged;
 	OnTextCommittedCallback = InArgs._OnTextCommitted;
 	OnCursorMovedWithSelectionCallback = InArgs._OnCursorMovedWithSelection;
-	MinDesiredWidth = InArgs._MinDesiredWidth;
+	SetMinDesiredWidth(InArgs._MinDesiredWidth);
 	bSelectAllTextOnCommit = InArgs._SelectAllTextOnCommit;
 	bSelectWordOnMouseDoubleClick = InArgs._SelectWordOnMouseDoubleClick;
 	VirtualKeyboardType = InArgs._VirtualKeyboardType;
@@ -51,10 +47,6 @@ void SEditableText::Construct( const FArguments& InArgs )
 	OnKeyDownHandler = InArgs._OnKeyDownHandler;
 	bEnableIntegratedKeyboard = InArgs._EnableIntegratedKeyboard;
 
-	Font = InArgs._Font;
-	ColorAndOpacity = InArgs._ColorAndOpacity;
-	BackgroundImageSelected = InArgs._BackgroundImageSelected;
-
 	// We use the given style when creating the text layout as it may not be safe to call the override delegates until we've finished being constructed
 	// The first call to SynchronizeTextStyle will apply the correct overrides, and that will happen before the first paint
 	check(InArgs._Style);
@@ -64,26 +56,31 @@ void SEditableText::Construct( const FArguments& InArgs )
 	TextStyle.HighlightShape = InArgs._Style->BackgroundImageSelected;
 
 	PlainTextMarshaller = FPlainTextLayoutMarshaller::Create();
-	PlainTextMarshaller->SetIsPassword(bIsPassword);
 
 	// We use a separate marshaller for the hint text, as that should never be displayed as a password
 	TSharedRef<FPlainTextLayoutMarshaller> HintTextMarshaller = FPlainTextLayoutMarshaller::Create();
 
-	EditableTextLayout = MakeUnique<FSlateEditableTextLayout>(*this, InArgs._Text, TextStyle, InArgs._TextShapingMethod, InArgs._TextFlowDirection, FCreateSlateTextLayout(), PlainTextMarshaller.ToSharedRef(), HintTextMarshaller);
-	EditableTextLayout->SetHintText(InArgs._HintText);
-	EditableTextLayout->SetHintTextOpacity(InArgs._HintTextOpacity);
-	EditableTextLayout->SetSearchText(InArgs._SearchText);
-	EditableTextLayout->SetCursorBrush(InArgs._CaretImage.IsSet() ? InArgs._CaretImage : &InArgs._Style->CaretImage);
-	EditableTextLayout->SetCompositionBrush(InArgs._BackgroundImageComposing.IsSet() ? InArgs._BackgroundImageComposing : &InArgs._Style->BackgroundImageComposing);
-	EditableTextLayout->SetDebugSourceInfo(TAttribute<FString>::Create(TAttribute<FString>::FGetter::CreateLambda([this]{ return FReflectionMetaData::GetWidgetDebugInfo(this); })));
-	EditableTextLayout->SetJustification(InArgs._Justification);
-	EditableTextLayout->SetOverflowPolicy(InArgs._OverflowPolicy);
-	EditableTextLayout->SetFontFacesLoadingPaintPolicy(InArgs._FontFacesLoadingPaintPolicy);
+	SlateEditableTextLayout.Construct(InArgs._Text, TextStyle, InArgs._TextShapingMethod, InArgs._TextFlowDirection, FCreateSlateTextLayout(), PlainTextMarshaller.ToSharedRef(), MoveTemp(HintTextMarshaller));
+	SlateEditableTextLayout.SetCursorBrush(InArgs._CaretImage.IsSet() ? InArgs._CaretImage : &InArgs._Style->CaretImage);
+	SlateEditableTextLayout.SetCompositionBrush(InArgs._BackgroundImageComposing.IsSet() ? InArgs._BackgroundImageComposing : &InArgs._Style->BackgroundImageComposing);
+	SlateEditableTextLayout.SetDebugSourceInfo(TAttribute<FString>::Create(TAttribute<FString>::FGetter::CreateLambda([this]{ return FReflectionMetaData::GetWidgetDebugInfo(this); })));
+	SlateEditableTextLayout.SetOverflowPolicy(InArgs._OverflowPolicy);
+	SlateEditableTextLayout.SetFontFacesLoadingPaintPolicy(InArgs._FontFacesLoadingPaintPolicy);
 	if (InArgs._OnAllFontFacesFinishLoading.IsBound())
 	{
-		EditableTextLayout->SetOnAllFontFacesFinishLoadingDelegate(InArgs._OnAllFontFacesFinishLoading);
+		SlateEditableTextLayout.SetOnAllFontFacesFinishLoadingDelegate(InArgs._OnAllFontFacesFinishLoading);
 	}
 
+	SetText(InArgs._Text);
+	SetIsPassword(InArgs._IsPassword);
+	SetHintText(InArgs._HintText);
+	SetHintTextOpacity(InArgs._HintTextOpacity);
+	SetSearchText(InArgs._SearchText);
+	SetFont(InArgs._Font);
+	SetBackgroundImageSelected(InArgs._BackgroundImageSelected);
+	SetColorAndOpacity(InArgs._ColorAndOpacity);
+	SetJustification(InArgs._Justification);
+
 	// build context menu extender
 	MenuExtender = MakeShareable(new FExtender());
 	MenuExtender->AddMenuExtension("EditText", EExtensionHook::Before, TSharedPtr<FUICommandList>(), InArgs._ContextMenuExtender);
@@ -91,31 +88,31 @@ void SEditableText::Construct( const FArguments& InArgs )
 	AddMetadata(MakeShared<FTrackedMetaData>(this, FName("EditableText")));
 }
 
-void SEditableText::SetText( const TAttribute< FText >& InNewText )
+void SEditableText::SetText(TAttribute<FText> InNewText)
 {
-	EditableTextLayout->SetText(InNewText);
+	SlateEditableTextLayout.SetText(MoveTemp(InNewText));
 }
 
 FText SEditableText::GetText() const
 {
-	return EditableTextLayout->GetText();
+	return SlateEditableTextLayout.GetText();
 }
 
 bool SEditableText::SetEditableText(const FText& InNewText)
 {
-	return EditableTextLayout->SetEditableText(InNewText);
+	return SlateEditableTextLayout.SetEditableText(InNewText);
 }
 
-void SEditableText::SetFont( const TAttribute< FSlateFontInfo >& InNewFont )
+void SEditableText::SetFont(TAttribute<FSlateFontInfo> InNewFont)
 {
-	Font = InNewFont;
+	Font = MoveTemp(InNewFont);
 
 	Invalidate(EInvalidateWidgetReason::Layout);
 }
 
 FSlateFontInfo SEditableText::GetFont() const
 {
-	return Font.Get();
+	return Font.Get(SlateEditableTextLayout.GetTextStyle().Font);
 }
 
 void SEditableText::SetTextStyle( const FEditableTextStyle& InNewTextStyle )
@@ -131,25 +128,25 @@ void SEditableText::SetTextBlockStyle(const FTextBlockStyle* InTextStyle)
 {
 	if (InTextStyle)
 	{
-		EditableTextLayout->SetTextStyle(*InTextStyle);
+		SlateEditableTextLayout.SetTextStyle(*InTextStyle);
 		Invalidate(EInvalidateWidgetReason::Layout); //Using Layout as changing text block size can affect the size.
 	}
 }
 
 void SEditableText::Tick( const FGeometry& AllottedGeometry, const double InCurrentTime, const float InDeltaTime )
 {
-	EditableTextLayout->Tick(AllottedGeometry, InCurrentTime, InDeltaTime);
+	SlateEditableTextLayout.Tick(AllottedGeometry, InCurrentTime, InDeltaTime);
 }
 
 int32 SEditableText::OnPaint( const FPaintArgs& Args, const FGeometry& AllottedGeometry, const FSlateRect& MyCullingRect, FSlateWindowElementList& OutDrawElements, int32 LayerId, const FWidgetStyle& InWidgetStyle, bool bParentEnabled ) const
 {
-	const FTextBlockStyle& EditableTextStyle = EditableTextLayout->GetTextStyle();
+	const FTextBlockStyle& EditableTextStyle = SlateEditableTextLayout.GetTextStyle();
 	const FLinearColor ForegroundColor = EditableTextStyle.ColorAndOpacity.GetColor(InWidgetStyle);
 
 	FWidgetStyle TextWidgetStyle = FWidgetStyle(InWidgetStyle)
 		.SetForegroundColor(ForegroundColor);
 
-	LayerId = EditableTextLayout->OnPaint(Args, AllottedGeometry, MyCullingRect, OutDrawElements, LayerId, TextWidgetStyle, ShouldBeEnabled(bParentEnabled));
+	LayerId = SlateEditableTextLayout.OnPaint(Args, AllottedGeometry, MyCullingRect, OutDrawElements, LayerId, TextWidgetStyle, ShouldBeEnabled(bParentEnabled));
 
 	return LayerId;
 }
@@ -157,25 +154,25 @@ int32 SEditableText::OnPaint( const FPaintArgs& Args, const FGeometry& AllottedG
 void SEditableText::CacheDesiredSize(float LayoutScaleMultiplier)
 {
 	SynchronizeTextStyle();
-	EditableTextLayout->CacheDesiredSize(LayoutScaleMultiplier);
+	SlateEditableTextLayout.CacheDesiredSize(LayoutScaleMultiplier);
 	SWidget::CacheDesiredSize(LayoutScaleMultiplier);
 }
 
 FVector2D SEditableText::ComputeDesiredSize(float LayoutScaleMultiplier) const
 {
-	FVector2D TextLayoutSize = EditableTextLayout->ComputeDesiredSize(LayoutScaleMultiplier);
-	TextLayoutSize.X = FMath::Max(TextLayoutSize.X, MinDesiredWidth.Get());
+	FVector2D TextLayoutSize = SlateEditableTextLayout.ComputeDesiredSize(LayoutScaleMultiplier);
+	TextLayoutSize.X = FMath::Max(TextLayoutSize.X, MinDesiredWidth.Get(0.0f));
 	return TextLayoutSize;
 }
 
 FChildren* SEditableText::GetChildren()
 {
-	return EditableTextLayout->GetChildren();
+	return SlateEditableTextLayout.GetChildren();
 }
 
 void SEditableText::OnArrangeChildren(const FGeometry& AllottedGeometry, FArrangedChildren& ArrangedChildren) const
 {
-	EditableTextLayout->OnArrangeChildren(AllottedGeometry, ArrangedChildren);
+	SlateEditableTextLayout.OnArrangeChildren(AllottedGeometry, ArrangedChildren);
 }
 
 FReply SEditableText::OnDragOver( const FGeometry& MyGeometry, const FDragDropEvent& DragDropEvent )
@@ -188,7 +185,7 @@ FReply SEditableText::OnDragOver( const FGeometry& MyGeometry, const FDragDropEv
 			return FReply::Handled();
 		}
 	}
-	
+
 	return FReply::Unhandled();
 }
 
@@ -199,11 +196,11 @@ FReply SEditableText::OnDrop( const FGeometry& MyGeometry, const FDragDropEvent&
 	{
 		if ( DragDropOp->HasText() )
 		{
-			EditableTextLayout->SetText(FText::FromString(DragDropOp->GetText()));
+			SlateEditableTextLayout.SetText(FText::FromString(DragDropOp->GetText()));
 			return FReply::Handled();
 		}
 	}
-	
+
 	return FReply::Unhandled();
 }
 
@@ -214,13 +211,13 @@ bool SEditableText::SupportsKeyboardFocus() const
 
 FReply SEditableText::OnFocusReceived( const FGeometry& MyGeometry, const FFocusEvent& InFocusEvent )
 {
-	EditableTextLayout->HandleFocusReceived(InFocusEvent);
+	SlateEditableTextLayout.HandleFocusReceived(InFocusEvent);
 	return FReply::Handled();
 }
 
 void SEditableText::OnFocusLost( const FFocusEvent& InFocusEvent )
 {
-	EditableTextLayout->HandleFocusLost(InFocusEvent);
+	SlateEditableTextLayout.HandleFocusLost(InFocusEvent);
 }
 
 FReply SEditableText::OnKeyChar( const FGeometry& MyGeometry, const FCharacterEvent& InCharacterEvent )
@@ -232,10 +229,10 @@ FReply SEditableText::OnKeyChar( const FGeometry& MyGeometry, const FCharacterEv
 	{
 		Reply = OnKeyCharHandler.Execute(MyGeometry, InCharacterEvent);
 	}
-	
+
 	if (!Reply.IsEventHandled())
 	{
-		Reply = EditableTextLayout->HandleKeyChar(InCharacterEvent);
+		Reply = SlateEditableTextLayout.HandleKeyChar(InCharacterEvent);
 	}
 
 	return Reply;
@@ -253,7 +250,7 @@ FReply SEditableText::OnKeyDown( const FGeometry& MyGeometry, const FKeyEvent& I
 
 	if (!Reply.IsEventHandled())
 	{
-		Reply = EditableTextLayout->HandleKeyDown(InKeyEvent);
+		Reply = SlateEditableTextLayout.HandleKeyDown(InKeyEvent);
 
 		if (!Reply.IsEventHandled())
 		{
@@ -266,27 +263,27 @@ FReply SEditableText::OnKeyDown( const FGeometry& MyGeometry, const FKeyEvent& I
 
 FReply SEditableText::OnKeyUp( const FGeometry& MyGeometry, const FKeyEvent& InKeyEvent )
 {
-	return EditableTextLayout->HandleKeyUp(InKeyEvent);
+	return SlateEditableTextLayout.HandleKeyUp(InKeyEvent);
 }
 
 FReply SEditableText::OnMouseButtonDown( const FGeometry& InMyGeometry, const FPointerEvent& InMouseEvent )
 {
-	return EditableTextLayout->HandleMouseButtonDown(InMyGeometry, InMouseEvent);
+	return SlateEditableTextLayout.HandleMouseButtonDown(InMyGeometry, InMouseEvent);
 }
 
 FReply SEditableText::OnMouseButtonUp( const FGeometry& InMyGeometry, const FPointerEvent& InMouseEvent )
 {
-	return EditableTextLayout->HandleMouseButtonUp(InMyGeometry, InMouseEvent);
+	return SlateEditableTextLayout.HandleMouseButtonUp(InMyGeometry, InMouseEvent);
 }
 
 FReply SEditableText::OnMouseMove( const FGeometry& InMyGeometry, const FPointerEvent& InMouseEvent )
 {
-	return EditableTextLayout->HandleMouseMove(InMyGeometry, InMouseEvent);
+	return SlateEditableTextLayout.HandleMouseMove(InMyGeometry, InMouseEvent);
 }
 
 FReply SEditableText::OnMouseButtonDoubleClick( const FGeometry& InMyGeometry, const FPointerEvent& InMouseEvent )
 {
-	return EditableTextLayout->HandleMouseButtonDoubleClick(InMyGeometry, InMouseEvent);
+	return SlateEditableTextLayout.HandleMouseButtonDoubleClick(InMyGeometry, InMouseEvent);
 }
 
 FCursorReply SEditableText::OnCursorQuery( const FGeometry& MyGeometry, const FPointerEvent& CursorEvent ) const
@@ -308,7 +305,7 @@ bool SEditableText::ComputeVolatility() const
 {
 	return SWidget::ComputeVolatility()
 		|| HasKeyboardFocus()
-		|| EditableTextLayout->ComputeVolatility()
+		|| SlateEditableTextLayout.ComputeVolatility()
 		|| Font.IsBound()
 		|| ColorAndOpacity.IsBound()
 		|| BackgroundImageSelected.IsBound()
@@ -317,200 +314,202 @@ bool SEditableText::ComputeVolatility() const
 		|| MinDesiredWidth.IsBound();
 }
 
-void SEditableText::SetHintText( const TAttribute< FText >& InHintText )
+void SEditableText::SetHintText(TAttribute<FText> InHintText )
 {
-	EditableTextLayout->SetHintText(InHintText);
+	SlateEditableTextLayout.SetHintText(MoveTemp(InHintText));
 }
 
 FText SEditableText::GetHintText() const
 {
-	return EditableTextLayout->GetHintText();
+	return SlateEditableTextLayout.GetHintText();
 }
 
-void SEditableText::SetHintTextOpacity(const TAttribute<float>& InHintTextOpacity)
+void SEditableText::SetHintTextOpacity(TAttribute<float> InHintTextOpacity)
 {
-	EditableTextLayout->SetHintTextOpacity(InHintTextOpacity);
+	SlateEditableTextLayout.SetHintTextOpacity(MoveTemp(InHintTextOpacity));
 }
 
 float SEditableText::GetHintTextOpacity() const
 {
-	return EditableTextLayout->GetHintTextOpacity();
+	return SlateEditableTextLayout.GetHintTextOpacity();
 }
 
-void SEditableText::SetSearchText(const TAttribute<FText>& InSearchText)
+void SEditableText::SetSearchText(TAttribute<FText> InSearchText)
 {
-	EditableTextLayout->SetSearchText(InSearchText);
+	SlateEditableTextLayout.SetSearchText(MoveTemp(InSearchText));
 }
 
 FText SEditableText::GetSearchText() const
 {
-	return EditableTextLayout->GetSearchText();
+	return SlateEditableTextLayout.GetSearchText();
 }
 
-void SEditableText::SetIsReadOnly( TAttribute< bool > InIsReadOnly )
+void SEditableText::SetIsReadOnly(TAttribute<bool> InIsReadOnly)
 {
-	bIsReadOnly = InIsReadOnly;
+	bIsReadOnly = MoveTemp(InIsReadOnly);
 }
 
-void SEditableText::SetIsPassword( TAttribute< bool > InIsPassword )
+void SEditableText::SetIsPassword(TAttribute<bool> InIsPassword)
 {
-	bIsPassword = InIsPassword;
-	PlainTextMarshaller->SetIsPassword(bIsPassword);
+	PlainTextMarshaller->SetIsPassword(InIsPassword);
+	bIsPassword = MoveTemp(InIsPassword);
 }
 
 void SEditableText::SetColorAndOpacity(TAttribute<FSlateColor> Color)
 {
-	ColorAndOpacity = Color;
+	ColorAndOpacity = MoveTemp(Color);
 }
 
-void SEditableText::SetMinDesiredWidth(const TAttribute<float>& InMinDesiredWidth)
+void SEditableText::SetBackgroundImageSelected(TAttribute<const FSlateBrush*> InBackgroundImageSelected)
 {
-	MinDesiredWidth = InMinDesiredWidth;
+	BackgroundImageSelected = MoveTemp(InBackgroundImageSelected);
 }
 
-void SEditableText::SetIsCaretMovedWhenGainFocus(const TAttribute<bool>& InIsCaretMovedWhenGainFocus)
+void SEditableText::SetMinDesiredWidth(TAttribute<float> InMinDesiredWidth)
 {
-	bIsCaretMovedWhenGainFocus = InIsCaretMovedWhenGainFocus;
+	MinDesiredWidth = MoveTemp(InMinDesiredWidth);
 }
 
-void SEditableText::SetSelectAllTextWhenFocused(const TAttribute<bool>& InSelectAllTextWhenFocused)
+void SEditableText::SetIsCaretMovedWhenGainFocus(TAttribute<bool> InIsCaretMovedWhenGainFocus)
 {
-	bSelectAllTextWhenFocused = InSelectAllTextWhenFocused;
+	bIsCaretMovedWhenGainFocus = MoveTemp(InIsCaretMovedWhenGainFocus);
 }
 
-void SEditableText::SetRevertTextOnEscape(const TAttribute<bool>& InRevertTextOnEscape)
+void SEditableText::SetSelectAllTextWhenFocused(TAttribute<bool> InSelectAllTextWhenFocused)
 {
-	bRevertTextOnEscape = InRevertTextOnEscape;
+	bSelectAllTextWhenFocused = MoveTemp(InSelectAllTextWhenFocused);
 }
 
-void SEditableText::SetClearKeyboardFocusOnCommit(const TAttribute<bool>& InClearKeyboardFocusOnCommit)
+void SEditableText::SetRevertTextOnEscape(TAttribute<bool> InRevertTextOnEscape)
 {
-	bClearKeyboardFocusOnCommit = InClearKeyboardFocusOnCommit;
+	bRevertTextOnEscape = MoveTemp(InRevertTextOnEscape);
 }
 
-void SEditableText::SetSelectAllTextOnCommit(const TAttribute<bool>& InSelectAllTextOnCommit)
+void SEditableText::SetClearKeyboardFocusOnCommit(TAttribute<bool> InClearKeyboardFocusOnCommit)
 {
-	bSelectAllTextOnCommit = InSelectAllTextOnCommit;
+	bClearKeyboardFocusOnCommit = MoveTemp(InClearKeyboardFocusOnCommit);
 }
 
-void SEditableText::SetSelectWordOnMouseDoubleClick(const TAttribute<bool>& InSelectWordOnMouseDoubleClick)
+void SEditableText::SetSelectAllTextOnCommit(TAttribute<bool> InSelectAllTextOnCommit)
 {
-	bSelectWordOnMouseDoubleClick = InSelectWordOnMouseDoubleClick;
+	bSelectAllTextOnCommit = MoveTemp(InSelectAllTextOnCommit);
 }
 
-void SEditableText::SetJustification(const TAttribute<ETextJustify::Type>& InJustification)
+void SEditableText::SetSelectWordOnMouseDoubleClick(TAttribute<bool> InSelectWordOnMouseDoubleClick)
 {
-	EditableTextLayout->SetJustification(InJustification);
+	bSelectWordOnMouseDoubleClick = MoveTemp(InSelectWordOnMouseDoubleClick);
 }
 
-void SEditableText::SetAllowContextMenu(const TAttribute< bool >& InAllowContextMenu)
+void SEditableText::SetJustification(TAttribute<ETextJustify::Type> InJustification)
 {
-	bAllowContextMenu = InAllowContextMenu;
+	SlateEditableTextLayout.SetJustification(MoveTemp(InJustification));
 }
 
-void SEditableText::SetEnableIntegratedKeyboard(const TAttribute<bool>& InEnableIntegratedKeyboard)
+void SEditableText::SetAllowContextMenu(TAttribute<bool> InAllowContextMenu)
 {
-	bEnableIntegratedKeyboard = InEnableIntegratedKeyboard;
+	bAllowContextMenu = MoveTemp(InAllowContextMenu);
 }
 
-void SEditableText::SetVirtualKeyboardDismissAction(TAttribute< EVirtualKeyboardDismissAction > InVirtualKeyboardDismissAction)
+void SEditableText::SetEnableIntegratedKeyboard(TAttribute<bool> InEnableIntegratedKeyboard)
 {
-	VirtualKeyboardDismissAction = InVirtualKeyboardDismissAction;
+	bEnableIntegratedKeyboard = MoveTemp(InEnableIntegratedKeyboard);
+}
+
+void SEditableText::SetVirtualKeyboardDismissAction(TAttribute<EVirtualKeyboardDismissAction> InVirtualKeyboardDismissAction)
+{
+	VirtualKeyboardDismissAction = MoveTemp(InVirtualKeyboardDismissAction);
 }
 
 void SEditableText::SetTextShapingMethod(const TOptional<ETextShapingMethod>& InTextShapingMethod)
 {
-	EditableTextLayout->SetTextShapingMethod(InTextShapingMethod);
+	SlateEditableTextLayout.SetTextShapingMethod(InTextShapingMethod);
 }
 
 void SEditableText::SetTextFlowDirection(const TOptional<ETextFlowDirection>& InTextFlowDirection)
 {
-	EditableTextLayout->SetTextFlowDirection(InTextFlowDirection);
+	SlateEditableTextLayout.SetTextFlowDirection(InTextFlowDirection);
 }
 
-void SEditableText::SetOverflowPolicy(TOptional<ETextOverflowPolicy> InOverflowPolicy)
+void SEditableText::SetOverflowPolicy(const TOptional<ETextOverflowPolicy>& InOverflowPolicy)
 {
-	EditableTextLayout->SetOverflowPolicy(InOverflowPolicy);
+	SlateEditableTextLayout.SetOverflowPolicy(InOverflowPolicy);
 }
 
 void SEditableText::SetFontFacesLoadingPaintPolicy(EFontFacesLoadingPaintPolicy InFontFacesLoadingPaintPolicy)
 {
-	if (EditableTextLayout)
-	{
-		EditableTextLayout->SetFontFacesLoadingPaintPolicy(InFontFacesLoadingPaintPolicy);
-	}
+	SlateEditableTextLayout.SetFontFacesLoadingPaintPolicy(InFontFacesLoadingPaintPolicy);
 }
 
 bool SEditableText::AnyTextSelected() const
 {
-	return EditableTextLayout->AnyTextSelected();
+	return SlateEditableTextLayout.AnyTextSelected();
 }
 
 void SEditableText::SelectAllText()
 {
-	EditableTextLayout->SelectAllText();
+	SlateEditableTextLayout.SelectAllText();
 }
 
 void SEditableText::ClearSelection()
 {
-	EditableTextLayout->ClearSelection();
+	SlateEditableTextLayout.ClearSelection();
 }
 
 FText SEditableText::GetSelectedText() const
 {
-	return EditableTextLayout->GetSelectedText();
+	return SlateEditableTextLayout.GetSelectedText();
 }
 
 void SEditableText::GoTo(const FTextLocation& NewLocation)
 {
-	EditableTextLayout->GoTo(NewLocation);
+	SlateEditableTextLayout.GoTo(NewLocation);
 }
 
 void SEditableText::GoTo(const ETextLocation NewLocation)
 {
-	EditableTextLayout->GoTo(NewLocation);
+	SlateEditableTextLayout.GoTo(NewLocation);
 }
 
 void SEditableText::ScrollTo(const FTextLocation& NewLocation)
 {
-	EditableTextLayout->ScrollTo(NewLocation);
+	SlateEditableTextLayout.ScrollTo(NewLocation);
 }
 
 void SEditableText::ScrollTo(const ETextLocation NewLocation)
 {
-	EditableTextLayout->ScrollTo(NewLocation);
+	SlateEditableTextLayout.ScrollTo(NewLocation);
 }
 
 void SEditableText::BeginSearch(const FText& InSearchText, const ESearchCase::Type InSearchCase, const bool InReverse)
 {
-	EditableTextLayout->BeginSearch(InSearchText, InSearchCase, InReverse);
+	SlateEditableTextLayout.BeginSearch(InSearchText, InSearchCase, InReverse);
 }
 
 void SEditableText::AdvanceSearch(const bool InReverse)
 {
-	EditableTextLayout->AdvanceSearch(InReverse);
+	SlateEditableTextLayout.AdvanceSearch(InReverse);
 }
 
 void SEditableText::EnableTextInputMethodContext()
 {
-	EditableTextLayout->EnableTextInputMethodContext();
+	SlateEditableTextLayout.EnableTextInputMethodContext();
 }
 
 FTextSelection SEditableText::GetSelection() const
 {
-	return EditableTextLayout->GetSelection();
+	return SlateEditableTextLayout.GetSelection();
 }
 
 void SEditableText::SelectText(const FTextLocation& InSelectionStart, const FTextLocation& InCursorLocation)
 {
-	EditableTextLayout->SelectText(InSelectionStart, InCursorLocation);
+	SlateEditableTextLayout.SelectText(InSelectionStart, InCursorLocation);
 }
 
 void SEditableText::SynchronizeTextStyle()
 {
 	// Has the style used for this editable text changed?
 	bool bTextStyleChanged = false;
-	FTextBlockStyle NewTextStyle = EditableTextLayout->GetTextStyle();
+	FTextBlockStyle NewTextStyle = SlateEditableTextLayout.GetTextStyle();
 
 	// Sync from the font override
 	if (Font.IsSet())
@@ -547,8 +546,8 @@ void SEditableText::SynchronizeTextStyle()
 
 	if (bTextStyleChanged)
 	{
-		EditableTextLayout->SetTextStyle(NewTextStyle);
-		EditableTextLayout->ForceRefreshTextLayout(EditableTextLayout->GetEditableText());
+		SlateEditableTextLayout.SetTextStyle(NewTextStyle);
+		SlateEditableTextLayout.ForceRefreshTextLayout(SlateEditableTextLayout.GetEditableText());
 	}
 }
 
@@ -633,7 +632,7 @@ void SEditableText::EnsureActiveTick()
 	auto DoActiveTick = [this](double InCurrentTime, float InDeltaTime) -> EActiveTimerReturnType
 	{
 		// Continue if we still have focus, otherwise treat as a fire-and-forget Tick() request
-		const bool bShouldAppearFocused = HasKeyboardFocus() || EditableTextLayout->HasActiveContextMenu();
+		const bool bShouldAppearFocused = HasKeyboardFocus() || SlateEditableTextLayout.HasActiveContextMenu();
 		return (bShouldAppearFocused) ? EActiveTimerReturnType::Continue : EActiveTimerReturnType::Stop;
 	};
 
@@ -686,8 +685,8 @@ TSharedPtr<SWidget> SEditableText::BuildContextMenuContent() const
 	{
 		return OnContextMenuOpening.Execute();
 	}
-	
-	return EditableTextLayout->BuildDefaultContextMenu(MenuExtender);
+
+	return SlateEditableTextLayout.BuildDefaultContextMenu(MenuExtender);
 }
 
 void SEditableText::OnBeginTextEdit(const FText& InText)
@@ -708,17 +707,17 @@ void SEditableText::OnTextCommitted(const FText& InText, const ETextCommit::Type
 void SEditableText::OnCursorMoved(const FTextLocation& InLocation)
 {
 	FTextSelection TextSelection = GetSelection();
-	
+
 	OnCursorMovedWithSelectionCallback.ExecuteIfBound(InLocation, TextSelection);
 	Invalidate(EInvalidateWidgetReason::Layout);
 }
 
-float SEditableText::UpdateAndClampHorizontalScrollBar(const float InViewOffset, const float InViewFraction, const EVisibility InVisiblityOverride)
+float SEditableText::UpdateAndClampHorizontalScrollBar(const float InViewOffset, const float InViewFraction, const EVisibility InVisibilityOverride)
 {
-	return EditableTextLayout->GetScrollOffset().X;
+	return SlateEditableTextLayout.GetScrollOffset().X;
 }
 
-float SEditableText::UpdateAndClampVerticalScrollBar(const float InViewOffset, const float InViewFraction, const EVisibility InVisiblityOverride)
+float SEditableText::UpdateAndClampVerticalScrollBar(const float InViewOffset, const float InViewFraction, const EVisibility InVisibilityOverride)
 {
 	return 0.0f;
 }
@@ -737,8 +736,5 @@ TOptional<FText> SEditableText::GetDefaultAccessibleText(EAccessibleType Accessi
 
 void SEditableText::ToggleVirtualKeyboard(bool bShow, const int32 UserIndex)
 {
-	if (EditableTextLayout.IsValid())
-	{
-		EditableTextLayout->ToggleVirtualKeyboard(bShow, UserIndex);
-	}
+	SlateEditableTextLayout.ToggleVirtualKeyboard(bShow, UserIndex);
 }
diff --git a/Engine/Source/Runtime/Slate/Private/Widgets/Input/SEditableTextBox.cpp b/Engine/Source/Runtime/Slate/Private/Widgets/Input/SEditableTextBox.cpp
--- a/Engine/Source/Runtime/Slate/Private/Widgets/Input/SEditableTextBox.cpp
+++ b/Engine/Source/Runtime/Slate/Private/Widgets/Input/SEditableTextBox.cpp
@@ -43,7 +43,7 @@ void SEditableTextBox::Construct( const FArguments& InArgs )
 	OnCursorMovedWithSelection = InArgs._OnCursorMovedWithSelection;
 
 	SBorder::Construct( SBorder::FArguments()
-		.BorderImage( this, &SEditableTextBox::GetBorderImage )
+		.BorderImage( this, &SEditableTextBox::DetermineBorderImage )
 		.BorderBackgroundColor( this, &SEditableTextBox::DetermineBackgroundColor )
 		.ForegroundColor( this, &SEditableTextBox::DetermineForegroundColor )
 		.Padding(0.f)
@@ -140,9 +140,9 @@ void SEditableTextBox::SetTextBlockStyle(const FTextBlockStyle* InTextStyle)
 	}
 }
 
-void SEditableTextBox::SetText( const TAttribute< FText >& InNewText )
+void SEditableTextBox::SetText(TAttribute<FText> InNewText)
 {
-	EditableText->SetText( InNewText );
+	EditableText->SetText(MoveTemp(InNewText));
 }
 
 
@@ -186,13 +186,13 @@ void SEditableTextBox::SetError( const FString& InError )
 
 void SEditableTextBox::SetOnKeyCharHandler(FOnKeyChar InOnKeyCharHandler)
 {
-	EditableText->SetOnKeyCharHandler(InOnKeyCharHandler);
+	EditableText->SetOnKeyCharHandler(MoveTemp(InOnKeyCharHandler));
 }
 
 
 void SEditableTextBox::SetOnKeyDownHandler(FOnKeyDown InOnKeyDownHandler)
 {
-	EditableText->SetOnKeyDownHandler(InOnKeyDownHandler);
+	EditableText->SetOnKeyDownHandler(MoveTemp(InOnKeyDownHandler));
 }
 
 
@@ -208,7 +208,7 @@ void SEditableTextBox::SetTextFlowDirection(const TOptional<ETextFlowDirection>&
 }
 
 
-void SEditableTextBox::SetOverflowPolicy(TOptional<ETextOverflowPolicy> InOverflowPolicy)
+void SEditableTextBox::SetOverflowPolicy(const TOptional<ETextOverflowPolicy>& InOverflowPolicy)
 {
 	EditableText->SetOverflowPolicy(InOverflowPolicy);
 }
@@ -269,7 +269,7 @@ bool SEditableTextBox::HasError() const
 	return ErrorReporting.IsValid() && ErrorReporting->HasError();
 }
 
-const FSlateBrush* SEditableTextBox::GetBorderImage() const
+const FSlateBrush* SEditableTextBox::DetermineBorderImage() const
 {
 	if ( EditableText->IsTextReadOnly() )
 	{
@@ -378,126 +378,122 @@ FSlateColor SEditableTextBox::DetermineForegroundColor() const
 	}
 }
 
-void SEditableTextBox::SetHintText(const TAttribute< FText >& InHintText)
+void SEditableTextBox::SetHintText(TAttribute<FText> InHintText)
 {
-	EditableText->SetHintText(InHintText);
+	EditableText->SetHintText(MoveTemp(InHintText));
 }
 
-
-void SEditableTextBox::SetSearchText(const TAttribute<FText>& InSearchText)
+void SEditableTextBox::SetSearchText(TAttribute<FText> InSearchText)
 {
-	EditableText->SetSearchText(InSearchText);
+	EditableText->SetSearchText(MoveTemp(InSearchText));
 }
 
-
 FText SEditableTextBox::GetSearchText() const
 {
 	return EditableText->GetSearchText();
 }
 
-
-void SEditableTextBox::SetIsReadOnly(TAttribute< bool > InIsReadOnly)
+void SEditableTextBox::SetIsReadOnly(TAttribute<bool> InIsReadOnly)
 {
-	EditableText->SetIsReadOnly(InIsReadOnly);
+	EditableText->SetIsReadOnly(MoveTemp(InIsReadOnly));
 }
 
-
-void SEditableTextBox::SetIsPassword(TAttribute< bool > InIsPassword)
+void SEditableTextBox::SetIsPassword(TAttribute<bool> InIsPassword)
 {
-	EditableText->SetIsPassword(InIsPassword);
+	EditableText->SetIsPassword(MoveTemp(InIsPassword));
 }
 
-
-void SEditableTextBox::SetFont(const TAttribute<FSlateFontInfo>& InFont)
+void SEditableTextBox::SetFont(TAttribute<FSlateFontInfo> InFont)
 {
-	FontOverride = InFont;
+	FontOverride = MoveTemp(InFont);
 }
 
-void SEditableTextBox::SetTextBoxForegroundColor(const TAttribute<FSlateColor>& InForegroundColor)
+void SEditableTextBox::SetTextBoxForegroundColor(TAttribute<FSlateColor> InForegroundColor)
 {
-	ForegroundColorOverride = InForegroundColor;
+	ForegroundColorOverride = MoveTemp(InForegroundColor);
 }
 
-void SEditableTextBox::SetTextBoxBackgroundColor(const TAttribute<FSlateColor>& InBackgroundColor)
+void SEditableTextBox::SetTextBoxBackgroundColor(TAttribute<FSlateColor> InBackgroundColor)
 {
-	BackgroundColorOverride = InBackgroundColor;
+	BackgroundColorOverride = MoveTemp(InBackgroundColor);
 }
 
-void SEditableTextBox::SetReadOnlyForegroundColor(const TAttribute<FSlateColor>& InReadOnlyForegroundColor)
+void SEditableTextBox::SetReadOnlyForegroundColor(TAttribute<FSlateColor> InReadOnlyForegroundColor)
 {
-	ReadOnlyForegroundColorOverride = InReadOnlyForegroundColor;
+	ReadOnlyForegroundColorOverride = MoveTemp(InReadOnlyForegroundColor);
 }
 
-void SEditableTextBox::SetFocusedForegroundColor(const TAttribute<FSlateColor>& InFocusedForegroundColor)
+void SEditableTextBox::SetFocusedForegroundColor(TAttribute<FSlateColor> InFocusedForegroundColor)
 {
-	FocusedForegroundColorOverride = InFocusedForegroundColor;
+	FocusedForegroundColorOverride = MoveTemp(InFocusedForegroundColor);
 }
 
-void SEditableTextBox::SetMaximumLength(const TAttribute<int32>& InMaximumLength)
+void SEditableTextBox::SetMaximumLength(TAttribute<int32> InMaximumLength)
 {
-	MaximumLength = InMaximumLength;
+	MaximumLength = MoveTemp(InMaximumLength);
 }
 
-void SEditableTextBox::SetMinimumDesiredWidth(const TAttribute<float>& InMinimumDesiredWidth)
+void SEditableTextBox::SetMinimumDesiredWidth(TAttribute<float> InMinimumDesiredWidth)
 {
-	EditableText->SetMinDesiredWidth(InMinimumDesiredWidth);
+	EditableText->SetMinDesiredWidth(MoveTemp(InMinimumDesiredWidth));
 }
 
 
-void SEditableTextBox::SetIsCaretMovedWhenGainFocus(const TAttribute<bool>& InIsCaretMovedWhenGainFocus)
+void SEditableTextBox::SetIsCaretMovedWhenGainFocus(TAttribute<bool> InIsCaretMovedWhenGainFocus)
 {
-	EditableText->SetIsCaretMovedWhenGainFocus(InIsCaretMovedWhenGainFocus);
+	EditableText->SetIsCaretMovedWhenGainFocus(MoveTemp(InIsCaretMovedWhenGainFocus));
 }
 
 
-void SEditableTextBox::SetSelectAllTextWhenFocused(const TAttribute<bool>& InSelectAllTextWhenFocused)
+void SEditableTextBox::SetSelectAllTextWhenFocused(TAttribute<bool> InSelectAllTextWhenFocused)
 {
-	EditableText->SetSelectAllTextWhenFocused(InSelectAllTextWhenFocused);
+	EditableText->SetSelectAllTextWhenFocused(MoveTemp(InSelectAllTextWhenFocused));
 }
 
 
-void SEditableTextBox::SetRevertTextOnEscape(const TAttribute<bool>& InRevertTextOnEscape)
+void SEditableTextBox::SetRevertTextOnEscape(TAttribute<bool> InRevertTextOnEscape)
 {
-	EditableText->SetRevertTextOnEscape(InRevertTextOnEscape);
+	EditableText->SetRevertTextOnEscape(MoveTemp(InRevertTextOnEscape));
 }
 
 
-void SEditableTextBox::SetClearKeyboardFocusOnCommit(const TAttribute<bool>& InClearKeyboardFocusOnCommit)
+void SEditableTextBox::SetClearKeyboardFocusOnCommit(TAttribute<bool> InClearKeyboardFocusOnCommit)
 {
-	EditableText->SetClearKeyboardFocusOnCommit(InClearKeyboardFocusOnCommit);
+	EditableText->SetClearKeyboardFocusOnCommit(MoveTemp(InClearKeyboardFocusOnCommit));
 }
 
 
-void SEditableTextBox::SetSelectAllTextOnCommit(const TAttribute<bool>& InSelectAllTextOnCommit)
+void SEditableTextBox::SetSelectAllTextOnCommit(TAttribute<bool> InSelectAllTextOnCommit)
 {
-	EditableText->SetSelectAllTextOnCommit(InSelectAllTextOnCommit);
+	EditableText->SetSelectAllTextOnCommit(MoveTemp(InSelectAllTextOnCommit));
 }
 
-void SEditableTextBox::SetSelectWordOnMouseDoubleClick(const TAttribute<bool>& InSelectWordOnMouseDoubleClick)
+void SEditableTextBox::SetSelectWordOnMouseDoubleClick(TAttribute<bool> InSelectWordOnMouseDoubleClick)
 {
-	EditableText->SetSelectWordOnMouseDoubleClick(InSelectWordOnMouseDoubleClick);
+	EditableText->SetSelectWordOnMouseDoubleClick(MoveTemp(InSelectWordOnMouseDoubleClick));
 }
 
-void SEditableTextBox::SetJustification(const TAttribute<ETextJustify::Type>& InJustification)
+void SEditableTextBox::SetJustification(TAttribute<ETextJustify::Type> InJustification)
 {
-	EditableText->SetJustification(InJustification);
+	EditableText->SetJustification(MoveTemp(InJustification));
 }
 
 
 void SEditableTextBox::SetAllowContextMenu(TAttribute<bool> InAllowContextMenu)
 {
-	EditableText->SetAllowContextMenu(InAllowContextMenu);
+	EditableText->SetAllowContextMenu(MoveTemp(InAllowContextMenu));
 }
 
 void SEditableTextBox::SetVirtualKeyboardDismissAction(TAttribute<EVirtualKeyboardDismissAction> InVirtualKeyboardDismissAction)
 {
-	EditableText->SetVirtualKeyboardDismissAction(InVirtualKeyboardDismissAction);
+	EditableText->SetVirtualKeyboardDismissAction(MoveTemp(InVirtualKeyboardDismissAction));
 }
 
 void SEditableTextBox::EnableTextInputMethodContext()
 {
 	EditableText->EnableTextInputMethodContext();
 }
+
 #if WITH_ACCESSIBILITY
 TSharedRef<FSlateAccessibleWidget> SEditableTextBox::CreateAccessibleWidget()
 {
@@ -520,27 +516,23 @@ void SEditableTextBox::OnEditableTextChanged(const FText& InText)
 {
 	OnTextChanged.ExecuteIfBound(InText);
 
-	const int32 MaximumLengthValue = MaximumLength.Get();
-	if (OnVerifyTextChanged.IsBound() || MaximumLengthValue >= 0)
+	FText OutErrorMessage;
+	if (!FTextEditHelper::VerifyTextLength(InText, OutErrorMessage, MaximumLength.Get(0)) ||
+		(OnVerifyTextChanged.IsBound() && !OnVerifyTextChanged.Execute(InText, OutErrorMessage)))
 	{
-		FText OutErrorMessage;
-		if (!FTextEditHelper::VerifyTextLength(InText, OutErrorMessage, MaximumLengthValue) ||
-			(OnVerifyTextChanged.IsBound() && !OnVerifyTextChanged.Execute(InText, OutErrorMessage)))
-		{
-			// Display as an error.
-			SetError(OutErrorMessage);
-		}
-		else
-		{
-			SetError(FText::GetEmpty());
-		}
+		// Display as an error.
+		SetError(OutErrorMessage);
+	}
+	else
+	{
+		SetError(FText::GetEmpty());
 	}
 }
 
 void SEditableTextBox::OnEditableTextCommitted(const FText& InText, ETextCommit::Type InCommitType)
 {
 	FText OutErrorMessage;
-	if (!FTextEditHelper::VerifyTextLength(InText, OutErrorMessage, MaximumLength.Get()) ||
+	if (!FTextEditHelper::VerifyTextLength(InText, OutErrorMessage, MaximumLength.Get(0)) ||
 		(OnVerifyTextChanged.IsBound() && !OnVerifyTextChanged.Execute(InText, OutErrorMessage)))
 	{
 		// Display as an error.
diff --git a/Engine/Source/Runtime/Slate/Private/Widgets/Input/SMultiLineEditableTextBox.cpp b/Engine/Source/Runtime/Slate/Private/Widgets/Input/SMultiLineEditableTextBox.cpp
--- a/Engine/Source/Runtime/Slate/Private/Widgets/Input/SMultiLineEditableTextBox.cpp
+++ b/Engine/Source/Runtime/Slate/Private/Widgets/Input/SMultiLineEditableTextBox.cpp
@@ -31,7 +31,7 @@ namespace
  *
  * @param	InArgs	The declaration data for this widget
  */
-void SMultiLineEditableTextBox::Construct( const FArguments& InArgs )
+void SMultiLineEditableTextBox::Construct(const FArguments& InArgs)
 {
 	check (InArgs._Style);
 	SetStyle(InArgs._Style);
@@ -76,7 +76,7 @@ void SMultiLineEditableTextBox::Construct( const FArguments& InArgs )
 	}
 
 	SBorder::Construct( SBorder::FArguments()
-		.BorderImage( this, &SMultiLineEditableTextBox::GetBorderImage )
+		.BorderImage( this, &SMultiLineEditableTextBox::DetermineBorderImage )
 		.BorderBackgroundColor( this, &SMultiLineEditableTextBox::DetermineBackgroundColor )
 		.ForegroundColor( this, &SMultiLineEditableTextBox::DetermineForegroundColor )
 		.Padding( this, &SMultiLineEditableTextBox::DeterminePadding )
@@ -317,19 +317,19 @@ FSlateColor SMultiLineEditableTextBox::DetermineForegroundColor() const
 	}
 }
 
-void SMultiLineEditableTextBox::SetText( const TAttribute< FText >& InNewText )
+void SMultiLineEditableTextBox::SetText(TAttribute<FText> InNewText)
 {
-	EditableText->SetText( InNewText );
+	EditableText->SetText(MoveTemp(InNewText));
 }
 
-void SMultiLineEditableTextBox::SetHintText( const TAttribute< FText >& InHintText )
+void SMultiLineEditableTextBox::SetHintText(TAttribute<FText> InHintText)
 {
-	EditableText->SetHintText( InHintText );
+	EditableText->SetHintText(MoveTemp(InHintText));
 }
 
-void SMultiLineEditableTextBox::SetSearchText(const TAttribute<FText>& InSearchText)
+void SMultiLineEditableTextBox::SetSearchText(TAttribute<FText> InSearchText)
 {
-	EditableText->SetSearchText(InSearchText);
+	EditableText->SetSearchText(MoveTemp(InSearchText));
 }
 
 FText SMultiLineEditableTextBox::GetSearchText() const
@@ -337,29 +337,29 @@ FText SMultiLineEditableTextBox::GetSearchText() const
 	return EditableText->GetSearchText();
 }
 
-void SMultiLineEditableTextBox::SetTextBoxForegroundColor(const TAttribute<FSlateColor>& InForegroundColor)
+void SMultiLineEditableTextBox::SetTextBoxForegroundColor(TAttribute<FSlateColor> InForegroundColor)
 {
-	ForegroundColorOverride = InForegroundColor;
+	ForegroundColorOverride = MoveTemp(InForegroundColor);
 }
 
-void SMultiLineEditableTextBox::SetTextBoxBackgroundColor(const TAttribute<FSlateColor>& InBackgroundColor)
+void SMultiLineEditableTextBox::SetTextBoxBackgroundColor(TAttribute<FSlateColor> InBackgroundColor)
 {
-	BackgroundColorOverride = InBackgroundColor;
+	BackgroundColorOverride = MoveTemp(InBackgroundColor);
 }
 
-void SMultiLineEditableTextBox::SetReadOnlyForegroundColor(const TAttribute<FSlateColor>& InReadOnlyForegroundColor)
+void SMultiLineEditableTextBox::SetReadOnlyForegroundColor(TAttribute<FSlateColor> InReadOnlyForegroundColor)
 {
-	ReadOnlyForegroundColorOverride = InReadOnlyForegroundColor;
+	ReadOnlyForegroundColorOverride = MoveTemp(InReadOnlyForegroundColor);
 }
 
-void SMultiLineEditableTextBox::SetMaximumLength(const TAttribute<int32>& InMaximumLength)
+void SMultiLineEditableTextBox::SetMaximumLength(TAttribute<int32> InMaximumLength)
 {
-	MaximumLength = InMaximumLength;
+	MaximumLength = MoveTemp(InMaximumLength);
 }
 
-void SMultiLineEditableTextBox::SetSelectWordOnMouseDoubleClick(const TAttribute<bool>& InSelectWordOnMouseDoubleClick)
+void SMultiLineEditableTextBox::SetSelectWordOnMouseDoubleClick(TAttribute<bool> InSelectWordOnMouseDoubleClick)
 {
-	EditableText->SetSelectWordOnMouseDoubleClick(InSelectWordOnMouseDoubleClick);
+	EditableText->SetSelectWordOnMouseDoubleClick(MoveTemp(InSelectWordOnMouseDoubleClick));
 }
 
 void SMultiLineEditableTextBox::SetTextShapingMethod(const TOptional<ETextShapingMethod>& InTextShapingMethod)
@@ -372,59 +372,59 @@ void SMultiLineEditableTextBox::SetTextFlowDirection(const TOptional<ETextFlowDi
 	EditableText->SetTextFlowDirection(InTextFlowDirection);
 }
 
-void SMultiLineEditableTextBox::SetWrapTextAt(const TAttribute<float>& InWrapTextAt)
+void SMultiLineEditableTextBox::SetWrapTextAt(TAttribute<float> InWrapTextAt)
 {
-	EditableText->SetWrapTextAt(InWrapTextAt);
+	EditableText->SetWrapTextAt(MoveTemp(InWrapTextAt));
 }
 
-void SMultiLineEditableTextBox::SetAutoWrapText(const TAttribute<bool>& InAutoWrapText)
+void SMultiLineEditableTextBox::SetAutoWrapText(TAttribute<bool> InAutoWrapText)
 {
-	EditableText->SetAutoWrapText(InAutoWrapText);
+	EditableText->SetAutoWrapText(MoveTemp(InAutoWrapText));
 }
 
-void SMultiLineEditableTextBox::SetWrappingPolicy(const TAttribute<ETextWrappingPolicy>& InWrappingPolicy)
+void SMultiLineEditableTextBox::SetWrappingPolicy(TAttribute<ETextWrappingPolicy> InWrappingPolicy)
 {
-	EditableText->SetWrappingPolicy(InWrappingPolicy);
+	EditableText->SetWrappingPolicy(MoveTemp(InWrappingPolicy));
 }
 
-void SMultiLineEditableTextBox::SetLineHeightPercentage(const TAttribute<float>& InLineHeightPercentage)
+void SMultiLineEditableTextBox::SetLineHeightPercentage(TAttribute<float> InLineHeightPercentage)
 {
-	EditableText->SetLineHeightPercentage(InLineHeightPercentage);
+	EditableText->SetLineHeightPercentage(MoveTemp(InLineHeightPercentage));
 }
 
-void SMultiLineEditableTextBox::SetApplyLineHeightToBottomLine(const TAttribute<bool>& InApplyLineHeightToBottomLine)
+void SMultiLineEditableTextBox::SetApplyLineHeightToBottomLine(TAttribute<bool> InApplyLineHeightToBottomLine)
 {
-	EditableText->SetApplyLineHeightToBottomLine(InApplyLineHeightToBottomLine);
+	EditableText->SetApplyLineHeightToBottomLine(MoveTemp(InApplyLineHeightToBottomLine));
 }
 
-void SMultiLineEditableTextBox::SetMargin(const TAttribute<FMargin>& InMargin)
+void SMultiLineEditableTextBox::SetMargin(TAttribute<FMargin> InMargin)
 {
-	EditableText->SetMargin(InMargin);
+	EditableText->SetMargin(MoveTemp(InMargin));
 }
 
-void SMultiLineEditableTextBox::SetJustification(const TAttribute<ETextJustify::Type>& InJustification)
+void SMultiLineEditableTextBox::SetJustification(TAttribute<ETextJustify::Type> InJustification)
 {
-	EditableText->SetJustification(InJustification);
+	EditableText->SetJustification(MoveTemp(InJustification));
 }
 
-void SMultiLineEditableTextBox::SetOverflowPolicy(TOptional<ETextOverflowPolicy> InOverflowPolicy)
+void SMultiLineEditableTextBox::SetOverflowPolicy(const TOptional<ETextOverflowPolicy>& InOverflowPolicy)
 {
 	EditableText->SetOverflowPolicy(InOverflowPolicy);
 }
 
-void SMultiLineEditableTextBox::SetAllowContextMenu(const TAttribute< bool >& InAllowContextMenu)
+void SMultiLineEditableTextBox::SetAllowContextMenu(TAttribute<bool> InAllowContextMenu)
 {
-	EditableText->SetAllowContextMenu(InAllowContextMenu);
+	EditableText->SetAllowContextMenu(MoveTemp(InAllowContextMenu));
 }
 
 void SMultiLineEditableTextBox::SetVirtualKeyboardDismissAction(TAttribute<EVirtualKeyboardDismissAction> InVirtualKeyboardDismissAction)
 {
-	EditableText->SetVirtualKeyboardDismissAction(InVirtualKeyboardDismissAction);
+	EditableText->SetVirtualKeyboardDismissAction(MoveTemp(InVirtualKeyboardDismissAction));
 }
 
-void SMultiLineEditableTextBox::SetIsReadOnly(const TAttribute<bool>& InIsReadOnly)
+void SMultiLineEditableTextBox::SetIsReadOnly(TAttribute<bool> InIsReadOnly)
 {
-	EditableText->SetIsReadOnly(InIsReadOnly);
+	EditableText->SetIsReadOnly(MoveTemp(InIsReadOnly));
 }
 
 void SMultiLineEditableTextBox::SetFontFacesLoadingPaintPolicy( EFontFacesLoadingPaintPolicy InFontFacesLoadingPaintPolicy)
@@ -435,12 +435,12 @@ void SMultiLineEditableTextBox::SetFontFacesLoadingPaintPolicy( EFontFacesLoadin
 	}
 }
 
-void SMultiLineEditableTextBox::SetError( const FText& InError )
+void SMultiLineEditableTextBox::SetError(const FText& InError)
 {
 	SetError( InError.ToString() );
 }
 
-void SMultiLineEditableTextBox::SetError( const FString& InError )
+void SMultiLineEditableTextBox::SetError(const FString& InError)
 {
 	const bool bHaveError = !InError.IsEmpty();
 
@@ -466,7 +466,7 @@ bool SMultiLineEditableTextBox::HasError() const
 }
 
 /** @return Border image for the text box based on the hovered and focused state */
-const FSlateBrush* SMultiLineEditableTextBox::GetBorderImage() const
+const FSlateBrush* SMultiLineEditableTextBox::DetermineBorderImage() const
 {
 	if ( EditableText->IsTextReadOnly() )
 	{
@@ -558,7 +558,7 @@ void SMultiLineEditableTextBox::InsertTextAtCursor(const FString& InString)
 
 void SMultiLineEditableTextBox::InsertRunAtCursor(TSharedRef<IRun> InRun)
 {
-	EditableText->InsertRunAtCursor(InRun);
+	EditableText->InsertRunAtCursor(MoveTemp(InRun));
 }
 
 void SMultiLineEditableTextBox::GoTo(const FTextLocation& NewLocation)
@@ -618,12 +618,12 @@ void SMultiLineEditableTextBox::Refresh()
 
 void SMultiLineEditableTextBox::SetOnKeyCharHandler(FOnKeyChar InOnKeyCharHandler)
 {
-	EditableText->SetOnKeyCharHandler(InOnKeyCharHandler);
+	EditableText->SetOnKeyCharHandler(MoveTemp(InOnKeyCharHandler));
 }
 
 void SMultiLineEditableTextBox::SetOnKeyDownHandler(FOnKeyDown InOnKeyDownHandler)
 {
-	EditableText->SetOnKeyDownHandler(InOnKeyDownHandler);
+	EditableText->SetOnKeyDownHandler(MoveTemp(InOnKeyDownHandler));
 }
 
 void SMultiLineEditableTextBox::ForceScroll(int32 UserIndex, float ScrollAxisMagnitude)
diff --git a/Engine/Source/Runtime/Slate/Private/Widgets/Text/SMultiLineEditableText.cpp b/Engine/Source/Runtime/Slate/Private/Widgets/Text/SMultiLineEditableText.cpp
--- a/Engine/Source/Runtime/Slate/Private/Widgets/Text/SMultiLineEditableText.cpp
+++ b/Engine/Source/Runtime/Slate/Private/Widgets/Text/SMultiLineEditableText.cpp
@@ -10,22 +10,19 @@
 #include "Framework/Text/TextEditHelper.h"
 #include "Framework/Text/PlainTextLayoutMarshaller.h"
 #include "Widgets/Text/SlateEditableTextLayout.h"
+#include "Widgets/Text/SlateTextBlockLayout.h"
 #include "Types/ReflectionMetadata.h"
 #include "Types/TrackedMetaData.h"
 
 SMultiLineEditableText::SMultiLineEditableText()
-	: bSelectAllTextWhenFocused(false)
+	: SlateEditableTextLayout(*this)
+	, bSelectAllTextWhenFocused(false)
 	, bIsReadOnly(false)
 	, AmountScrolledWhileRightMouseDown(0.0f)
 	, bIsSoftwareCursor(false)
 {
 }
 
-SMultiLineEditableText::~SMultiLineEditableText()
-{
-	// Needed to avoid "deletion of pointer to incomplete type 'FSlateEditableTextLayout'; no destructor called" error when using TUniquePtr
-}
-
 void SMultiLineEditableText::Construct( const FArguments& InArgs )
 {
 	bIsReadOnly = InArgs._IsReadOnly;
@@ -53,7 +50,7 @@ void SMultiLineEditableText::Construct( const FArguments& InArgs )
 
 	HScrollBar = InArgs._HScrollBar;
 
-	// Horizontal scroll is unecessary when auto wrap is enabled.
+	// Horizontal scroll is unnecessary when auto wrap is enabled.
 	// Otherwise if allowed, having both enabled will cause a wrap+scroll relayout feedback loop which can result in jittering
 	const bool bAutoWrapEnabled = InArgs._AutoWrapText.IsSet() && InArgs._AutoWrapText.Get();
 
@@ -82,19 +79,19 @@ void SMultiLineEditableText::Construct( const FArguments& InArgs )
 		Marshaller = FPlainTextLayoutMarshaller::Create();
 	}
 
-	EditableTextLayout = MakeUnique<FSlateEditableTextLayout>(*this, InArgs._Text, TextStyle, InArgs._TextShapingMethod, InArgs._TextFlowDirection, InArgs._CreateSlateTextLayout, Marshaller.ToSharedRef(), Marshaller.ToSharedRef());
-	EditableTextLayout->SetHintText(InArgs._HintText);
-	EditableTextLayout->SetSearchText(InArgs._SearchText);
-	EditableTextLayout->SetTextWrapping(InArgs._WrapTextAt, InArgs._AutoWrapText, InArgs._WrappingPolicy);
-	EditableTextLayout->SetMargin(InArgs._Margin);
-	EditableTextLayout->SetJustification(InArgs._Justification);
-	EditableTextLayout->SetLineHeightPercentage(InArgs._LineHeightPercentage);
-	EditableTextLayout->SetDebugSourceInfo(TAttribute<FString>::Create(TAttribute<FString>::FGetter::CreateLambda([this]{ return FReflectionMetaData::GetWidgetDebugInfo(this); })));
-	EditableTextLayout->SetOverflowPolicy(InArgs._OverflowPolicy);
-	EditableTextLayout->SetFontFacesLoadingPaintPolicy(InArgs._FontFacesLoadingPaintPolicy);
+	SlateEditableTextLayout.Construct(InArgs._Text, TextStyle, InArgs._TextShapingMethod, InArgs._TextFlowDirection, InArgs._CreateSlateTextLayout, Marshaller.ToSharedRef(), Marshaller.ToSharedRef());
+	SlateEditableTextLayout.SetHintText(InArgs._HintText);
+	SlateEditableTextLayout.SetSearchText(InArgs._SearchText);
+	SlateEditableTextLayout.SetTextWrapping(InArgs._WrapTextAt, InArgs._AutoWrapText, InArgs._WrappingPolicy);
+	SlateEditableTextLayout.SetMargin(InArgs._Margin);
+	SlateEditableTextLayout.SetJustification(InArgs._Justification);
+	SlateEditableTextLayout.SetLineHeightPercentage(InArgs._LineHeightPercentage);
+	SlateEditableTextLayout.SetDebugSourceInfo(TAttribute<FString>::Create(TAttribute<FString>::FGetter::CreateLambda([this]{ return FReflectionMetaData::GetWidgetDebugInfo(this); })));
+	SlateEditableTextLayout.SetOverflowPolicy(InArgs._OverflowPolicy);
+	SlateEditableTextLayout.SetFontFacesLoadingPaintPolicy(InArgs._FontFacesLoadingPaintPolicy);
 	if (InArgs._OnAllFontFacesFinishLoading.IsBound())
 	{
-		EditableTextLayout->SetOnAllFontFacesFinishLoadingDelegate(InArgs._OnAllFontFacesFinishLoading);
+		SlateEditableTextLayout.SetOnAllFontFacesFinishLoadingDelegate(InArgs._OnAllFontFacesFinishLoading);
 	}
 
 	// build context menu extender
@@ -106,198 +103,195 @@ void SMultiLineEditableText::Construct( const FArguments& InArgs )
 
 void SMultiLineEditableText::GetCurrentTextLine(FString& OutTextLine) const
 {
-	EditableTextLayout->GetCurrentTextLine(OutTextLine);
+	SlateEditableTextLayout.GetCurrentTextLine(OutTextLine);
 }
 
 void SMultiLineEditableText::GetTextLine(const int32 InLineIndex, FString& OutTextLine) const
 {
-	EditableTextLayout->GetTextLine(InLineIndex, OutTextLine);
+	SlateEditableTextLayout.GetTextLine(InLineIndex, OutTextLine);
 }
 
-void SMultiLineEditableText::SetText(const TAttribute< FText >& InText)
+void SMultiLineEditableText::SetText(TAttribute<FText> InText)
 {
-	EditableTextLayout->SetText(InText);
+	SlateEditableTextLayout.SetText(MoveTemp(InText));
 }
 
-int32 SMultiLineEditableText::GetTextLineCount()
+int32 SMultiLineEditableText::GetTextLineCount() const
 {
-	return EditableTextLayout->GetTextLineCount();
+	return SlateEditableTextLayout.GetTextLineCount();
 }
 
 FText SMultiLineEditableText::GetText() const
 {
-	return EditableTextLayout->GetText();
+	return SlateEditableTextLayout.GetText();
 }
 
 FText SMultiLineEditableText::GetPlainText() const
 {
-	return EditableTextLayout->GetPlainText();
+	return SlateEditableTextLayout.GetPlainText();
 }
 
-void SMultiLineEditableText::SetHintText(const TAttribute< FText >& InHintText)
+void SMultiLineEditableText::SetHintText(TAttribute<FText> InHintText)
 {
-	EditableTextLayout->SetHintText(InHintText);
+	SlateEditableTextLayout.SetHintText(MoveTemp(InHintText));
 }
 
 FText SMultiLineEditableText::GetHintText() const
 {
-	return EditableTextLayout->GetHintText();
+	return SlateEditableTextLayout.GetHintText();
 }
 
-void SMultiLineEditableText::SetSearchText(const TAttribute<FText>& InSearchText)
+void SMultiLineEditableText::SetSearchText(TAttribute<FText> InSearchText)
 {
-	EditableTextLayout->SetSearchText(InSearchText);
+	SlateEditableTextLayout.SetSearchText(MoveTemp(InSearchText));
 }
 
 FText SMultiLineEditableText::GetSearchText() const
 {
-	return EditableTextLayout->GetSearchText();
+	return SlateEditableTextLayout.GetSearchText();
 }
 
 int32 SMultiLineEditableText::GetSearchResultIndex() const
 {
-	return EditableTextLayout->GetSearchResultIndex();
+	return SlateEditableTextLayout.GetSearchResultIndex();
 }
 
 int32 SMultiLineEditableText::GetNumSearchResults() const
 {
-	return EditableTextLayout->GetNumSearchResults();
+	return SlateEditableTextLayout.GetNumSearchResults();
 }
 
 void SMultiLineEditableText::SetTextStyle(const FTextBlockStyle* InTextStyle)
 {
 	if (InTextStyle)
 	{
-		EditableTextLayout->SetTextStyle(*InTextStyle);
+		SlateEditableTextLayout.SetTextStyle(*InTextStyle);
 	}
 	else
 	{
-		FArguments Defaults;
+		const FArguments Defaults;
 		check(Defaults._TextStyle);
-		EditableTextLayout->SetTextStyle(*Defaults._TextStyle);
+		SlateEditableTextLayout.SetTextStyle(*Defaults._TextStyle);
 	}
 }
 
-void SMultiLineEditableText::SetFont(const TAttribute< FSlateFontInfo >& InNewFont)
+void SMultiLineEditableText::SetFont(TAttribute<FSlateFontInfo> InNewFont)
 {
-	FTextBlockStyle TextStyle = EditableTextLayout->GetTextStyle();
-	TextStyle.SetFont(InNewFont.Get());
-	EditableTextLayout->SetTextStyle(TextStyle);
+	FTextBlockStyle TextStyle = SlateEditableTextLayout.GetTextStyle();
+	TextStyle.SetFont(InNewFont.Get(TextStyle.Font));
+	SlateEditableTextLayout.SetTextStyle(TextStyle);
 }
 
 FSlateFontInfo SMultiLineEditableText::GetFont() const
 {
-	FTextBlockStyle TextStyle = EditableTextLayout->GetTextStyle();
+	FTextBlockStyle TextStyle = SlateEditableTextLayout.GetTextStyle();
 	return TextStyle.Font;
 }
 
 void SMultiLineEditableText::SetTextShapingMethod(const TOptional<ETextShapingMethod>& InTextShapingMethod)
 {
-	EditableTextLayout->SetTextShapingMethod(InTextShapingMethod);
+	SlateEditableTextLayout.SetTextShapingMethod(InTextShapingMethod);
 }
 
 void SMultiLineEditableText::SetTextFlowDirection(const TOptional<ETextFlowDirection>& InTextFlowDirection)
 {
-	EditableTextLayout->SetTextFlowDirection(InTextFlowDirection);
+	SlateEditableTextLayout.SetTextFlowDirection(InTextFlowDirection);
 }
 
-void SMultiLineEditableText::SetWrapTextAt(const TAttribute<float>& InWrapTextAt)
+void SMultiLineEditableText::SetWrapTextAt(TAttribute<float> InWrapTextAt)
 {
-	EditableTextLayout->SetWrapTextAt(InWrapTextAt);
+	SlateEditableTextLayout.SetWrapTextAt(MoveTemp(InWrapTextAt));
 }
 
-void SMultiLineEditableText::SetAutoWrapText(const TAttribute<bool>& InAutoWrapText)
+void SMultiLineEditableText::SetAutoWrapText(TAttribute<bool> InAutoWrapText)
 {
-	EditableTextLayout->SetAutoWrapText(InAutoWrapText);
+	SlateEditableTextLayout.SetAutoWrapText(MoveTemp(InAutoWrapText));
 }
 
-void SMultiLineEditableText::SetWrappingPolicy(const TAttribute<ETextWrappingPolicy>& InWrappingPolicy)
+void SMultiLineEditableText::SetWrappingPolicy(TAttribute<ETextWrappingPolicy> InWrappingPolicy)
 {
-	EditableTextLayout->SetWrappingPolicy(InWrappingPolicy);
+	SlateEditableTextLayout.SetWrappingPolicy(MoveTemp(InWrappingPolicy));
 }
 
-void SMultiLineEditableText::SetLineHeightPercentage(const TAttribute<float>& InLineHeightPercentage)
+void SMultiLineEditableText::SetLineHeightPercentage(TAttribute<float> InLineHeightPercentage)
 {
-	EditableTextLayout->SetLineHeightPercentage(InLineHeightPercentage);
+	SlateEditableTextLayout.SetLineHeightPercentage(MoveTemp(InLineHeightPercentage));
 }
 
-void SMultiLineEditableText::SetApplyLineHeightToBottomLine(const TAttribute<bool>& InApplyLineHeightToBottomLine)
+void SMultiLineEditableText::SetApplyLineHeightToBottomLine(TAttribute<bool> InApplyLineHeightToBottomLine)
 {
-	EditableTextLayout->SetApplyLineHeightToBottomLine(InApplyLineHeightToBottomLine);
+	SlateEditableTextLayout.SetApplyLineHeightToBottomLine(MoveTemp(InApplyLineHeightToBottomLine));
 }
 
-void SMultiLineEditableText::SetMargin(const TAttribute<FMargin>& InMargin)
+void SMultiLineEditableText::SetMargin(TAttribute<FMargin> InMargin)
 {
-	EditableTextLayout->SetMargin(InMargin);
+	SlateEditableTextLayout.SetMargin(MoveTemp(InMargin));
 }
 
-void SMultiLineEditableText::SetJustification(const TAttribute<ETextJustify::Type>& InJustification)
+void SMultiLineEditableText::SetJustification(TAttribute<ETextJustify::Type> InJustification)
 {
-	EditableTextLayout->SetJustification(InJustification);
+	SlateEditableTextLayout.SetJustification(MoveTemp(InJustification));
 }
 
-void SMultiLineEditableText::SetOverflowPolicy(TOptional<ETextOverflowPolicy> InOverflowPolicy)
+void SMultiLineEditableText::SetOverflowPolicy(const TOptional<ETextOverflowPolicy>& InOverflowPolicy)
 {
-	EditableTextLayout->SetOverflowPolicy(InOverflowPolicy);
+	SlateEditableTextLayout.SetOverflowPolicy(InOverflowPolicy);
 }
 
-void SMultiLineEditableText::SetAllowContextMenu(const TAttribute< bool >& InAllowContextMenu)
+void SMultiLineEditableText::SetAllowContextMenu(TAttribute<bool> InAllowContextMenu)
 {
-	bAllowContextMenu = InAllowContextMenu;
+	bAllowContextMenu = MoveTemp(InAllowContextMenu);
 }
 
-void SMultiLineEditableText::SetVirtualKeyboardDismissAction(TAttribute< EVirtualKeyboardDismissAction > InVirtualKeyboardDismissAction)
+void SMultiLineEditableText::SetVirtualKeyboardDismissAction(TAttribute<EVirtualKeyboardDismissAction> InVirtualKeyboardDismissAction)
 {
-	VirtualKeyboardDismissAction = InVirtualKeyboardDismissAction;
+	VirtualKeyboardDismissAction = MoveTemp(InVirtualKeyboardDismissAction);
 }
 
-void SMultiLineEditableText::SetIsReadOnly(const TAttribute<bool>& InIsReadOnly)
+void SMultiLineEditableText::SetIsReadOnly(TAttribute<bool> InIsReadOnly)
 {
-	bIsReadOnly = InIsReadOnly;
+	bIsReadOnly = MoveTemp(InIsReadOnly);
 }
 
-void SMultiLineEditableText::SetSelectAllTextWhenFocused(const TAttribute<bool>& InSelectAllTextWhenFocused)
+void SMultiLineEditableText::SetSelectAllTextWhenFocused(TAttribute<bool> InSelectAllTextWhenFocused)
 {
-	bSelectAllTextWhenFocused = InSelectAllTextWhenFocused;
+	bSelectAllTextWhenFocused = MoveTemp(InSelectAllTextWhenFocused);
 }
 
-void SMultiLineEditableText::SetSelectWordOnMouseDoubleClick(const TAttribute<bool>& InSelectWordOnMouseDoubleClick)
+void SMultiLineEditableText::SetSelectWordOnMouseDoubleClick(TAttribute<bool> InSelectWordOnMouseDoubleClick)
 {
-	bSelectWordOnMouseDoubleClick = InSelectWordOnMouseDoubleClick;
+	bSelectWordOnMouseDoubleClick = MoveTemp(InSelectWordOnMouseDoubleClick);
 }
 
-void SMultiLineEditableText::SetClearTextSelectionOnFocusLoss(const TAttribute<bool>& InClearTextSelectionOnFocusLoss)
+void SMultiLineEditableText::SetClearTextSelectionOnFocusLoss(TAttribute<bool> InClearTextSelectionOnFocusLoss)
 {
-	bClearTextSelectionOnFocusLoss = InClearTextSelectionOnFocusLoss;
+	bClearTextSelectionOnFocusLoss = MoveTemp(InClearTextSelectionOnFocusLoss);
 }
 
-void SMultiLineEditableText::SetRevertTextOnEscape(const TAttribute<bool>& InRevertTextOnEscape)
+void SMultiLineEditableText::SetRevertTextOnEscape(TAttribute<bool> InRevertTextOnEscape)
 {
-	bRevertTextOnEscape = InRevertTextOnEscape;
+	bRevertTextOnEscape = MoveTemp(InRevertTextOnEscape);
 }
 
-void SMultiLineEditableText::SetClearKeyboardFocusOnCommit(const TAttribute<bool>& InClearKeyboardFocusOnCommit)
+void SMultiLineEditableText::SetClearKeyboardFocusOnCommit(TAttribute<bool> InClearKeyboardFocusOnCommit)
 {
-	bClearKeyboardFocusOnCommit = InClearKeyboardFocusOnCommit;
+	bClearKeyboardFocusOnCommit = MoveTemp(InClearKeyboardFocusOnCommit);
 }
 
 void SMultiLineEditableText::SetFontFacesLoadingPaintPolicy(EFontFacesLoadingPaintPolicy InFontFacesLoadingPaintPolicy)
 {
-	if (EditableTextLayout)
-	{
-		EditableTextLayout->SetFontFacesLoadingPaintPolicy(InFontFacesLoadingPaintPolicy);
-	}
+	SlateEditableTextLayout.SetFontFacesLoadingPaintPolicy(InFontFacesLoadingPaintPolicy);
 }
 
 void SMultiLineEditableText::OnHScrollBarMoved(const float InScrollOffsetFraction)
 {
-	EditableTextLayout->SetHorizontalScrollFraction(InScrollOffsetFraction);
+	SlateEditableTextLayout.SetHorizontalScrollFraction(InScrollOffsetFraction);
 	OnHScrollBarUserScrolled.ExecuteIfBound(InScrollOffsetFraction);
 }
 
 void SMultiLineEditableText::OnVScrollBarMoved(const float InScrollOffsetFraction)
 {
-	EditableTextLayout->SetVerticalScrollFraction(InScrollOffsetFraction);
+	SlateEditableTextLayout.SetVerticalScrollFraction(InScrollOffsetFraction);
 	OnVScrollBarUserScrolled.ExecuteIfBound(InScrollOffsetFraction);
 }
 
@@ -382,7 +376,7 @@ void SMultiLineEditableText::EnsureActiveTick()
 	auto DoActiveTick = [this](double InCurrentTime, float InDeltaTime) -> EActiveTimerReturnType
 	{
 		// Continue if we still have focus, otherwise treat as a fire-and-forget Tick() request
-		const bool bShouldAppearFocused = HasKeyboardFocus() || EditableTextLayout->HasActiveContextMenu();
+		const bool bShouldAppearFocused = HasKeyboardFocus() || SlateEditableTextLayout.HasActiveContextMenu();
 		return (bShouldAppearFocused) ? EActiveTimerReturnType::Continue : EActiveTimerReturnType::Stop;
 	};
 
@@ -436,7 +430,7 @@ TSharedPtr<SWidget> SMultiLineEditableText::BuildContextMenuContent() const
 		return OnContextMenuOpening.Execute();
 	}
 
-	return EditableTextLayout->BuildDefaultContextMenu(MenuExtender);
+	return SlateEditableTextLayout.BuildDefaultContextMenu(MenuExtender);
 }
 
 void SMultiLineEditableText::OnBeginTextEdit(const FText& InText)
@@ -460,12 +454,12 @@ void SMultiLineEditableText::OnCursorMoved(const FTextLocation& InLocation)
 	Invalidate(EInvalidateWidget::Layout);
 }
 
-float SMultiLineEditableText::UpdateAndClampHorizontalScrollBar(const float InViewOffset, const float InViewFraction, const EVisibility InVisiblityOverride)
+float SMultiLineEditableText::UpdateAndClampHorizontalScrollBar(const float InViewOffset, const float InViewFraction, const EVisibility InVisibilityOverride)
 {
 	if (HScrollBar.IsValid())
 	{
 		HScrollBar->SetState(InViewOffset, InViewFraction);
-		HScrollBar->SetUserVisibility(InVisiblityOverride);
+		HScrollBar->SetUserVisibility(InVisibilityOverride);
 		if (!HScrollBar->IsNeeded())
 		{
 			// We cannot scroll, so ensure that there is no offset
@@ -473,15 +467,15 @@ float SMultiLineEditableText::UpdateAndClampHorizontalScrollBar(const float InVi
 		}
 	}
 
-	return EditableTextLayout->GetScrollOffset().X;
+	return SlateEditableTextLayout.GetScrollOffset().X;
 }
 
-float SMultiLineEditableText::UpdateAndClampVerticalScrollBar(const float InViewOffset, const float InViewFraction, const EVisibility InVisiblityOverride)
+float SMultiLineEditableText::UpdateAndClampVerticalScrollBar(const float InViewOffset, const float InViewFraction, const EVisibility InVisibilityOverride)
 {
 	if (VScrollBar.IsValid())
 	{
 		VScrollBar->SetState(InViewOffset, InViewFraction);
-		VScrollBar->SetUserVisibility(InVisiblityOverride);
+		VScrollBar->SetUserVisibility(InVisibilityOverride);
 		if (!VScrollBar->IsNeeded())
 		{
 			// We cannot scroll, so ensure that there is no offset
@@ -489,134 +483,134 @@ float SMultiLineEditableText::UpdateAndClampVerticalScrollBar(const float InView
 		}
 	}
 
-	return EditableTextLayout->GetScrollOffset().Y;
+	return SlateEditableTextLayout.GetScrollOffset().Y;
 }
 
 void SMultiLineEditableText::BeginEditTransaction()
 {
-	EditableTextLayout->BeginEditTransation();
+	SlateEditableTextLayout.BeginEditTransaction();
 }
 
 void SMultiLineEditableText::EndEditTransaction()
 {
-	EditableTextLayout->EndEditTransaction();
+	SlateEditableTextLayout.EndEditTransaction();
 }
 
 FReply SMultiLineEditableText::OnFocusReceived( const FGeometry& MyGeometry, const FFocusEvent& InFocusEvent )
 {
-	EditableTextLayout->HandleFocusReceived(InFocusEvent);
+	SlateEditableTextLayout.HandleFocusReceived(InFocusEvent);
 	return FReply::Handled();
 }
 
 void SMultiLineEditableText::OnFocusLost( const FFocusEvent& InFocusEvent )
 {
 	bIsSoftwareCursor = false;
-	EditableTextLayout->HandleFocusLost(InFocusEvent);
+	SlateEditableTextLayout.HandleFocusLost(InFocusEvent);
 }
 
 bool SMultiLineEditableText::AnyTextSelected() const
 {
-	return EditableTextLayout->AnyTextSelected();
+	return SlateEditableTextLayout.AnyTextSelected();
 }
 
 void SMultiLineEditableText::SelectAllText()
 {
-	EditableTextLayout->SelectAllText();
+	SlateEditableTextLayout.SelectAllText();
 }
 
 void SMultiLineEditableText::SelectText(const FTextLocation& InSelectionStart, const FTextLocation& InCursorLocation)
 {
-	EditableTextLayout->SelectText(InSelectionStart, InCursorLocation);
+	SlateEditableTextLayout.SelectText(InSelectionStart, InCursorLocation);
 }
 
 void SMultiLineEditableText::ClearSelection()
 {
-	EditableTextLayout->ClearSelection();
+	SlateEditableTextLayout.ClearSelection();
 }
 
 FText SMultiLineEditableText::GetSelectedText() const
 {
-	return EditableTextLayout->GetSelectedText();
+	return SlateEditableTextLayout.GetSelectedText();
 }
 
 FTextSelection SMultiLineEditableText::GetSelection() const
 {
-	return EditableTextLayout->GetSelection();
+	return SlateEditableTextLayout.GetSelection();
 }
 
 void SMultiLineEditableText::DeleteSelectedText()
 {
-	EditableTextLayout->DeleteSelectedText();
+	SlateEditableTextLayout.DeleteSelectedText();
 }
 
 void SMultiLineEditableText::InsertTextAtCursor(const FText& InText)
 {
-	EditableTextLayout->InsertTextAtCursor(InText.ToString());
+	SlateEditableTextLayout.InsertTextAtCursor(InText.ToString());
 }
 
 void SMultiLineEditableText::InsertTextAtCursor(const FString& InString)
 {
-	EditableTextLayout->InsertTextAtCursor(InString);
+	SlateEditableTextLayout.InsertTextAtCursor(InString);
 }
 
 void SMultiLineEditableText::InsertRunAtCursor(TSharedRef<IRun> InRun)
 {
-	EditableTextLayout->InsertRunAtCursor(InRun);
+	SlateEditableTextLayout.InsertRunAtCursor(MoveTemp(InRun));
 }
 
 void SMultiLineEditableText::GoTo(const FTextLocation& NewLocation)
 {
-	EditableTextLayout->GoTo(NewLocation);
+	SlateEditableTextLayout.GoTo(NewLocation);
 }
 
 void SMultiLineEditableText::GoTo(const ETextLocation NewLocation)
 {
-	EditableTextLayout->GoTo(NewLocation);
+	SlateEditableTextLayout.GoTo(NewLocation);
 }
 
 void SMultiLineEditableText::ScrollTo(const FTextLocation& NewLocation)
 {
-	EditableTextLayout->ScrollTo(NewLocation);
+	SlateEditableTextLayout.ScrollTo(NewLocation);
 }
 
 void SMultiLineEditableText::ScrollTo(const ETextLocation NewLocation)
 {
-	EditableTextLayout->ScrollTo(NewLocation);
+	SlateEditableTextLayout.ScrollTo(NewLocation);
 }
 
 void SMultiLineEditableText::ApplyToSelection(const FRunInfo& InRunInfo, const FTextBlockStyle& InStyle)
 {
-	EditableTextLayout->ApplyToSelection(InRunInfo, InStyle);
+	SlateEditableTextLayout.ApplyToSelection(InRunInfo, InStyle);
 }
 
 void SMultiLineEditableText::BeginSearch(const FText& InSearchText, const ESearchCase::Type InSearchCase, const bool InReverse)
 {
-	EditableTextLayout->BeginSearch(InSearchText, InSearchCase, InReverse);
+	SlateEditableTextLayout.BeginSearch(InSearchText, InSearchCase, InReverse);
 }
 
 void SMultiLineEditableText::AdvanceSearch(const bool InReverse)
 {
-	EditableTextLayout->AdvanceSearch(InReverse);
+	SlateEditableTextLayout.AdvanceSearch(InReverse);
 }
 
 TSharedPtr<const IRun> SMultiLineEditableText::GetRunUnderCursor() const
 {
-	return EditableTextLayout->GetRunUnderCursor();
+	return SlateEditableTextLayout.GetRunUnderCursor();
 }
 
 TArray<TSharedRef<const IRun>> SMultiLineEditableText::GetSelectedRuns() const
 {
-	return EditableTextLayout->GetSelectedRuns();
+	return SlateEditableTextLayout.GetSelectedRuns();
 }
 
 FTextLocation SMultiLineEditableText::GetCursorLocation() const
 {
-	return EditableTextLayout->GetCursorLocation();
+	return SlateEditableTextLayout.GetCursorLocation();
 }
 
 TCHAR SMultiLineEditableText::GetCharacterAt(const FTextLocation& Location) const
 {
-	return EditableTextLayout->GetCharacterAt(Location);
+	return SlateEditableTextLayout.GetCharacterAt(Location);
 }
 
 TSharedPtr<const SScrollBar> SMultiLineEditableText::GetHScrollBar() const
@@ -631,7 +625,7 @@ TSharedPtr<const SScrollBar> SMultiLineEditableText::GetVScrollBar() const
 
 void SMultiLineEditableText::Refresh()
 {
-	EditableTextLayout->Refresh();
+	SlateEditableTextLayout.Refresh();
 }
 
 void SMultiLineEditableText::ForceScroll(int32 UserIndex, float ScrollAxisMagnitude)
@@ -645,25 +639,25 @@ void SMultiLineEditableText::ForceScroll(int32 UserIndex, float ScrollAxisMagnit
 
 void SMultiLineEditableText::Tick( const FGeometry& AllottedGeometry, const double InCurrentTime, const float InDeltaTime )
 {
-	EditableTextLayout->Tick(AllottedGeometry, InCurrentTime, InDeltaTime);
+	SlateEditableTextLayout.Tick(AllottedGeometry, InCurrentTime, InDeltaTime);
 }
 
 int32 SMultiLineEditableText::OnPaint( const FPaintArgs& Args, const FGeometry& AllottedGeometry, const FSlateRect& MyCullingRect, FSlateWindowElementList& OutDrawElements, int32 LayerId, const FWidgetStyle& InWidgetStyle, bool bParentEnabled ) const
 {
-	const FTextBlockStyle& EditableTextStyle = EditableTextLayout->GetTextStyle();
+	const FTextBlockStyle& EditableTextStyle = SlateEditableTextLayout.GetTextStyle();
 	const FLinearColor ForegroundColor = EditableTextStyle.ColorAndOpacity.GetColor(InWidgetStyle);
 
 	FWidgetStyle TextWidgetStyle = FWidgetStyle(InWidgetStyle)
 		.SetForegroundColor(ForegroundColor);
 
-	const bool bAutoWrap = EditableTextLayout->GetAutoWrapText();
-	if (bAutoWrap && EditableTextLayout->GetComputedWrappingWidth() == 0)
+	const bool bAutoWrap = SlateEditableTextLayout.GetAutoWrapText();
+	if (bAutoWrap && SlateEditableTextLayout.GetComputedWrappingWidth() == 0)
 	{
 		QUICK_SCOPE_CYCLE_COUNTER(STAT_Slate_AutoWrapRecompute)
 		const_cast<SMultiLineEditableText*>(this)->CacheDesiredSize(GetPrepassLayoutScaleMultiplier());
 	}
 
-	LayerId = EditableTextLayout->OnPaint(Args, AllottedGeometry, MyCullingRect, OutDrawElements, LayerId, TextWidgetStyle, ShouldBeEnabled(bParentEnabled));
+	LayerId = SlateEditableTextLayout.OnPaint(Args, AllottedGeometry, MyCullingRect, OutDrawElements, LayerId, TextWidgetStyle, ShouldBeEnabled(bParentEnabled));
 
 	if (bIsSoftwareCursor)
 	{
@@ -683,23 +677,23 @@ int32 SMultiLineEditableText::OnPaint( const FPaintArgs& Args, const FGeometry&
 
 void SMultiLineEditableText::CacheDesiredSize(float LayoutScaleMultiplier)
 {
-	EditableTextLayout->CacheDesiredSize(LayoutScaleMultiplier);
+	SlateEditableTextLayout.CacheDesiredSize(LayoutScaleMultiplier);
 	SWidget::CacheDesiredSize(LayoutScaleMultiplier);
 }
 
 FVector2D SMultiLineEditableText::ComputeDesiredSize( float LayoutScaleMultiplier ) const
 {
-	return EditableTextLayout->ComputeDesiredSize(LayoutScaleMultiplier);
+	return SlateEditableTextLayout.ComputeDesiredSize(LayoutScaleMultiplier);
 }
 
 FChildren* SMultiLineEditableText::GetChildren()
 {
-	return EditableTextLayout->GetChildren();
+	return SlateEditableTextLayout.GetChildren();
 }
 
 void SMultiLineEditableText::OnArrangeChildren( const FGeometry& AllottedGeometry, FArrangedChildren& ArrangedChildren ) const
 {
-	EditableTextLayout->OnArrangeChildren(AllottedGeometry, ArrangedChildren);
+	SlateEditableTextLayout.OnArrangeChildren(AllottedGeometry, ArrangedChildren);
 }
 
 bool SMultiLineEditableText::SupportsKeyboardFocus() const
@@ -719,7 +713,7 @@ FReply SMultiLineEditableText::OnKeyChar( const FGeometry& MyGeometry,const FCha
 
 	if (!Reply.IsEventHandled())
 	{
-		Reply = EditableTextLayout->HandleKeyChar(InCharacterEvent);
+		Reply = SlateEditableTextLayout.HandleKeyChar(InCharacterEvent);
 	}
 
 	return Reply;
@@ -737,7 +731,7 @@ FReply SMultiLineEditableText::OnKeyDown( const FGeometry& MyGeometry, const FKe
 
 	if (!Reply.IsEventHandled())
 	{
-		Reply = EditableTextLayout->HandleKeyDown(InKeyEvent);
+		Reply = SlateEditableTextLayout.HandleKeyDown(InKeyEvent);
 	}
 
 	return Reply;
@@ -745,20 +739,20 @@ FReply SMultiLineEditableText::OnKeyDown( const FGeometry& MyGeometry, const FKe
 
 FReply SMultiLineEditableText::OnKeyUp( const FGeometry& MyGeometry, const FKeyEvent& InKeyEvent )
 {
-	return EditableTextLayout->HandleKeyUp(InKeyEvent);
+	return SlateEditableTextLayout.HandleKeyUp(InKeyEvent);
 }
 
-FReply SMultiLineEditableText::OnMouseButtonDown( const FGeometry& MyGeometry, const FPointerEvent& MouseEvent ) 
+FReply SMultiLineEditableText::OnMouseButtonDown( const FGeometry& MyGeometry, const FPointerEvent& MouseEvent )
 {
 	if (MouseEvent.GetEffectingButton() == EKeys::RightMouseButton)
 	{
 		AmountScrolledWhileRightMouseDown = 0.0f;
 	}
 
-	return EditableTextLayout->HandleMouseButtonDown(MyGeometry, MouseEvent);
+	return SlateEditableTextLayout.HandleMouseButtonDown(MyGeometry, MouseEvent);
 }
 
-FReply SMultiLineEditableText::OnMouseButtonUp( const FGeometry& MyGeometry, const FPointerEvent& MouseEvent ) 
+FReply SMultiLineEditableText::OnMouseButtonUp( const FGeometry& MyGeometry, const FPointerEvent& MouseEvent )
 {
 	if (MouseEvent.GetEffectingButton() == EKeys::RightMouseButton)
 	{
@@ -774,7 +768,7 @@ FReply SMultiLineEditableText::OnMouseButtonUp( const FGeometry& MyGeometry, con
 		}
 	}
 
-	return EditableTextLayout->HandleMouseButtonUp(MyGeometry, MouseEvent);
+	return SlateEditableTextLayout.HandleMouseButtonUp(MyGeometry, MouseEvent);
 }
 
 FReply SMultiLineEditableText::OnMouseMove( const FGeometry& MyGeometry, const FPointerEvent& MouseEvent )
@@ -789,11 +783,11 @@ FReply SMultiLineEditableText::OnMouseMove( const FGeometry& MyGeometry, const F
 
 		if (IsRightClickScrolling())
 		{
-			const FVector2f PreviousScrollOffset = EditableTextLayout->GetScrollOffset();
+			const FVector2f PreviousScrollOffset = SlateEditableTextLayout.GetScrollOffset();
 
 			FVector2f NewScrollOffset = PreviousScrollOffset;
 			NewScrollOffset.Y -= ScrollByAmount;
-			EditableTextLayout->SetScrollOffset(NewScrollOffset, MyGeometry);
+			SlateEditableTextLayout.SetScrollOffset(NewScrollOffset, MyGeometry);
 
 			if (!bIsSoftwareCursor)
 			{
@@ -803,7 +797,7 @@ FReply SMultiLineEditableText::OnMouseMove( const FGeometry& MyGeometry, const F
 
 			if (PreviousScrollOffset.Y != NewScrollOffset.Y)
 			{
-				const float ScrollMax = EditableTextLayout->GetSize().Y - MyGeometry.GetLocalSize().Y;
+				const float ScrollMax = SlateEditableTextLayout.GetSize().Y - MyGeometry.GetLocalSize().Y;
 				const float ScrollbarOffset = (ScrollMax != 0.0f) ? NewScrollOffset.Y / ScrollMax : 0.0f;
 				OnVScrollBarUserScrolled.ExecuteIfBound(ScrollbarOffset);
 				SoftwareCursorPosition.Y += (PreviousScrollOffset.Y - NewScrollOffset.Y);
@@ -813,7 +807,7 @@ FReply SMultiLineEditableText::OnMouseMove( const FGeometry& MyGeometry, const F
 		}
 	}
 
-	return EditableTextLayout->HandleMouseMove(MyGeometry, MouseEvent);
+	return SlateEditableTextLayout.HandleMouseMove(MyGeometry, MouseEvent);
 }
 
 FReply SMultiLineEditableText::OnMouseWheel( const FGeometry& MyGeometry, const FPointerEvent& MouseEvent )
@@ -822,15 +816,15 @@ FReply SMultiLineEditableText::OnMouseWheel( const FGeometry& MyGeometry, const
 	{
 		const float ScrollAmount = -MouseEvent.GetWheelDelta() * GetGlobalScrollAmount();
 
-		const FVector2f PreviousScrollOffset = EditableTextLayout->GetScrollOffset();
-		
+		const FVector2f PreviousScrollOffset = SlateEditableTextLayout.GetScrollOffset();
+
 		FVector2f NewScrollOffset = PreviousScrollOffset;
 		NewScrollOffset.Y += ScrollAmount;
-		EditableTextLayout->SetScrollOffset(NewScrollOffset, MyGeometry);
+		SlateEditableTextLayout.SetScrollOffset(NewScrollOffset, MyGeometry);
 
 		if (PreviousScrollOffset.Y != NewScrollOffset.Y)
 		{
-			const float ScrollMax = EditableTextLayout->GetSize().Y - MyGeometry.GetLocalSize().Y;
+			const float ScrollMax = SlateEditableTextLayout.GetSize().Y - MyGeometry.GetLocalSize().Y;
 			const float ScrollbarOffset = (ScrollMax != 0.0f) ? NewScrollOffset.Y / ScrollMax : 0.0f;
 			OnVScrollBarUserScrolled.ExecuteIfBound(ScrollbarOffset);
 			return FReply::Handled();
@@ -842,7 +836,7 @@ FReply SMultiLineEditableText::OnMouseWheel( const FGeometry& MyGeometry, const
 
 FReply SMultiLineEditableText::OnMouseButtonDoubleClick(const FGeometry& MyGeometry, const FPointerEvent& MouseEvent)
 {
-	return EditableTextLayout->HandleMouseButtonDoubleClick(MyGeometry, MouseEvent);
+	return SlateEditableTextLayout.HandleMouseButtonDoubleClick(MyGeometry, MouseEvent);
 }
 
 FCursorReply SMultiLineEditableText::OnCursorQuery( const FGeometry& MyGeometry, const FPointerEvent& CursorEvent ) const
@@ -866,7 +860,7 @@ bool SMultiLineEditableText::ComputeVolatility() const
 {
 	return SWidget::ComputeVolatility()
 		|| HasKeyboardFocus()
-		|| EditableTextLayout->ComputeVolatility()
+		|| SlateEditableTextLayout.ComputeVolatility()
 		|| bIsReadOnly.IsBound();
 }
 
diff --git a/Engine/Source/Runtime/Slate/Private/Widgets/Text/SlateEditableTextLayout.cpp b/Engine/Source/Runtime/Slate/Private/Widgets/Text/SlateEditableTextLayout.cpp
--- a/Engine/Source/Runtime/Slate/Private/Widgets/Text/SlateEditableTextLayout.cpp
+++ b/Engine/Source/Runtime/Slate/Private/Widgets/Text/SlateEditableTextLayout.cpp
@@ -16,12 +16,11 @@
 #include "Framework/Text/TextEditHelper.h"
 #include "Framework/Commands/GenericCommands.h"
 #include "Internationalization/BreakIterator.h"
-#include "SlateSettings.h"
 #include "HAL/PlatformApplicationMisc.h"
 
 /**
  * Ensure that text transactions are always completed.
- * i.e. never forget to call FinishChangingText();
+ * i.e. never forget to call EndEditTransaction();
  */
 struct FScopedEditableTextTransaction
 {
@@ -29,7 +28,7 @@ public:
 	explicit FScopedEditableTextTransaction(FSlateEditableTextLayout& InEditableTextLayout)
 		: EditableTextLayout(&InEditableTextLayout)
 	{
-		EditableTextLayout->BeginEditTransation();
+		EditableTextLayout->BeginEditTransaction();
 	}
 
 	~FScopedEditableTextTransaction()
@@ -65,7 +64,7 @@ bool IsCharAllowed(const TCHAR InChar)
 
 }
 
-FSlateEditableTextLayout::FSlateEditableTextLayout(ISlateEditableTextWidget& InOwnerWidget, const TAttribute<FText>& InInitialText, FTextBlockStyle InTextStyle, const TOptional<ETextShapingMethod> InTextShapingMethod, const TOptional<ETextFlowDirection> InTextFlowDirection, const FCreateSlateTextLayout& InCreateSlateTextLayout, TSharedRef<ITextLayoutMarshaller> InTextMarshaller, TSharedRef<ITextLayoutMarshaller> InHintTextMarshaller)
+void FSlateEditableTextLayout::Construct(TAttribute<FText> InInitialText, FTextBlockStyle InTextStyle, const TOptional<ETextShapingMethod>& InTextShapingMethod, const TOptional<ETextFlowDirection>& InTextFlowDirection, const FCreateSlateTextLayout& InCreateSlateTextLayout, TSharedRef<ITextLayoutMarshaller> InTextMarshaller, TSharedRef<ITextLayoutMarshaller> InHintTextMarshaller)
 {
 	CreateSlateTextLayout = InCreateSlateTextLayout;
 	if (!CreateSlateTextLayout.IsBound())
@@ -76,12 +75,13 @@ FSlateEditableTextLayout::FSlateEditableTextLayout(ISlateEditableTextWidget& InO
 			});
 	}
 
-	OwnerWidget = &InOwnerWidget;
-	Marshaller = InTextMarshaller;
-	HintMarshaller = InHintTextMarshaller;
-	TextStyle = InTextStyle;
+	const TSharedRef<SWidget> OwnerSlateWidget = OwnerWidget->GetSlateWidget();
+
+	Marshaller = MoveTemp(InTextMarshaller);
+	HintMarshaller = MoveTemp(InHintTextMarshaller);
+	TextStyle = MoveTemp(InTextStyle);
 	HintTextStyle = TextStyle;
-	TextLayout = CreateSlateTextLayout.Execute(&InOwnerWidget.GetSlateWidget().Get(), TextStyle);
+	TextLayout = CreateSlateTextLayout.Execute(&OwnerSlateWidget.Get(), TextStyle);
 
 	WrapTextAt = 0.0f;
 	AutoWrapText = false;
@@ -96,7 +96,7 @@ FSlateEditableTextLayout::FSlateEditableTextLayout(ISlateEditableTextWidget& InO
 
 	GraphemeBreakIterator = GSlateUseSharedBreakIterator ? FBreakIterator::GetCharacterBoundaryIterator() : FBreakIterator::CreateCharacterBoundaryIterator();
 
-	BoundText = InInitialText;
+	BoundText = MoveTemp(InInitialText);
 
 	// Set the initial text - the same as SetText, but doesn't call OnTextChanged as that can cause a crash during construction
 	{
@@ -149,7 +149,7 @@ FSlateEditableTextLayout::FSlateEditableTextLayout(ISlateEditableTextWidget& InO
 
 	auto ExecuteDeleteAction = [this]()
 	{
-		BeginEditTransation();
+		BeginEditTransaction();
 		DeleteSelectedText();
 		EndEditTransaction();
 	};
@@ -220,10 +220,10 @@ FSlateEditableTextLayout::~FSlateEditableTextLayout()
 	}
 }
 
-void FSlateEditableTextLayout::SetText(const TAttribute<FText>& InText)
+void FSlateEditableTextLayout::SetText(TAttribute<FText> InText)
 {
 	const FText PreviousText = BoundText.Get(FText::GetEmpty());
-	BoundText = InText;
+	BoundText = MoveTemp(InText);
 	const FText NewText = BoundText.Get(FText::GetEmpty());
 
 	// We need to force an update if the text doesn't match the editable text, as the editable 
@@ -257,7 +257,7 @@ void FSlateEditableTextLayout::SetText(const TAttribute<FText>& InText)
 	}
 }
 
-int32 FSlateEditableTextLayout::GetTextLineCount()
+int32 FSlateEditableTextLayout::GetTextLineCount() const
 {
 	return TextLayout->GetLineCount();
 }
@@ -269,9 +269,9 @@ FText FSlateEditableTextLayout::GetText() const
 	return BoundText.Get(FText::GetEmpty());
 }
 
-void FSlateEditableTextLayout::SetHintText(const TAttribute<FText>& InHintText)
+void FSlateEditableTextLayout::SetHintText(TAttribute<FText> InHintText)
 {
-	HintText = InHintText;
+	HintText = MoveTemp(InHintText);
 
 	// If we have hint text that is either non-empty or bound to a delegate, we'll also need to make the hint text layout
 	if (HintText.IsBound() || !HintText.Get(FText::GetEmpty()).IsEmpty())
@@ -292,11 +292,11 @@ FText FSlateEditableTextLayout::GetHintText() const
 	return HintText.Get(FText::GetEmpty());
 }
 
-void FSlateEditableTextLayout::SetSearchText(const TAttribute<FText>& InSearchText)
+void FSlateEditableTextLayout::SetSearchText(TAttribute<FText> InSearchText)
 {
 	const FText& SearchTextToSet = InSearchText.Get(FText::GetEmpty());
 
-	BoundSearchText = InSearchText;
+	BoundSearchText = MoveTemp(InSearchText);
 	BoundSearchTextLastTick = FTextSnapshot(SearchTextToSet);
 
 	BeginSearch(SearchTextToSet);
@@ -325,6 +325,9 @@ void FSlateEditableTextLayout::SetTextStyle(const FTextBlockStyle& InTextStyle)
 
 	TextLayout->SetDefaultTextStyle(TextStyle);
 	Marshaller->MakeDirty(); // will regenerate the text using the new default style
+
+	// Using Layout as changing text block size can affect the size.
+	OwnerWidget->GetSlateWidget()->Invalidate(EInvalidateWidget::Layout);
 }
 
 const FTextBlockStyle& FSlateEditableTextLayout::GetTextStyle() const
@@ -332,9 +335,9 @@ const FTextBlockStyle& FSlateEditableTextLayout::GetTextStyle() const
 	return TextStyle;
 }
 
-void FSlateEditableTextLayout::SetHintTextOpacity(const TAttribute<float>& InOpacity)
+void FSlateEditableTextLayout::SetHintTextOpacity(TAttribute<float> InOpacity)
 {
-	HintTextOpacity = InOpacity;
+	HintTextOpacity = MoveTemp(InOpacity);
 }
 
 float FSlateEditableTextLayout::GetHintTextOpacity() const
@@ -467,76 +470,93 @@ FTextSelection FSlateEditableTextLayout::GetSelection() const
 
 void FSlateEditableTextLayout::SetTextShapingMethod(const TOptional<ETextShapingMethod>& InTextShapingMethod)
 {
-	TextLayout->SetTextShapingMethod((InTextShapingMethod.IsSet()) ? InTextShapingMethod.GetValue() : GetDefaultTextShapingMethod());
+	const ETextShapingMethod NewShapingMethod = InTextShapingMethod.Get(GetDefaultTextShapingMethod());
+	if (NewShapingMethod != TextLayout->GetTextShapingMethod())
+	{
+		TextLayout->SetTextShapingMethod(NewShapingMethod);
+
+		OwnerWidget->GetSlateWidget()->Invalidate(EInvalidateWidget::Layout);
+	}
 }
 
 void FSlateEditableTextLayout::SetTextFlowDirection(const TOptional<ETextFlowDirection>& InTextFlowDirection)
 {
-	TextLayout->SetTextFlowDirection((InTextFlowDirection.IsSet()) ? InTextFlowDirection.GetValue() : GetDefaultTextFlowDirection());
+	const ETextFlowDirection NewFlowDirection = InTextFlowDirection.Get(GetDefaultTextFlowDirection());
+	if (NewFlowDirection != TextLayout->GetTextFlowDirection())
+	{
+		TextLayout->SetTextFlowDirection(NewFlowDirection);
+
+		OwnerWidget->GetSlateWidget()->Invalidate(EInvalidateWidget::Layout);
+	}
 }
 
-void FSlateEditableTextLayout::SetTextWrapping(const TAttribute<float>& InWrapTextAt, const TAttribute<bool>& InAutoWrapText, const TAttribute<ETextWrappingPolicy>& InWrappingPolicy)
+void FSlateEditableTextLayout::SetTextWrapping(TAttribute<float> InWrapTextAt, TAttribute<bool> InAutoWrapText, TAttribute<ETextWrappingPolicy> InWrappingPolicy)
 {
-	WrapTextAt = InWrapTextAt;
-	AutoWrapText = InAutoWrapText;
-	WrappingPolicy = InWrappingPolicy;
+	WrapTextAt = MoveTemp(InWrapTextAt);
+	AutoWrapText = MoveTemp(InAutoWrapText);
+	WrappingPolicy = MoveTemp(InWrappingPolicy);
 
 	OwnerWidget->GetSlateWidget()->Invalidate(EInvalidateWidget::LayoutAndVolatility);
 }
 
-void FSlateEditableTextLayout::SetWrapTextAt(const TAttribute<float>& InWrapTextAt)
+void FSlateEditableTextLayout::SetWrapTextAt(TAttribute<float> InWrapTextAt)
 {
-	WrapTextAt = InWrapTextAt;
+	WrapTextAt = MoveTemp(InWrapTextAt);
 
 	OwnerWidget->GetSlateWidget()->Invalidate(EInvalidateWidget::LayoutAndVolatility);
 }
 
-void FSlateEditableTextLayout::SetAutoWrapText(const TAttribute<bool>& InAutoWrapText)
+void FSlateEditableTextLayout::SetAutoWrapText(TAttribute<bool> InAutoWrapText)
 {
-	AutoWrapText = InAutoWrapText;
+	AutoWrapText = MoveTemp(InAutoWrapText);
 
 	OwnerWidget->GetSlateWidget()->Invalidate(EInvalidateWidget::LayoutAndVolatility);
 }
 
-void FSlateEditableTextLayout::SetWrappingPolicy(const TAttribute<ETextWrappingPolicy>& InWrappingPolicy)
+void FSlateEditableTextLayout::SetWrappingPolicy(TAttribute<ETextWrappingPolicy> InWrappingPolicy)
 {
-	WrappingPolicy = InWrappingPolicy;
+	WrappingPolicy = MoveTemp(InWrappingPolicy);
 
 	OwnerWidget->GetSlateWidget()->Invalidate(EInvalidateWidget::LayoutAndVolatility);
 }
 
-void FSlateEditableTextLayout::SetMargin(const TAttribute<FMargin>& InMargin)
+void FSlateEditableTextLayout::SetMargin(TAttribute<FMargin> InMargin)
 {
-	Margin = InMargin;
+	Margin = MoveTemp(InMargin);
 
 	OwnerWidget->GetSlateWidget()->Invalidate(EInvalidateWidget::LayoutAndVolatility);
 }
 
-void FSlateEditableTextLayout::SetJustification(const TAttribute<ETextJustify::Type>& InJustification)
+void FSlateEditableTextLayout::SetJustification(TAttribute<ETextJustify::Type> InJustification)
 {
 	if (!Justification.IdenticalTo(InJustification))
 	{
-		Justification = InJustification;
+		Justification = MoveTemp(InJustification);
 
 		OwnerWidget->GetSlateWidget()->Invalidate(EInvalidateWidget::LayoutAndVolatility);
 	}
 }
 
-void FSlateEditableTextLayout::SetLineHeightPercentage(const TAttribute<float>& InLineHeightPercentage)
+ETextJustify::Type FSlateEditableTextLayout::GetJustification() const
+{
+	return Justification.Get(ETextJustify::Left);
+}
+
+void FSlateEditableTextLayout::SetLineHeightPercentage(TAttribute<float> InLineHeightPercentage)
 {
-	LineHeightPercentage = InLineHeightPercentage;
+	LineHeightPercentage = MoveTemp(InLineHeightPercentage);
 
 	OwnerWidget->GetSlateWidget()->Invalidate(EInvalidateWidget::LayoutAndVolatility);
 }
 
-void FSlateEditableTextLayout::SetApplyLineHeightToBottomLine(const TAttribute<bool>& InApplyLineHeightToBottomLine)
+void FSlateEditableTextLayout::SetApplyLineHeightToBottomLine(TAttribute<bool> InApplyLineHeightToBottomLine)
 {
-	ApplyLineHeightToBottomLine = InApplyLineHeightToBottomLine;
+	ApplyLineHeightToBottomLine = MoveTemp(InApplyLineHeightToBottomLine);
 
 	OwnerWidget->GetSlateWidget()->Invalidate(EInvalidateWidget::LayoutAndVolatility);
 }
 
-void FSlateEditableTextLayout::SetOverflowPolicy(TOptional<ETextOverflowPolicy> InOverflowPolicy)
+void FSlateEditableTextLayout::SetOverflowPolicy(const TOptional<ETextOverflowPolicy>& InOverflowPolicy)
 {
 	if(OverflowPolicyOverride != InOverflowPolicy)
 	{
@@ -830,7 +850,7 @@ float FSlateEditableTextLayout::GetComputedWrappingWidth() const
 
 bool FSlateEditableTextLayout::GetAutoWrapText() const
 {
-	return AutoWrapText.Get();
+	return AutoWrapText.Get(false);
 }
 
 bool FSlateEditableTextLayout::HandleFocusReceived(const FFocusEvent& InFocusEvent)
@@ -3179,7 +3199,7 @@ bool FSlateEditableTextLayout::HasTextChangedFromOriginal() const
 	return bHasChanged;
 }
 
-void FSlateEditableTextLayout::BeginEditTransation()
+void FSlateEditableTextLayout::BeginEditTransaction()
 {
 	NumTransactionsOpened += 1;
 	
@@ -3529,7 +3549,7 @@ void FSlateEditableTextLayout::Tick(const FGeometry& AllottedGeometry, const dou
 	// If we're auto-wrapping, we need to hide the scrollbars until the first valid auto-wrap has been performed
 	// If we don't do this, then we can get some nasty layout shuffling as the scrollbars appear for one frame and then vanish again
 	// We also hide the scrollbars for non-multi-line text widgets
-	const EVisibility ScrollBarVisiblityOverride = ((AutoWrapText.Get() && CachedSize.IsZero()) || !OwnerWidget->IsMultiLineTextEdit()) ? EVisibility::Collapsed : EVisibility::Visible;
+	const EVisibility ScrollBarVisibilityOverride = ((AutoWrapText.Get(false) && CachedSize.IsZero()) || !OwnerWidget->IsMultiLineTextEdit()) ? EVisibility::Collapsed : EVisibility::Visible;
 
 	// Try and make sure that the line containing the cursor is in view
 	if (PositionToScrollIntoView.IsSet())
@@ -3579,7 +3599,7 @@ void FSlateEditableTextLayout::Tick(const FGeometry& AllottedGeometry, const dou
 
 		// Update the scrollbar with the clamped version of the offset
 		ScrollOffset.X = ViewOffset * ContentSize;
-		ScrollOffset.X = OwnerWidget->UpdateAndClampHorizontalScrollBar(ViewOffset, ViewFraction, ScrollBarVisiblityOverride);
+		ScrollOffset.X = OwnerWidget->UpdateAndClampHorizontalScrollBar(ViewOffset, ViewFraction, ScrollBarVisibilityOverride);
 	}
 
 	{
@@ -3592,15 +3612,15 @@ void FSlateEditableTextLayout::Tick(const FGeometry& AllottedGeometry, const dou
 
 		// Update the scrollbar with the clamped version of the offset
 		ScrollOffset.Y = ViewOffset * ContentSize;
-		ScrollOffset.Y = OwnerWidget->UpdateAndClampVerticalScrollBar(ViewOffset, ViewFraction, ScrollBarVisiblityOverride);
+		ScrollOffset.Y = OwnerWidget->UpdateAndClampVerticalScrollBar(ViewOffset, ViewFraction, ScrollBarVisibilityOverride);
 	}
 
 	TextLayout->SetVisibleRegion(AllottedGeometry.Size, FVector2D(ScrollOffset) * TextLayout->GetScale());
 }
 
-int32 FSlateEditableTextLayout::OnPaint(const FPaintArgs& Args, const FGeometry& AllottedGeometry, const FSlateRect& MyCullingRect, FSlateWindowElementList& OutDrawElements, int32 LayerId, const FWidgetStyle& InWidgetStyle, bool bParentEnabled)
+int32 FSlateEditableTextLayout::OnPaint(const FPaintArgs& Args, const FGeometry& AllottedGeometry, const FSlateRect& MyCullingRect, FSlateWindowElementList& OutDrawElements, int32 LayerId, const FWidgetStyle& InWidgetStyle, bool bParentEnabled) const
 {
-	// Update the auto-wrap size now that we have computed paint geometry; won't take affect until text frame
+	// Update the auto-wrap size now that we have computed paint geometry; won't take effect until text frame
 	// Note: This is done here rather than in Tick(), because Tick() doesn't get called while resizing windows, but OnPaint() does
 	CachedSize = FVector2f(AllottedGeometry.GetLocalSize());
 
@@ -3627,27 +3647,27 @@ void FSlateEditableTextLayout::CacheDesiredSize(float LayoutScaleMultiplier)
 	const float CaretWidth = FTextEditHelper::CalculateCaretWidth(FontMaxCharHeight);
 
 	// Get the wrapping width and font to see if they have changed
-	float WrappingWidth = WrapTextAt.Get();
+	float WrappingWidth = WrapTextAt.Get(0.0f);
 
 	// Text wrapping can either be used defined (WrapTextAt), automatic (AutoWrapText), or a mixture of both
 	// Take whichever has the smallest value (>1)
-	if (AutoWrapText.Get() && CachedSize.X >= 1.0f)
+	if (AutoWrapText.Get(false) && CachedSize.X >= 1.0f)
 	{
 		WrappingWidth = (WrappingWidth >= 1.0f) ? FMath::Min(WrappingWidth, CachedSize.X) : CachedSize.X;
 	}
 
 	// Append the caret width to the margin to make sure it doesn't get clipped
-	FMargin MarginValue = Margin.Get();
+	FMargin MarginValue = Margin.Get(FMargin(0));
 	MarginValue.Left += CaretWidth;
 	MarginValue.Right += CaretWidth;
 
 	TextLayout->SetScale(LayoutScaleMultiplier);
 	TextLayout->SetWrappingWidth(WrappingWidth);
-	TextLayout->SetWrappingPolicy(WrappingPolicy.Get());
+	TextLayout->SetWrappingPolicy(WrappingPolicy.Get(ETextWrappingPolicy::DefaultWrapping));
 	TextLayout->SetMargin(MarginValue);
-	TextLayout->SetLineHeightPercentage(LineHeightPercentage.Get());
-	TextLayout->SetApplyLineHeightToBottomLine(ApplyLineHeightToBottomLine.Get());
-	TextLayout->SetJustification(Justification.Get());
+	TextLayout->SetLineHeightPercentage(LineHeightPercentage.Get(1.0f));
+	TextLayout->SetApplyLineHeightToBottomLine(ApplyLineHeightToBottomLine.Get(false));
+	TextLayout->SetJustification(Justification.Get(ETextJustify::Left));
 	TextLayout->SetVisibleRegion(FVector2D(CachedSize), FVector2D(ScrollOffset) * TextLayout->GetScale());
 	TextLayout->UpdateIfNeeded();
 }
@@ -3659,7 +3679,7 @@ FVector2D FSlateEditableTextLayout::ComputeDesiredSize(float LayoutScaleMultipli
 	const float FontMaxCharHeight = FTextEditHelper::GetFontHeight(TextStyle.Font);
 	const float CaretWidth = FTextEditHelper::CalculateCaretWidth(FontMaxCharHeight);
 
-	const float WrappingWidth = WrapTextAt.Get();
+	const float WrappingWidth = WrapTextAt.Get(0.0f);
 	float DesiredWidth = 0.0f;
 	float DesiredHeight = 0.0f;
 
@@ -3667,12 +3687,12 @@ FVector2D FSlateEditableTextLayout::ComputeDesiredSize(float LayoutScaleMultipli
 	if (TextLayout->IsEmpty() && HintTextLayout.IsValid())
 	{
 		// Append the caret width to the margin to mimic what happens to the main text layout
-		FMargin MarginValue = Margin.Get();
+		FMargin MarginValue = Margin.Get(FMargin(0));
 		MarginValue.Left += CaretWidth;
 		MarginValue.Right += CaretWidth;
 
 		const FVector2D HintTextSize = HintTextLayout->ComputeDesiredSize(
-			FSlateTextBlockLayout::FWidgetDesiredSizeArgs(HintText.Get(), FText::GetEmpty(), WrapTextAt.Get(), AutoWrapText.Get(), WrappingPolicy.Get(), ETextTransformPolicy::None, MarginValue, LineHeightPercentage.Get(), ApplyLineHeightToBottomLine.Get(), Justification.Get()),
+			FSlateTextBlockLayout::FWidgetDesiredSizeArgs(HintText.Get(FText::GetEmpty()), FText::GetEmpty(), WrappingWidth, AutoWrapText.Get(false), WrappingPolicy.Get(ETextWrappingPolicy::DefaultWrapping), ETextTransformPolicy::None, MarginValue, LineHeightPercentage.Get(1.0f), ApplyLineHeightToBottomLine.Get(false), Justification.Get(ETextJustify::Left)),
 			LayoutScaleMultiplier, HintTextStyle
 			);
 
@@ -4184,7 +4204,7 @@ void FSlateEditableTextLayout::FTextInputMethodContext::BeginComposition()
 	{
 		bIsComposing = true;
 
-		OwnerLayout->BeginEditTransation();
+		OwnerLayout->BeginEditTransaction();
 		OwnerLayout->UpdateCursorHighlight();
 	}
 }
diff --git a/Engine/Source/Runtime/Slate/Public/Widgets/Input/SEditableText.h b/Engine/Source/Runtime/Slate/Public/Widgets/Input/SEditableText.h
--- a/Engine/Source/Runtime/Slate/Public/Widgets/Input/SEditableText.h
+++ b/Engine/Source/Runtime/Slate/Public/Widgets/Input/SEditableText.h
@@ -19,13 +19,13 @@
 #include "Framework/Text/TextLayout.h"
 #include "Framework/Text/SlateAsyncFontLoadingTypes.h"
 #include "Widgets/Text/ISlateEditableTextWidget.h"
+#include "Widgets/Text/SlateEditableTextLayout.h"
 
 class FActiveTimerHandle;
 class FArrangedChildren;
 class FChildren;
 class FPaintArgs;
 class FPlainTextLayoutMarshaller;
-class FSlateEditableTextLayout;
 class FSlateWindowElementList;
 class IBreakIterator;
 struct FTextLocation;
@@ -203,9 +203,6 @@ public:
 	/** Constructor */
 	SLATE_API SEditableText();
 
-	/** Destructor */
-	SLATE_API ~SEditableText();
-
 	/**
 	 * Construct this widget
 	 *
@@ -218,7 +215,7 @@ public:
 	 *
 	 * @param  InNewText  The new text
 	 */
-	SLATE_API void SetText( const TAttribute< FText >& InNewText );
+	SLATE_API void SetText(TAttribute<FText> InNewText);
 	
 	/**
 	 * Returns the text string
@@ -238,47 +235,50 @@ public:
 	SLATE_API bool SetEditableText(const FText& InNewText);
 
 	/** See the HintText attribute */
-	SLATE_API void SetHintText( const TAttribute< FText >& InHintText );
+	SLATE_API void SetHintText(TAttribute<FText> InHintText);
 	
 	/** Get the text that appears when there is no text in the text box */
 	SLATE_API FText GetHintText() const;
 	
 	/** Sets the hint text opacity */
-	SLATE_API void SetHintTextOpacity(const TAttribute<float>& InHintTextOpacity);
+	SLATE_API void SetHintTextOpacity(TAttribute<float> InHintTextOpacity);
 	
 	/** Get the hint text opacity */
 	SLATE_API float GetHintTextOpacity() const;
 
 	/** Set the text that is currently being searched for (if any) */
-	SLATE_API void SetSearchText(const TAttribute<FText>& InSearchText);
+	SLATE_API void SetSearchText(TAttribute<FText> InSearchText);
 
 	/** Get the text that is currently being searched for (if any) */
 	SLATE_API FText GetSearchText() const;
 
 	/** See the IsReadOnly attribute */
-	SLATE_API void SetIsReadOnly( TAttribute< bool > InIsReadOnly );
+	SLATE_API void SetIsReadOnly(TAttribute<bool> InIsReadOnly);
 	
 	/** See the IsPassword attribute */
-	SLATE_API void SetIsPassword( TAttribute< bool > InIsPassword );
+	SLATE_API void SetIsPassword(TAttribute<bool> InIsPassword);
 	
 	/** See the ColorAndOpacity attribute */
 	SLATE_API void SetColorAndOpacity(TAttribute<FSlateColor> Color);
 
+	/** See the BackgroundImageSelected attribute */
+	SLATE_API void SetBackgroundImageSelected(TAttribute<const FSlateBrush*> InBackgroundImageSelected);
+
 	/** See the AllowContextMenu attribute */
-	SLATE_API void SetAllowContextMenu(const TAttribute< bool >& InAllowContextMenu);
+	SLATE_API void SetAllowContextMenu(TAttribute<bool> InAllowContextMenu);
 
 	/** See the EnableIntegratedKeyboard attribute */
-	SLATE_API void SetEnableIntegratedKeyboard(const TAttribute< bool >& InEnableIntegratedKeyboard);
+	SLATE_API void SetEnableIntegratedKeyboard(TAttribute<bool> InEnableIntegratedKeyboard);
 
 	/** Set the VirtualKeyboardDismissAction attribute */
-	SLATE_API void SetVirtualKeyboardDismissAction(TAttribute< EVirtualKeyboardDismissAction > InVirtualKeyboardDismissAction);
+	SLATE_API void SetVirtualKeyboardDismissAction(TAttribute<EVirtualKeyboardDismissAction> InVirtualKeyboardDismissAction);
 
 	/**
 	 * Sets the font used to draw the text
 	 *
 	 * @param  InNewFont	The new font to use
 	 */
-	SLATE_API void SetFont(const TAttribute< FSlateFontInfo >& InNewFont);
+	SLATE_API void SetFont(TAttribute<FSlateFontInfo> InNewFont);
 
 	/** Gets the font used to draw the text. */
 	SLATE_API FSlateFontInfo GetFont() const;
@@ -286,7 +286,7 @@ public:
 	/**
 	 * Sets the text style used to draw the text
 	 *
-	 * @param  NewTextStyle	The new text style to use
+	 * @param  InNewTextStyle	The new text style to use
 	 */
 	SLATE_API void SetTextStyle( const FEditableTextStyle& InNewTextStyle );
 
@@ -303,52 +303,53 @@ public:
 	 *
 	 * @param  InMinDesiredWidth	The minimum width
 	 */
-	SLATE_API void SetMinDesiredWidth(const TAttribute<float>& InMinDesiredWidth);
+	SLATE_API void SetMinDesiredWidth(TAttribute<float> InMinDesiredWidth);
 
 	/**
 	 * Workaround as we loose focus when the auto completion closes.
 	 *
 	 * @param  InIsCaretMovedWhenGainFocus	Workaround
 	 */
-	SLATE_API void SetIsCaretMovedWhenGainFocus(const TAttribute<bool>& InIsCaretMovedWhenGainFocus);
+	SLATE_API void SetIsCaretMovedWhenGainFocus(TAttribute<bool> InIsCaretMovedWhenGainFocus);
 
 	/**
 	 * Sets whether to select all text when the user clicks to give focus on the widget
 	 *
 	 * @param  InSelectAllTextWhenFocused	Select all text when the user clicks?
 	 */
-	SLATE_API void SetSelectAllTextWhenFocused(const TAttribute<bool>& InSelectAllTextWhenFocused);
+	SLATE_API void SetSelectAllTextWhenFocused(TAttribute<bool> InSelectAllTextWhenFocused);
 
 	/**
 	 * Sets whether to allow the user to back out of changes when they press the escape key
 	 *
 	 * @param  InRevertTextOnEscape			Allow the user to back out of changes?
 	 */
-	SLATE_API void SetRevertTextOnEscape(const TAttribute<bool>& InRevertTextOnEscape);
+	SLATE_API void SetRevertTextOnEscape(TAttribute<bool> InRevertTextOnEscape);
 
 	/**
 	 * Sets whether to clear keyboard focus when pressing enter to commit changes
 	 *
 	 * @param  InClearKeyboardFocusOnCommit		Clear keyboard focus when pressing enter?
 	 */
-	SLATE_API void SetClearKeyboardFocusOnCommit(const TAttribute<bool>& InClearKeyboardFocusOnCommit);
+	SLATE_API void SetClearKeyboardFocusOnCommit(TAttribute<bool> InClearKeyboardFocusOnCommit);
 
 	/**
 	 * Sets whether to select all text when pressing enter to commit changes
 	 *
 	 * @param  InSelectAllTextOnCommit		Select all text when pressing enter?
 	 */
-	SLATE_API void SetSelectAllTextOnCommit(const TAttribute<bool>& InSelectAllTextOnCommit);
+	SLATE_API void SetSelectAllTextOnCommit(TAttribute<bool> InSelectAllTextOnCommit);
 
 	/**
 	 * Sets whether to select word on the mouse double click
 	 *
 	 * @param  InSelectWordOnMouseDoubleClick		Select word on the mouse double click
 	 */
-	SLATE_API void SetSelectWordOnMouseDoubleClick(const TAttribute<bool>& InSelectWordOnMouseDoubleClick);
+	SLATE_API void SetSelectWordOnMouseDoubleClick(TAttribute<bool> InSelectWordOnMouseDoubleClick);
 
 	/** See Justification attribute */
-	SLATE_API void SetJustification(const TAttribute<ETextJustify::Type>& InJustification);
+	SLATE_API void SetJustification(TAttribute<ETextJustify::Type> InJustification);
+
 	/**
 	 * Sets the OnKeyCharHandler to provide first chance handling of the OnKeyChar event
 	 *
@@ -356,7 +357,7 @@ public:
 	 */
 	void SetOnKeyCharHandler(FOnKeyChar InOnKeyCharHandler)
 	{
-		OnKeyCharHandler = InOnKeyCharHandler;
+		OnKeyCharHandler = MoveTemp(InOnKeyCharHandler);
 	}
 
 	/**
@@ -366,7 +367,7 @@ public:
 	 */
 	void SetOnKeyDownHandler(FOnKeyDown InOnKeyDownHandler)
 	{
-		OnKeyDownHandler = InOnKeyDownHandler;
+		OnKeyDownHandler = MoveTemp(InOnKeyDownHandler);
 	}
 
 	/** See TextShapingMethod attribute */
@@ -376,7 +377,7 @@ public:
 	SLATE_API void SetTextFlowDirection(const TOptional<ETextFlowDirection>& InTextFlowDirection);
 
 	/** Sets the overflow policy for this text block */
-	SLATE_API void SetOverflowPolicy(TOptional<ETextOverflowPolicy> InOverflowPolicy);
+	SLATE_API void SetOverflowPolicy(const TOptional<ETextOverflowPolicy>& InOverflowPolicy);
 
 	/** Sets the paint policy to use for text if the font faces are still async loading. */
 	SLATE_API void SetFontFacesLoadingPaintPolicy(EFontFacesLoadingPaintPolicy InFontFacesLoadingPaintPolicy);
@@ -443,7 +444,7 @@ protected:
 	SLATE_API virtual FReply OnMouseMove( const FGeometry& InMyGeometry, const FPointerEvent& InMouseEvent ) override;
 	SLATE_API virtual FReply OnMouseButtonDoubleClick( const FGeometry& InMyGeometry, const FPointerEvent& InMouseEvent ) override;
 	SLATE_API virtual FCursorReply OnCursorQuery( const FGeometry& MyGeometry, const FPointerEvent& CursorEvent ) const override;
-	SLATE_API virtual const FSlateBrush* GetFocusBrush() const;
+	SLATE_API virtual const FSlateBrush* GetFocusBrush() const override;
 	SLATE_API virtual bool IsInteractable() const override;
 	SLATE_API virtual bool ComputeVolatility() const override;
 #if WITH_ACCESSIBILITY
@@ -487,16 +488,21 @@ protected:
 	SLATE_API virtual void OnTextChanged(const FText& InText) override;
 	SLATE_API virtual void OnTextCommitted(const FText& InText, const ETextCommit::Type InTextAction) override;
 	SLATE_API virtual void OnCursorMoved(const FTextLocation& InLocation) override;
-	SLATE_API virtual float UpdateAndClampHorizontalScrollBar(const float InViewOffset, const float InViewFraction, const EVisibility InVisiblityOverride) override;
-	SLATE_API virtual float UpdateAndClampVerticalScrollBar(const float InViewOffset, const float InViewFraction, const EVisibility InVisiblityOverride) override;
+	SLATE_API virtual float UpdateAndClampHorizontalScrollBar(const float InViewOffset, const float InViewFraction, const EVisibility InVisibilityOverride) override;
+	SLATE_API virtual float UpdateAndClampVerticalScrollBar(const float InViewOffset, const float InViewFraction, const EVisibility InVisibilityOverride) override;
 	//~ End ISlateEditableTextWidget Interface
 
 protected:
 	/** Text marshaller used by the editable text layout */
 	TSharedPtr<FPlainTextLayoutMarshaller> PlainTextMarshaller;
 
+#if WITH_EDITORONLY_DATA
 	/** The text layout that deals with the editable text */
+	UE_DEPRECATED(5.8, "Use SlateEditableTextLayout, this variable will never be set.")
 	TUniquePtr<FSlateEditableTextLayout> EditableTextLayout;
+#endif
+	/** The text layout that deals with the editable text */
+	FSlateEditableTextLayout SlateEditableTextLayout;
 
 	/** The font used to draw the text */
 	TAttribute<FSlateFontInfo> Font;
@@ -576,12 +582,12 @@ protected:
 	/** The type of virtual keyboard to use for editing this text on mobile */
 	TAttribute<EKeyboardType> VirtualKeyboardType;
 
-	/** Additional options used by the virtual keyboard summoned by this widget */
-	FVirtualKeyboardOptions VirtualKeyboardOptions;
-
 	/** The type of event that will trigger the display of the virtual keyboard */
 	TAttribute<EVirtualKeyboardTrigger> VirtualKeyboardTrigger;
 
 	/** The message action to take when the virtual keyboard is dismissed by the user */
 	TAttribute<EVirtualKeyboardDismissAction> VirtualKeyboardDismissAction;
+
+	/** Additional options used by the virtual keyboard summoned by this widget */
+	FVirtualKeyboardOptions VirtualKeyboardOptions;
 };
diff --git a/Engine/Source/Runtime/Slate/Public/Widgets/Input/SEditableTextBox.h b/Engine/Source/Runtime/Slate/Public/Widgets/Input/SEditableTextBox.h
--- a/Engine/Source/Runtime/Slate/Public/Widgets/Input/SEditableTextBox.h
+++ b/Engine/Source/Runtime/Slate/Public/Widgets/Input/SEditableTextBox.h
@@ -195,7 +195,7 @@ public:
 	 *
 	 * @param	InArgs	The declaration data for this widget
 	 */
-	SLATE_API void Construct( const FArguments& InArgs );
+	SLATE_API void Construct(const FArguments& InArgs);
 
 	/**
 	 * Returns the text string
@@ -222,7 +222,7 @@ public:
 	 *
 	 * @param  InNewText  The new text string
 	 */
-	SLATE_API void SetText( const TAttribute< FText >& InNewText );
+	SLATE_API void SetText(TAttribute<FText> InNewText);
 	
 	/**
 	 * Returns the Hint text string
@@ -235,130 +235,130 @@ public:
 	}
 
 	/** See the HintText attribute */
-	SLATE_API void SetHintText( const TAttribute< FText >& InHintText );
+	SLATE_API void SetHintText(TAttribute<FText> InHintText);
 	
 	/** Set the text that is currently being searched for (if any) */
-	SLATE_API void SetSearchText(const TAttribute<FText>& InSearchText);
+	SLATE_API void SetSearchText(TAttribute<FText> InSearchText);
 
 	/** Get the text that is currently being searched for (if any) */
 	SLATE_API FText GetSearchText() const;
 
 	/** See the IsReadOnly attribute */
-	SLATE_API void SetIsReadOnly( TAttribute< bool > InIsReadOnly );
+	SLATE_API void SetIsReadOnly(TAttribute<bool> InIsReadOnly);
 
 	bool IsReadOnly() const { return EditableText->IsTextReadOnly(); }
 	
 	/** See the IsPassword attribute */
-	SLATE_API void SetIsPassword( TAttribute< bool > InIsPassword );
+	SLATE_API void SetIsPassword(TAttribute<bool> InIsPassword);
 
 	bool IsPassword() const { return EditableText->IsTextPassword(); }
 
 	/** See the AllowContextMenu attribute */
-	SLATE_API void SetAllowContextMenu(TAttribute< bool > InAllowContextMenu);
+	SLATE_API void SetAllowContextMenu(TAttribute<bool> InAllowContextMenu);
 
 	/** Set the VirtualKeyboardDismissAction attribute */
-	SLATE_API void SetVirtualKeyboardDismissAction(TAttribute< EVirtualKeyboardDismissAction > InVirtualKeyboardDismissAction);
+	SLATE_API void SetVirtualKeyboardDismissAction(TAttribute<EVirtualKeyboardDismissAction> InVirtualKeyboardDismissAction);
 	
 	/**
 	 * Sets the font used to draw the text
 	 *
 	 * @param  InFont	The new font to use
 	 */
-	SLATE_API void SetFont(const TAttribute<FSlateFontInfo>& InFont);
+	SLATE_API void SetFont(TAttribute<FSlateFontInfo> InFont);
 
 	/**
 	 * Sets the text color and opacity (overrides Style)
 	 *
 	 * @param  InForegroundColor 	The text color and opacity
 	 */
-	SLATE_API void SetTextBoxForegroundColor(const TAttribute<FSlateColor>& InForegroundColor);
+	SLATE_API void SetTextBoxForegroundColor(TAttribute<FSlateColor> InForegroundColor);
 
 	/**
 	 * Sets the color of the background/border around the editable text (overrides Style) 
 	 *
 	 * @param  InBackgroundColor 	The background/border color
 	 */
-	SLATE_API void SetTextBoxBackgroundColor(const TAttribute<FSlateColor>& InBackgroundColor);
+	SLATE_API void SetTextBoxBackgroundColor(TAttribute<FSlateColor> InBackgroundColor);
 
 	/**
 	 * Sets the text color and opacity when read-only (overrides Style) 
 	 *
 	 * @param  InReadOnlyForegroundColor 	The read-only text color and opacity
 	 */
-	SLATE_API void SetReadOnlyForegroundColor(const TAttribute<FSlateColor>& InReadOnlyForegroundColor);
+	SLATE_API void SetReadOnlyForegroundColor(TAttribute<FSlateColor> InReadOnlyForegroundColor);
 
 	/**
 	 * Sets the text color and opacity when this box has keyboard focus(overrides Style)
 	 *
 	 * @param  InFocusedForegroundColor 	The focused color and opacity
 	 */
-	SLATE_API void SetFocusedForegroundColor(const TAttribute<FSlateColor>& InFocusedForegroundColor);
+	SLATE_API void SetFocusedForegroundColor(TAttribute<FSlateColor> InFocusedForegroundColor);
 
 	/**
 	 * Sets the maximum text length that will be accepted by the widget
 	 *
 	 * @param  InMaximumLength		Length to check for. Only values superior to 0 will be taken into account, otherwise no length check will occur.
 	 */
-	SLATE_API void SetMaximumLength(const TAttribute<int32>& InMaximumLength);
+	SLATE_API void SetMaximumLength(TAttribute<int32> InMaximumLength);
 
 	/**
 	 * Sets the minimum width that a text box should be.
 	 *
 	 * @param  InMinimumDesiredWidth	The minimum width
 	 */
-	SLATE_API void SetMinimumDesiredWidth(const TAttribute<float>& InMinimumDesiredWidth);
+	SLATE_API void SetMinimumDesiredWidth(TAttribute<float> InMinimumDesiredWidth);
 
 	/**
 	 * Workaround as we loose focus when the auto completion closes.
 	 *
 	 * @param  InIsCaretMovedWhenGainFocus	Workaround
 	 */
-	SLATE_API void SetIsCaretMovedWhenGainFocus(const TAttribute<bool>& InIsCaretMovedWhenGainFocus);
+	SLATE_API void SetIsCaretMovedWhenGainFocus(TAttribute<bool> InIsCaretMovedWhenGainFocus);
 
 	/**
 	 * Sets whether to select all text when the user clicks to give focus on the widget
 	 *
 	 * @param  InSelectAllTextWhenFocused	Select all text when the user clicks?
 	 */
-	SLATE_API void SetSelectAllTextWhenFocused(const TAttribute<bool>& InSelectAllTextWhenFocused);
+	SLATE_API void SetSelectAllTextWhenFocused(TAttribute<bool> InSelectAllTextWhenFocused);
 
 	/**
 	 * Sets whether to allow the user to back out of changes when they press the escape key
 	 *
 	 * @param  InRevertTextOnEscape			Allow the user to back out of changes?
 	 */
-	SLATE_API void SetRevertTextOnEscape(const TAttribute<bool>& InRevertTextOnEscape);
+	SLATE_API void SetRevertTextOnEscape(TAttribute<bool> InRevertTextOnEscape);
 
 	/**
 	 * Sets whether to clear keyboard focus when pressing enter to commit changes
 	 *
 	 * @param  InClearKeyboardFocusOnCommit		Clear keyboard focus when pressing enter?
 	 */
-	SLATE_API void SetClearKeyboardFocusOnCommit(const TAttribute<bool>& InClearKeyboardFocusOnCommit);
+	SLATE_API void SetClearKeyboardFocusOnCommit(TAttribute<bool> InClearKeyboardFocusOnCommit);
 
 	/**
 	 * Sets whether to select all text when pressing enter to commit changes
 	 *
 	 * @param  InSelectAllTextOnCommit		Select all text when pressing enter?
 	 */
-	SLATE_API void SetSelectAllTextOnCommit(const TAttribute<bool>& InSelectAllTextOnCommit);
+	SLATE_API void SetSelectAllTextOnCommit(TAttribute<bool> InSelectAllTextOnCommit);
 
 	/**
 	 * Sets whether to select select word on mouse double click
 	 *
 	 * @param  InSelectWordOnMouseDoubleClick		Select select word on mouse double click?
 	 */
-	SLATE_API void SetSelectWordOnMouseDoubleClick(const TAttribute<bool>& InSelectWordOnMouseDoubleClick);
+	SLATE_API void SetSelectWordOnMouseDoubleClick(TAttribute<bool> InSelectWordOnMouseDoubleClick);
 
 	/** See Justification attribute */
-	SLATE_API void SetJustification(const TAttribute<ETextJustify::Type>& InJustification);
+	SLATE_API void SetJustification(TAttribute<ETextJustify::Type> InJustification);
 
 	/**
 	 * If InError is a non-empty string the TextBox will the ErrorReporting provided during construction
 	 * If no error reporting was provided, the TextBox will create a default error reporter.
 	 */
-	SLATE_API void SetError( const FText& InError );
-	SLATE_API void SetError( const FString& InError );
+	SLATE_API void SetError(const FText& InError);
+	SLATE_API void SetError(const FString& InError);
 
 	/**
 	 * Sets the OnKeyCharHandler to provide first chance handling of the SEditableText's OnKeyChar event
@@ -381,7 +381,7 @@ public:
 	SLATE_API void SetTextFlowDirection(const TOptional<ETextFlowDirection>& InTextFlowDirection);
 
 	/** Sets the overflow policy for this text block */
-	SLATE_API void SetOverflowPolicy(TOptional<ETextOverflowPolicy> InOverflowPolicy);
+	SLATE_API void SetOverflowPolicy(const TOptional<ETextOverflowPolicy>& InOverflowPolicy);
 
 	/** Sets the paint policy to use for text if the font faces are still async loading. */
 	SLATE_API void SetFontFacesLoadingPaintPolicy(EFontFacesLoadingPaintPolicy InFontFacesLoadingPaintPolicy);
@@ -413,7 +413,7 @@ public:
 	/** Scroll to the given location in the document (without moving the cursor) */
 	void ScrollTo(const ETextLocation NewLocation)
 	{
-		EditableText->GoTo(NewLocation);
+		EditableText->ScrollTo(NewLocation);
 	}
 
 	/** Begin a new text search (this is called automatically when the bound search text changes) */
@@ -496,10 +496,10 @@ protected:
 
 private:
 
-	SLATE_API FMargin DeterminePadding() const;
-	SLATE_API FSlateFontInfo DetermineFont() const;
-	SLATE_API FSlateColor DetermineBackgroundColor() const;
-	SLATE_API FSlateColor DetermineForegroundColor() const;
+	FMargin DeterminePadding() const;
+	FSlateFontInfo DetermineFont() const;
+	FSlateColor DetermineBackgroundColor() const;
+	FSlateColor DetermineForegroundColor() const;
 
 	/** Styling: border image to draw when not hovered or focused */
 	const FSlateBrush* BorderImageNormal;
@@ -511,6 +511,6 @@ private:
 	const FSlateBrush* BorderImageReadOnly;
 
 	/** @return Border image for the text box based on the hovered and focused state */
-	SLATE_API const FSlateBrush* GetBorderImage() const;
+	const FSlateBrush* DetermineBorderImage() const;
 
 };
diff --git a/Engine/Source/Runtime/Slate/Public/Widgets/Input/SMultiLineEditableTextBox.h b/Engine/Source/Runtime/Slate/Public/Widgets/Input/SMultiLineEditableTextBox.h
--- a/Engine/Source/Runtime/Slate/Public/Widgets/Input/SMultiLineEditableTextBox.h
+++ b/Engine/Source/Runtime/Slate/Public/Widgets/Input/SMultiLineEditableTextBox.h
@@ -263,7 +263,7 @@ public:
 	 *
 	 * @param	InArgs	The declaration data for this widget
 	 */
-	SLATE_API void Construct( const FArguments& InArgs );
+	SLATE_API void Construct(const FArguments& InArgs);
 	
 	/**
 	 * Returns the text string
@@ -303,7 +303,7 @@ public:
 	 *
 	 * @param  InNewText  The new text string
 	 */
-	SLATE_API void SetText( const TAttribute< FText >& InNewText );
+	SLATE_API void SetText(TAttribute<FText> InNewText);
 
 	/**
 	 * Returns the hint text string
@@ -320,10 +320,10 @@ public:
 	 *
 	 * @param  InHintText The hint text string
 	 */
-	SLATE_API void SetHintText( const TAttribute< FText >& InHintText );
+	SLATE_API void SetHintText(TAttribute<FText> InHintText);
 
 	/** Set the text that is currently being searched for (if any) */
-	SLATE_API void SetSearchText(const TAttribute<FText>& InSearchText);
+	SLATE_API void SetSearchText(TAttribute<FText> InSearchText);
 
 	/** Get the text that is currently being searched for (if any) */
 	SLATE_API FText GetSearchText() const;
@@ -333,35 +333,35 @@ public:
 	 *
 	 * @param  InForegroundColor 	The text color and opacity
 	 */
-	SLATE_API void SetTextBoxForegroundColor(const TAttribute<FSlateColor>& InForegroundColor);
+	SLATE_API void SetTextBoxForegroundColor(TAttribute<FSlateColor> InForegroundColor);
 
 	/**
 	 * Sets the color of the background/border around the editable text (overrides Style) 
 	 *
 	 * @param  InBackgroundColor 	The background/border color
 	 */
-	SLATE_API void SetTextBoxBackgroundColor(const TAttribute<FSlateColor>& InBackgroundColor);
+	SLATE_API void SetTextBoxBackgroundColor(TAttribute<FSlateColor> InBackgroundColor);
 
 	/**
 	 * Sets the text color and opacity when read-only (overrides Style) 
 	 *
 	 * @param  InReadOnlyForegroundColor 	The read-only text color and opacity
 	 */
-	SLATE_API void SetReadOnlyForegroundColor(const TAttribute<FSlateColor>& InReadOnlyForegroundColor);
+	SLATE_API void SetReadOnlyForegroundColor(TAttribute<FSlateColor> InReadOnlyForegroundColor);
 
 	/**
 	 * Sets the maximum text length that will be accepted by the widget
 	 *
 	 * @param  InMaximumLength		Length to check for. Only values superior to 0 will be taken into account, otherwise no length check will occur.
 	 */
-	SLATE_API void SetMaximumLength(const TAttribute<int32>& InMaximumLength);
+	SLATE_API void SetMaximumLength(TAttribute<int32> InMaximumLength);
 
 	/**
 	 * Sets whether to select word on the mouse double click
 	 *
 	 * @param  InSelectWordOnMouseDoubleClick		Select word on the mouse double click
 	 */
-	SLATE_API void SetSelectWordOnMouseDoubleClick(const TAttribute<bool>& InSelectWordOnMouseDoubleClick);
+	SLATE_API void SetSelectWordOnMouseDoubleClick(TAttribute<bool> InSelectWordOnMouseDoubleClick);
 
 	/** See TextShapingMethod attribute */
 	SLATE_API void SetTextShapingMethod(const TOptional<ETextShapingMethod>& InTextShapingMethod);
@@ -370,47 +370,47 @@ public:
 	SLATE_API void SetTextFlowDirection(const TOptional<ETextFlowDirection>& InTextFlowDirection);
 
 	/** See WrapTextAt attribute */
-	SLATE_API void SetWrapTextAt(const TAttribute<float>& InWrapTextAt);
+	SLATE_API void SetWrapTextAt(TAttribute<float> InWrapTextAt);
 
 	/** See AutoWrapText attribute */
-	SLATE_API void SetAutoWrapText(const TAttribute<bool>& InAutoWrapText);
+	SLATE_API void SetAutoWrapText(TAttribute<bool> InAutoWrapText);
 
 	/** Set WrappingPolicy attribute */
-	SLATE_API void SetWrappingPolicy(const TAttribute<ETextWrappingPolicy>& InWrappingPolicy);
+	SLATE_API void SetWrappingPolicy(TAttribute<ETextWrappingPolicy> InWrappingPolicy);
 
 	/** See LineHeightPercentage attribute */
-	SLATE_API void SetLineHeightPercentage(const TAttribute<float>& InLineHeightPercentage);
+	SLATE_API void SetLineHeightPercentage(TAttribute<float> InLineHeightPercentage);
 
 	/** See ApplyLineHeightToBottomLine attribute */
-	SLATE_API void SetApplyLineHeightToBottomLine(const TAttribute<bool>& InApplyLineHeightToBottomLine);
+	SLATE_API void SetApplyLineHeightToBottomLine(TAttribute<bool> InApplyLineHeightToBottomLine);
 
 	/** See Margin attribute */
-	SLATE_API void SetMargin(const TAttribute<FMargin>& InMargin);
+	SLATE_API void SetMargin(TAttribute<FMargin> InMargin);
 
 	/** See Justification attribute */
-	SLATE_API void SetJustification(const TAttribute<ETextJustify::Type>& InJustification);
+	SLATE_API void SetJustification(TAttribute<ETextJustify::Type> InJustification);
 
 	/** Sets the overflow policy for this text block */
-	SLATE_API void SetOverflowPolicy(TOptional<ETextOverflowPolicy> InOverflowPolicy);
+	SLATE_API void SetOverflowPolicy(const TOptional<ETextOverflowPolicy>& InOverflowPolicy);
 
 	/** See the AllowContextMenu attribute */
-	SLATE_API void SetAllowContextMenu(const TAttribute< bool >& InAllowContextMenu);
+	SLATE_API void SetAllowContextMenu(TAttribute<bool> InAllowContextMenu);
 
 	/** Set the VirtualKeyboardDismissAction attribute */
-	SLATE_API void SetVirtualKeyboardDismissAction(TAttribute< EVirtualKeyboardDismissAction > InVirtualKeyboardDismissAction);
+	SLATE_API void SetVirtualKeyboardDismissAction(TAttribute<EVirtualKeyboardDismissAction> InVirtualKeyboardDismissAction);
 	
 	/** Set the ReadOnly attribute */
-	SLATE_API void SetIsReadOnly(const TAttribute< bool >& InIsReadOnly);
+	SLATE_API void SetIsReadOnly(TAttribute<bool> InIsReadOnly);
 
 	/** Set the FontFacesLoadingPaintPolicy argument */
-	SLATE_API void SetFontFacesLoadingPaintPolicy( EFontFacesLoadingPaintPolicy InFontFacesLoadingPaintPolicy);
+	SLATE_API void SetFontFacesLoadingPaintPolicy(EFontFacesLoadingPaintPolicy InFontFacesLoadingPaintPolicy);
 
 	/**
 	 * If InError is a non-empty string the TextBox will the ErrorReporting provided during construction
 	 * If no error reporting was provided, the TextBox will create a default error reporter.
 	 */
-	SLATE_API void SetError( const FText& InError );
-	SLATE_API void SetError( const FString& InError );
+	SLATE_API void SetError(const FText& InError);
+	SLATE_API void SetError(const FString& InError);
 	SLATE_API bool HasError() const;
 
 	// SWidget overrides
@@ -586,12 +586,12 @@ protected:
 
 private:
 
-	SLATE_API FMargin DeterminePadding() const;
-	SLATE_API FMargin DetermineHScrollBarPadding() const;
-	SLATE_API FMargin DetermineVScrollBarPadding() const;
-	SLATE_API FSlateFontInfo DetermineFont() const;
-	SLATE_API FSlateColor DetermineBackgroundColor() const;
-	SLATE_API FSlateColor DetermineForegroundColor() const;
+	FMargin DeterminePadding() const;
+	FMargin DetermineHScrollBarPadding() const;
+	FMargin DetermineVScrollBarPadding() const;
+	FSlateFontInfo DetermineFont() const;
+	FSlateColor DetermineBackgroundColor() const;
+	FSlateColor DetermineForegroundColor() const;
 
 	/** Styling: border image to draw when not hovered or focused */
 	const FSlateBrush* BorderImageNormal;
@@ -603,7 +603,7 @@ private:
 	const FSlateBrush* BorderImageReadOnly;
 
 	/** @return Border image for the text box based on the hovered and focused state */
-	SLATE_API const FSlateBrush* GetBorderImage() const;
+	const FSlateBrush* DetermineBorderImage() const;
 
 };
 
diff --git a/Engine/Source/Runtime/Slate/Public/Widgets/Text/SMultiLineEditableText.h b/Engine/Source/Runtime/Slate/Public/Widgets/Text/SMultiLineEditableText.h
--- a/Engine/Source/Runtime/Slate/Public/Widgets/Text/SMultiLineEditableText.h
+++ b/Engine/Source/Runtime/Slate/Public/Widgets/Text/SMultiLineEditableText.h
@@ -3,6 +3,7 @@
 #pragma once
 
 #include "CoreMinimal.h"
+#include "SlateEditableTextLayout.h"
 #include "Misc/Attribute.h"
 #include "SlateGlobals.h"
 #include "Layout/Margin.h"
@@ -28,7 +29,6 @@
 class FActiveTimerHandle;
 class FArrangedChildren;
 class FPaintArgs;
-class FSlateEditableTextLayout;
 class FSlateWindowElementList;
 class ITextLayoutMarshaller;
 enum class ETextShapingMethod : uint8;
@@ -48,7 +48,7 @@ public:
 	{
 	public:
 		FScopedEditableTextTransaction(TSharedPtr<SMultiLineEditableText> InText)
-			: Text(InText)
+			: Text(MoveTemp(InText))
 		{
 			Text->BeginEditTransaction();
 		}
@@ -119,7 +119,7 @@ public:
 		SLATE_ATTRIBUTE(float, WrapTextAt)
 
 		/** Whether to wrap text automatically based on the widget's computed horizontal space.  IMPORTANT: Using automatic wrapping can result
-		    in visual artifacts, as the the wrapped size will computed be at least one frame late!  Consider using WrapTextAt instead.  The initial 
+			in visual artifacts, as the computed wrapped size will be at least one frame late!  Consider using WrapTextAt instead.  The initial
 			desired size will not be clamped.  This works best in cases where the text block's size is not affecting other widget's layout. */
 		SLATE_ATTRIBUTE(bool, AutoWrapText)
 
@@ -233,7 +233,7 @@ public:
 		/** Which text flow direction should we use? (unset to use the default returned by GetDefaultTextFlowDirection) */
 		SLATE_ARGUMENT( TOptional<ETextFlowDirection>, TextFlowDirection )
 
-		/** Determines what happens to text that is clipped and doesnt fit within the clip rect for this widget */
+		/** Determines what happens to text that is clipped and doesn't fit within the clip rect for this widget */
 		SLATE_ARGUMENT(TOptional<ETextOverflowPolicy>, OverflowPolicy)
 
 		/** The paint policy to use for text if the font faces are still async loading. */
@@ -244,14 +244,13 @@ public:
 	SLATE_END_ARGS()
 
 	SLATE_API SMultiLineEditableText();
-	SLATE_API ~SMultiLineEditableText();
 
-	SLATE_API void Construct( const FArguments& InArgs );
+	SLATE_API void Construct(const FArguments& InArgs);
 
 	/**
 	 * Sets the text for this text block
 	 */
-	SLATE_API void SetText(const TAttribute< FText >& InText);
+	SLATE_API void SetText(TAttribute<FText> InText);
 
 	/**
 	 * Returns the text string
@@ -284,13 +283,13 @@ public:
 	/**
 	 * Sets the text that appears when there is no text in the text box
 	 */
-	SLATE_API void SetHintText(const TAttribute< FText >& InHintText);
+	SLATE_API void SetHintText(TAttribute<FText> InHintText);
 
 	/** Get the text that appears when there is no text in the text box */
 	SLATE_API FText GetHintText() const;
 
 	/** Set the text that is currently being searched for (if any) */
-	SLATE_API void SetSearchText(const TAttribute<FText>& InSearchText);
+	SLATE_API void SetSearchText(TAttribute<FText> InSearchText);
 
 	/** Get the text that is currently being searched for (if any) */
 	SLATE_API FText GetSearchText() const;
@@ -305,7 +304,7 @@ public:
 	SLATE_API void SetTextStyle(const FTextBlockStyle* InTextStyle);
 
 	/** See attribute Font */
-	SLATE_API void SetFont(const TAttribute< FSlateFontInfo >& InNewFont);
+	SLATE_API void SetFont(TAttribute<FSlateFontInfo> InNewFont);
 	SLATE_API FSlateFontInfo GetFont() const;
 
 	/** See TextShapingMethod attribute */
@@ -315,75 +314,81 @@ public:
 	SLATE_API void SetTextFlowDirection(const TOptional<ETextFlowDirection>& InTextFlowDirection);
 
 	/** See WrapTextAt attribute */
-	SLATE_API void SetWrapTextAt(const TAttribute<float>& InWrapTextAt);
+	SLATE_API void SetWrapTextAt(TAttribute<float> InWrapTextAt);
 
 	/** See AutoWrapText attribute */
-	SLATE_API void SetAutoWrapText(const TAttribute<bool>& InAutoWrapText);
+	SLATE_API void SetAutoWrapText(TAttribute<bool> InAutoWrapText);
 
 	/** Set WrappingPolicy attribute */
-	SLATE_API void SetWrappingPolicy(const TAttribute<ETextWrappingPolicy>& InWrappingPolicy);
+	SLATE_API void SetWrappingPolicy(TAttribute<ETextWrappingPolicy> InWrappingPolicy);
 
 	/** See LineHeightPercentage attribute */
-	SLATE_API void SetLineHeightPercentage(const TAttribute<float>& InLineHeightPercentage);
+	SLATE_API void SetLineHeightPercentage(TAttribute<float> InLineHeightPercentage);
 
 	/** See ApplyLineHeightToBottomLine attribute */
-	SLATE_API void SetApplyLineHeightToBottomLine(const TAttribute<bool>& InApplyLineHeightToBottomLine);
+	SLATE_API void SetApplyLineHeightToBottomLine(TAttribute<bool> InApplyLineHeightToBottomLine);
 
 	/** See Margin attribute */
-	SLATE_API void SetMargin(const TAttribute<FMargin>& InMargin);
+	SLATE_API void SetMargin(TAttribute<FMargin> InMargin);
 
 	/** See Justification attribute */
-	SLATE_API void SetJustification(const TAttribute<ETextJustify::Type>& InJustification);
+	SLATE_API void SetJustification(TAttribute<ETextJustify::Type> InJustification);
 
 	/** Sets the overflow policy for this text block */
-	SLATE_API void SetOverflowPolicy(TOptional<ETextOverflowPolicy> InOverflowPolicy);
+	SLATE_API void SetOverflowPolicy(const TOptional<ETextOverflowPolicy>& InOverflowPolicy);
 
 	/** See the AllowContextMenu attribute */
-	SLATE_API void SetAllowContextMenu(const TAttribute< bool >& InAllowContextMenu);
+	SLATE_API void SetAllowContextMenu(TAttribute<bool> InAllowContextMenu);
 
 	/** Set the VirtualKeyboardDismissAction attribute */
-	SLATE_API void SetVirtualKeyboardDismissAction(TAttribute< EVirtualKeyboardDismissAction > InVirtualKeyboardDismissAction);
+	SLATE_API void SetVirtualKeyboardDismissAction(TAttribute<EVirtualKeyboardDismissAction> InVirtualKeyboardDismissAction);
 
 	/**
 	 * Sets whether to select word on the mouse double click
 	 *
 	 * @param  InSelectWordOnMouseDoubleClick		Select word on the mouse double click
 	 */
-	SLATE_API void SetSelectWordOnMouseDoubleClick(const TAttribute<bool>& InSelectWordOnMouseDoubleClick);
+	SLATE_API void SetSelectWordOnMouseDoubleClick(TAttribute<bool> InSelectWordOnMouseDoubleClick);
 	
 	/** Sets the ReadOnly attribute */
-	SLATE_API void SetIsReadOnly(const TAttribute< bool >& InIsReadOnly);
+	SLATE_API void SetIsReadOnly(TAttribute<bool> InIsReadOnly);
 
 	/** Get the number of Text Lines */
-	SLATE_API int32 GetTextLineCount();
+	UE_DEPRECATED(5.8, "Use the const version instead.")
+	int32 GetTextLineCount()
+	{
+		return const_cast<const SMultiLineEditableText*>(this)->GetTextLineCount();
+	}
+	/** Get the number of Text Lines */
+	SLATE_API int32 GetTextLineCount() const;
 
 	/**
 	 * Sets whether to select all text when the user clicks to give focus on the widget
 	 *
 	 * @param  InSelectAllTextWhenFocused	Select all text when the user clicks?
 	 */
-	SLATE_API void SetSelectAllTextWhenFocused(const TAttribute<bool>& InSelectAllTextWhenFocused);
+	SLATE_API void SetSelectAllTextWhenFocused(TAttribute<bool> InSelectAllTextWhenFocused);
 
 	/**
 	 * Sets whether to clear text selection when focus is lost
 	 *
 	 * @param  InClearTextSelectionOnFocusLoss	Clear text selection when focus is lost?
 	 */
-	SLATE_API void SetClearTextSelectionOnFocusLoss(const TAttribute<bool>& InClearTextSelectionOnFocusLoss);
+	SLATE_API void SetClearTextSelectionOnFocusLoss(TAttribute<bool> InClearTextSelectionOnFocusLoss);
 
 	/**
 	 * Sets whether to allow the user to back out of changes when they press the escape key
 	 *
 	 * @param  InRevertTextOnEscape			Allow the user to back out of changes?
 	 */
-	SLATE_API void SetRevertTextOnEscape(const TAttribute<bool>& InRevertTextOnEscape);
+	SLATE_API void SetRevertTextOnEscape(TAttribute<bool> InRevertTextOnEscape);
 
 	/**
 	 * Sets whether to clear keyboard focus when pressing enter to commit changes
 	 *
 	 * @param  InClearKeyboardFocusOnCommit		Clear keyboard focus when pressing enter?
 	 */
-	SLATE_API void SetClearKeyboardFocusOnCommit(const TAttribute<bool>& InClearKeyboardFocusOnCommit);
+	SLATE_API void SetClearKeyboardFocusOnCommit(TAttribute<bool> InClearKeyboardFocusOnCommit);
 
 	/** Query to see if any text is selected within the document */
 	SLATE_API bool AnyTextSelected() const;
@@ -437,7 +442,7 @@ public:
 	/** Get the run currently under the cursor, or null if there is no run currently under the cursor */
 	SLATE_API TSharedPtr<const IRun> GetRunUnderCursor() const;
 
-	/** Get the runs currently that are current selected, some of which may be only partially selected */
+	/** Get the runs that are currently selected, some of which may be only partially selected */
 	SLATE_API TArray<TSharedRef<const IRun>> GetSelectedRuns() const;
 
 	/** Get the interaction position of the cursor (where to insert, delete, etc, text from/to) */
@@ -452,7 +457,7 @@ public:
 	/** Get the vertical scroll bar widget */
 	SLATE_API TSharedPtr<const SScrollBar> GetVScrollBar() const;
 
-	/** Refresh this editable text immediately, rather than wait for the usual caching mechanisms to take affect on the text Tick */
+	/** Refresh this editable text immediately, rather than wait for the usual caching mechanisms to take effect on the text Tick */
 	SLATE_API void Refresh();
 
 	/**
@@ -462,7 +467,7 @@ public:
 	 */
 	void SetOnKeyCharHandler(FOnKeyChar InOnKeyCharHandler)
 	{
-		OnKeyCharHandler = InOnKeyCharHandler;
+		OnKeyCharHandler = MoveTemp(InOnKeyCharHandler);
 	}
 
 	/**
@@ -472,7 +477,7 @@ public:
 	 */
 	void SetOnKeyDownHandler(FOnKeyDown InOnKeyDownHandler)
 	{
-		OnKeyDownHandler = InOnKeyDownHandler;
+		OnKeyDownHandler = MoveTemp(InOnKeyDownHandler);
 	}
 
 	/** Sets the paint policy to use for text if the font faces are still async loading. */
@@ -544,8 +549,8 @@ protected:
 	SLATE_API virtual void OnTextChanged(const FText& InText) override;
 	SLATE_API virtual void OnTextCommitted(const FText& InText, const ETextCommit::Type InTextAction) override;
 	SLATE_API virtual void OnCursorMoved(const FTextLocation& InLocation) override;
-	SLATE_API virtual float UpdateAndClampHorizontalScrollBar(const float InViewOffset, const float InViewFraction, const EVisibility InVisiblityOverride) override;
-	SLATE_API virtual float UpdateAndClampVerticalScrollBar(const float InViewOffset, const float InViewFraction, const EVisibility InVisiblityOverride) override;
+	SLATE_API virtual float UpdateAndClampHorizontalScrollBar(const float InViewOffset, const float InViewFraction, const EVisibility InVisibilityOverride) override;
+	SLATE_API virtual float UpdateAndClampVerticalScrollBar(const float InViewOffset, const float InViewFraction, const EVisibility InVisibilityOverride) override;
 	//~ End ISlateEditableTextWidget Interface
 
 protected:
@@ -556,8 +561,13 @@ protected:
 	SLATE_API void EndEditTransaction();
 
 protected:
+#if WITH_EDITORONLY_DATA
 	/** The text layout that deals with the editable text */
+	UE_DEPRECATED(5.8, "Use SlateEditableTextLayout, this variable will never be set.")
 	TUniquePtr<FSlateEditableTextLayout> EditableTextLayout;
+#endif
+	/** The text layout that deals with the editable text */
+	FSlateEditableTextLayout SlateEditableTextLayout;
 
 	/** Whether to allow multi-line text */
 	TAttribute<bool> bAllowMultiLine;
@@ -574,7 +584,7 @@ protected:
 	/** True if any changes should be reverted if we receive an escape key */
 	TAttribute<bool> bRevertTextOnEscape;
 
-	/** True if we want the text control to lose focus on an text commit/revert events */
+	/** True if we want the text control to lose focus on a text commit/revert events */
 	TAttribute<bool> bClearKeyboardFocusOnCommit;
 
 	/** Whether the context menu can be opened */
diff --git a/Engine/Source/Runtime/Slate/Public/Widgets/Text/SlateEditableTextLayout.h b/Engine/Source/Runtime/Slate/Public/Widgets/Text/SlateEditableTextLayout.h
--- a/Engine/Source/Runtime/Slate/Public/Widgets/Text/SlateEditableTextLayout.h
+++ b/Engine/Source/Runtime/Slate/Public/Widgets/Text/SlateEditableTextLayout.h
@@ -36,16 +36,51 @@ enum class EFontFacesLoadingPaintPolicy : uint8;
 class FSlateEditableTextLayout
 {
 public:
-	SLATE_API FSlateEditableTextLayout(ISlateEditableTextWidget& InOwnerWidget, const TAttribute<FText>& InInitialText, FTextBlockStyle InTextStyle, const TOptional<ETextShapingMethod> InTextShapingMethod, const TOptional<ETextFlowDirection> InTextFlowDirection, const FCreateSlateTextLayout& InCreateSlateTextLayout, TSharedRef<ITextLayoutMarshaller> InTextMarshaller, TSharedRef<ITextLayoutMarshaller> InHintTextMarshaller);
+
+	template<typename RealWidgetType>
+	UE_DEPRECATED(5.8, "Use the other base constructor and call Construct during widget construction")
+	FSlateEditableTextLayout(RealWidgetType& InOwnerWidget, TAttribute<FText> InInitialText, FTextBlockStyle InTextStyle, const TOptional<ETextShapingMethod> InTextShapingMethod, const TOptional<ETextFlowDirection> InTextFlowDirection, const FCreateSlateTextLayout& InCreateSlateTextLayout, TSharedRef<ITextLayoutMarshaller> InTextMarshaller, TSharedRef<ITextLayoutMarshaller> InHintTextMarshaller)
+		: FSlateEditableTextLayout(InOwnerWidget)
+	{
+		Construct(MoveTemp(InInitialText), MoveTemp(InTextStyle), InTextShapingMethod, InTextFlowDirection, InCreateSlateTextLayout, MoveTemp(InTextMarshaller), MoveTemp(InHintTextMarshaller));
+	}
+
+
+	template<typename RealWidgetType>
+	FSlateEditableTextLayout(RealWidgetType& InOwnerWidget)
+		: OwnerWidget(&InOwnerWidget)
+		, bWasPasswordLastTick(false)
+		, SearchCase(ESearchCase::IgnoreCase)
+		, CurrentSearchResultIndex(0)
+		, bHasRegisteredTextInputMethodContext(false)
+		, PreferredCursorScreenOffsetInLine(0)
+		, CurrentUndoLevel(0)
+		, NumTransactionsOpened(0)
+		, bIsDragSelecting(false)
+		, bWasFocusedByLastMouseDown(false)
+		, bHasDragSelectedSinceFocused(false)
+		, bTextChangedByVirtualKeyboard(false)
+		, bTextCommittedByVirtualKeyboard(false)
+		, VirtualKeyboardTextCommitType()
+		, bSelectionChangedExternally(false)
+		, ExternalSelectionStart(0)
+		, ExternalSelectionEnd(0)
+	{
+		static_assert(std::is_base_of_v<SWidget, RealWidgetType>, "RealWidgetType must be a Widget type.");
+		static_assert(std::is_base_of_v<ISlateEditableTextWidget, RealWidgetType>, "RealWidgetType must derive from ISlateEditableTextWidget.");
+	}
+
 	SLATE_API ~FSlateEditableTextLayout();
 
-	SLATE_API void SetText(const TAttribute<FText>& InText);
+	SLATE_API void Construct(TAttribute<FText> InInitialText, FTextBlockStyle InTextStyle, const TOptional<ETextShapingMethod>& InTextShapingMethod, const TOptional<ETextFlowDirection>& InTextFlowDirection, const FCreateSlateTextLayout& InCreateSlateTextLayout, TSharedRef<ITextLayoutMarshaller> InTextMarshaller, TSharedRef<ITextLayoutMarshaller> InHintTextMarshaller);
+
+	SLATE_API void SetText(TAttribute<FText> InText);
 	SLATE_API FText GetText() const;
 
-	SLATE_API void SetHintText(const TAttribute<FText>& InHintText);
+	SLATE_API void SetHintText(TAttribute<FText> InHintText);
 	SLATE_API FText GetHintText() const;
 
-	SLATE_API void SetSearchText(const TAttribute<FText>& InSearchText);
+	SLATE_API void SetSearchText(TAttribute<FText> InSearchText);
 	SLATE_API FText GetSearchText() const;
 	
 	/** Get the index of the search result (0 if none) */
@@ -57,11 +92,17 @@ public:
 	SLATE_API void SetTextStyle(const FTextBlockStyle& InTextStyle);
 	SLATE_API const FTextBlockStyle& GetTextStyle() const;
 	
-	SLATE_API void SetHintTextOpacity(const TAttribute<float>& InOpacity);
+	SLATE_API void SetHintTextOpacity(TAttribute<float> InOpacity);
 	SLATE_API float GetHintTextOpacity() const;
 
 	/** Get the number of Text Lines */
-	SLATE_API int32 GetTextLineCount();
+	UE_DEPRECATED(5.8, "Use the const version instead.")
+	int32 GetTextLineCount()
+	{
+		return const_cast<const FSlateEditableTextLayout*>(this)->GetTextLineCount();
+	}
+	/** Get the number of Text Lines */
+	SLATE_API int32 GetTextLineCount() const;
 
 	/** Set the brush to use when drawing the cursor */
 	SLATE_API void SetCursorBrush(const TAttribute<const FSlateBrush*>& InCursorBrush);
@@ -102,31 +143,33 @@ public:
 	SLATE_API void SetTextFlowDirection(const TOptional<ETextFlowDirection>& InTextFlowDirection);
 
 	/** Set the wrapping to use for this document */
-	SLATE_API void SetTextWrapping(const TAttribute<float>& InWrapTextAt, const TAttribute<bool>& InAutoWrapText, const TAttribute<ETextWrappingPolicy>& InWrappingPolicy);
+	SLATE_API void SetTextWrapping(TAttribute<float> InWrapTextAt, TAttribute<bool> InAutoWrapText, TAttribute<ETextWrappingPolicy> InWrappingPolicy);
 
 	/** Set whether text wraps onto a new line when it's length exceeds this width; if this value is zero or negative, no wrapping occurs */
-	SLATE_API void SetWrapTextAt(const TAttribute<float>& InWrapTextAt);
+	SLATE_API void SetWrapTextAt(TAttribute<float> InWrapTextAt);
 
 	/** Set whether to wrap text automatically based on the widget's computed horizontal space */
-	SLATE_API void SetAutoWrapText(const TAttribute<bool>& InAutoWrapText);
+	SLATE_API void SetAutoWrapText(TAttribute<bool> InAutoWrapText);
 
 	/** Set the wrapping policy to use */
-	SLATE_API void SetWrappingPolicy(const TAttribute<ETextWrappingPolicy>& InWrappingPolicy);
+	SLATE_API void SetWrappingPolicy(TAttribute<ETextWrappingPolicy> InWrappingPolicy);
 
 	/** Set the amount of blank space left around the edges of text area */
-	SLATE_API void SetMargin(const TAttribute<FMargin>& InMargin);
+	SLATE_API void SetMargin(TAttribute<FMargin> InMargin);
 
 	/** Set how the text should be aligned with the margin */
-	SLATE_API void SetJustification(const TAttribute<ETextJustify::Type>& InJustification);
+	SLATE_API void SetJustification(TAttribute<ETextJustify::Type> InJustification);
+	/** Get how the text should be aligned with the margin */
+	SLATE_API ETextJustify::Type GetJustification() const;
 
 	/** Set the amount to scale each lines height by */
-	SLATE_API void SetLineHeightPercentage(const TAttribute<float>& InLineHeightPercentage);
+	SLATE_API void SetLineHeightPercentage(TAttribute<float> InLineHeightPercentage);
 
 	/** Set whether to leave extra space below the last line due to line height */
-	SLATE_API void SetApplyLineHeightToBottomLine(const TAttribute<bool>& InApplyLineHeightToBottomLine);
+	SLATE_API void SetApplyLineHeightToBottomLine(TAttribute<bool> InApplyLineHeightToBottomLine);
 
 	/** Set the text overflow policy that should be used to determine what happens to clipped text */
-	SLATE_API void SetOverflowPolicy(TOptional<ETextOverflowPolicy> InOverflowPolicy);
+	SLATE_API void SetOverflowPolicy(const TOptional<ETextOverflowPolicy>& InOverflowPolicy);
 
 	/** Set the information used to help identify who owns this text layout in the case of an error */
 	SLATE_API void SetDebugSourceInfo(const TAttribute<FString>& InDebugSourceInfo);
@@ -140,7 +183,7 @@ public:
 	/** Register and activate the IME context for this text layout */
 	SLATE_API void EnableTextInputMethodContext();
 
-	/** Refresh this editable text immediately, rather than wait for the usual caching mechanisms to take affect on the text Tick */
+	/** Refresh this editable text immediately, rather than wait for the usual caching mechanisms to take effect on the text Tick */
 	SLATE_API bool Refresh();
 
 	/** Force the text layout to be updated from the marshaller */
@@ -311,7 +354,7 @@ public:
 	/**
 	 * Given a location and a Direction to offset, return a new location.
 	 *
-	 * @param Location    Cursor location from which to offset
+	 * @param CurrentLocation    Cursor location from which to offset
 	 * @param Direction   Positive means right, negative means left.
 	 */
 	SLATE_API FTextLocation TranslatedLocation(const FTextLocation& CurrentLocation, int8 Direction) const;
@@ -354,8 +397,13 @@ public:
 	/** Returns whether the current text varies from the original */
 	SLATE_API bool HasTextChangedFromOriginal() const;
 
+	UE_DEPRECATED(5.8, "Use the properly spelled version BeginEditTransaction")
+	void BeginEditTransation()
+	{
+		BeginEditTransaction();
+	}
 	/** Called to begin an undoable editable text transaction */
-	SLATE_API void BeginEditTransation();
+	SLATE_API void BeginEditTransaction();
 
 	/** Called to end an undoable editable text transaction */
 	SLATE_API void EndEditTransaction();
@@ -389,7 +437,12 @@ public:
 
 	SLATE_API void Tick(const FGeometry& AllottedGeometry, const double InCurrentTime, const float InDeltaTime);
 
-	SLATE_API int32 OnPaint(const FPaintArgs& Args, const FGeometry& AllottedGeometry, const FSlateRect& MyCullingRect, FSlateWindowElementList& OutDrawElements, int32 LayerId, const FWidgetStyle& InWidgetStyle, bool bParentEnabled);
+	UE_DEPRECATED(5.8, "Use the const version instead.")
+	int32 OnPaint(const FPaintArgs& Args, const FGeometry& AllottedGeometry, const FSlateRect& MyCullingRect, FSlateWindowElementList& OutDrawElements, int32 LayerId, const FWidgetStyle& InWidgetStyle, bool bParentEnabled)
+	{
+		return const_cast<const FSlateEditableTextLayout*>(this)->OnPaint(Args, AllottedGeometry, MyCullingRect, OutDrawElements, LayerId, InWidgetStyle, bParentEnabled);
+	}
+	SLATE_API int32 OnPaint(const FPaintArgs& Args, const FGeometry& AllottedGeometry, const FSlateRect& MyCullingRect, FSlateWindowElementList& OutDrawElements, int32 LayerId, const FWidgetStyle& InWidgetStyle, bool bParentEnabled) const;
 
 	SLATE_API void CacheDesiredSize(float LayoutScaleMultiplier);
 
@@ -430,19 +483,19 @@ public:
 	SLATE_API void SetFontFacesLoadingPaintPolicy(EFontFacesLoadingPaintPolicy InFontFacesLoadingPaintPolicy);
 private:
 	/** Insert the given text at the current cursor position, correctly taking into account new line characters */
-	SLATE_API void InsertTextAtCursorImpl(const FString& InString);
+	void InsertTextAtCursorImpl(const FString& InString);
 
 	/** Insert a new-line at the current cursor position */
-	SLATE_API void InsertNewLineAtCursorImpl();
+	void InsertNewLineAtCursorImpl();
 
 	/** Implementation of Refresh that actually updates the layout. Optionally takes text to set, or will use the current editable text if none if provided */
-	SLATE_API bool RefreshImpl(const FText* InTextToSet, const bool bForce = false);
+	bool RefreshImpl(const FText* InTextToSet, const bool bForce = false);
 
 	/** Create a text or password run using the given text and style */
-	SLATE_API TSharedRef<IRun> CreateTextOrPasswordRun(const FRunInfo& InRunInfo, const TSharedRef<const FString>& InText, const FTextBlockStyle& InStyle);
+	TSharedRef<IRun> CreateTextOrPasswordRun(const FRunInfo& InRunInfo, const TSharedRef<const FString>& InText, const FTextBlockStyle& InStyle);
 
 	/** Called when the active context menu is closed */
-	SLATE_API void OnContextMenuClosed(TSharedRef<IMenu> Menu);
+	void OnContextMenuClosed(TSharedRef<IMenu> Menu);
 	
 	void UpdateTextChangedByVirtualKeyboard();
 
@@ -556,7 +609,7 @@ private:
 	FTextBlockStyle TextStyle;
 
 	/** Style used to draw the hint text (only valid when HintTextLayout is set) */
-	FTextBlockStyle HintTextStyle;
+	mutable FTextBlockStyle HintTextStyle;
 
 	/** The text displayed in this text block */
 	TAttribute<FText> BoundText;
@@ -569,8 +622,8 @@ private:
 
 	/** The text that appears when there is no text in the text box */
 	TAttribute<FText> HintText;
-	
-	/** Override opacity of the hint text */
+
+	/** Opacity of the hint text */
 	TAttribute<float> HintTextOpacity;
 
 	/** The text to be searched for */
@@ -663,7 +716,7 @@ private:
 	/** Current undo state level that we've rolled back to, or INDEX_NONE if we haven't undone. Used for 'Redo'. */
 	int32 CurrentUndoLevel;
 
-	/** Undo state that will be pushed if text is actually changed between calls to BeginEditTransation() and EndEditTransaction() */
+	/** Undo state that will be pushed if text is actually changed between calls to BeginEditTransaction() and EndEditTransaction() */
 	TOptional<SlateEditableTextTypes::FUndoState> StateBeforeChangingText;
 
 	/** Track the number transactions that are opened */
@@ -697,7 +750,7 @@ private:
 	TOptional<ETextOverflowPolicy> OverflowPolicyOverride;
 
 	/** The last known size of the widget from the previous OnPaint, used to recalculate wrapping */
-	FVector2f CachedSize;
+	mutable FVector2f CachedSize;
 
 	/** A list commands to execute if a user presses the corresponding key-binding in the text box */
 	TSharedPtr<FUICommandList> UICommandList;
