PR #14068: Fix scrollbar track visibility

Additionally deleted unnecessary boolean `SScrollBar::bAlwaysShowScrollbarTrack` and solely rely on `SScrollBarTrack::bIsAlwaysVisible`

#pr 14068

#Synced-CL 49956667

#Synced-CL 49956667

diff --git a/Engine/Source/Runtime/Slate/Private/Widgets/Layout/SScrollBar.cpp b/Engine/Source/Runtime/Slate/Private/Widgets/Layout/SScrollBar.cpp
--- a/Engine/Source/Runtime/Slate/Private/Widgets/Layout/SScrollBar.cpp
+++ b/Engine/Source/Runtime/Slate/Private/Widgets/Layout/SScrollBar.cpp
@@ -105,7 +105,7 @@ void SScrollBar::Construct(const FArguments& InArgs)
 	SetEnabled(TAttribute<bool>( Track.ToSharedRef(), &SScrollBarTrack::IsNeeded ));
 	SetVisibility(TAttribute<EVisibility>(SharedThis(this), &SScrollBar::ShouldBeVisible));
 	SetScrollBarAlwaysVisible(InArgs._AlwaysShowScrollbar);
-	bAlwaysShowScrollbarTrack = InArgs._AlwaysShowScrollbarTrack;
+	SetScrollBarTrackAlwaysVisible(InArgs._AlwaysShowScrollbarTrack);
 }
 
 void SScrollBar::SetOnUserScrolled( const FOnUserScrolled& InHandler )
@@ -269,7 +269,7 @@ FSlateColor SScrollBar::GetTrackOpacity() const
 	{
 		return FLinearColor(1,1,1,1);
 	}
-	else if ((bAlwaysShowScrollbarTrack  && Track->GetThumbSizeFraction() > KINDA_SMALL_NUMBER) || AlwaysShowScrollbar())
+	else if ((Track->GetIsAlwaysVisible()  && Track->GetThumbSizeFraction() > KINDA_SMALL_NUMBER) || AlwaysShowScrollbar())
 	{
 		return FLinearColor(1, 1, 1, 0.5f);
 	}
@@ -436,13 +436,11 @@ void SScrollBar::SetThickness(TAttribute<FVector2D> InThickness)
 void SScrollBar::SetScrollBarAlwaysVisible(bool InAlwaysVisible)
 {
 	bAlwaysShowScrollbar = InAlwaysVisible;
-
-	Track->SetIsAlwaysVisible(InAlwaysVisible);
 }
 
 void SScrollBar::SetScrollBarTrackAlwaysVisible(bool InAlwaysVisible)
 {
-	bAlwaysShowScrollbarTrack = InAlwaysVisible;
+	Track->SetIsAlwaysVisible(InAlwaysVisible);
 }
 
 void SScrollBar::SetScrollbarDisabledVisibility(EVisibility InVisibility)
diff --git a/Engine/Source/Runtime/Slate/Private/Widgets/Layout/SScrollBarTrack.cpp b/Engine/Source/Runtime/Slate/Private/Widgets/Layout/SScrollBarTrack.cpp
--- a/Engine/Source/Runtime/Slate/Private/Widgets/Layout/SScrollBarTrack.cpp
+++ b/Engine/Source/Runtime/Slate/Private/Widgets/Layout/SScrollBarTrack.cpp
@@ -149,7 +149,12 @@ bool SScrollBarTrack::IsNeeded() const
 {
 	// We use a small epsilon here to avoid the scroll bar showing up when all of the content is already in view, due to
 	// floating point precision when the scroll bar state is set
-	return ThumbSizeFraction < (1.0f - KINDA_SMALL_NUMBER) || bIsAlwaysVisible;
+	return ThumbSizeFraction < (1.0f - KINDA_SMALL_NUMBER) || GetIsAlwaysVisible();
+}
+
+bool SScrollBarTrack::GetIsAlwaysVisible() const
+{
+	return bIsAlwaysVisible;
 }
 
 float SScrollBarTrack::DistanceFromTop() const
diff --git a/Engine/Source/Runtime/Slate/Public/Widgets/Layout/SScrollBar.h b/Engine/Source/Runtime/Slate/Public/Widgets/Layout/SScrollBar.h
--- a/Engine/Source/Runtime/Slate/Public/Widgets/Layout/SScrollBar.h
+++ b/Engine/Source/Runtime/Slate/Public/Widgets/Layout/SScrollBar.h
@@ -212,7 +212,6 @@ protected:
 	float DragGrabOffset;
 	EOrientation Orientation;
 	bool bAlwaysShowScrollbar;
-	bool bAlwaysShowScrollbarTrack;
 	EFocusCause DragFocusCause;
 	bool bHideWhenNotInUse;
 	EVisibility ScrollbarDisabledVisibility = EVisibility::Collapsed;
diff --git a/Engine/Source/Runtime/Slate/Public/Widgets/Layout/SScrollBarTrack.h b/Engine/Source/Runtime/Slate/Public/Widgets/Layout/SScrollBarTrack.h
--- a/Engine/Source/Runtime/Slate/Public/Widgets/Layout/SScrollBarTrack.h
+++ b/Engine/Source/Runtime/Slate/Public/Widgets/Layout/SScrollBarTrack.h
@@ -104,6 +104,8 @@ public:
 
 	SLATE_API bool IsNeeded() const;
 
+	SLATE_API bool GetIsAlwaysVisible() const;
+
 	SLATE_API float DistanceFromTop() const;
 
 	SLATE_API float DistanceFromBottom() const;
