Fixed off-screen rendered Children of Children inside ScrollBoxes being unnavigable.
- ScrollPanel defines it's AllottedGeometry to be it's full content size if it uses a FNavigationMethod. SScrollBox overrides the HitTestBoundingBox with it's content size.

#Synced-CL 50207111

#Synced-CL 50207111

diff --git a/Engine/Source/Runtime/Slate/Private/Widgets/Layout/SScrollBox.cpp b/Engine/Source/Runtime/Slate/Private/Widgets/Layout/SScrollBox.cpp
--- a/Engine/Source/Runtime/Slate/Private/Widgets/Layout/SScrollBox.cpp
+++ b/Engine/Source/Runtime/Slate/Private/Widgets/Layout/SScrollBox.cpp
@@ -2,6 +2,7 @@
 
 #include "Widgets/Layout/SScrollBox.h"
 #include "Rendering/DrawElements.h"
+#include "Types/NavigationMetaData.h"
 #include "Types/SlateConstants.h"
 #include "Layout/LayoutUtils.h"
 #include "Widgets/SOverlay.h"
@@ -10,6 +11,37 @@
 
 #include UE_INLINE_GENERATED_CPP_BY_NAME(SScrollBox)
 
+namespace
+{
+static void AppendFocusableChildren(TArray<TSharedRef<SWidget>>& OutChildren, TSharedRef<SWidget> InWidget, TSharedPtr<SWidget>& OutFocusedWidget, uint32 InUserIndex)
+{
+	FChildren* InWidgetChildren = InWidget->GetChildren();
+	if (InWidgetChildren)
+	{
+		for (int32 SlotIndex = 0; SlotIndex < InWidgetChildren->Num(); ++SlotIndex)
+		{
+			TSharedRef<SWidget> Child = InWidgetChildren->GetChildAt(SlotIndex);
+			if (Child->GetVisibility().IsVisible())
+			{
+				if (Child->SupportsKeyboardFocus())
+				{
+					OutChildren.Add(Child);
+					const bool bHasUserFocus = Child->HasUserFocus(InUserIndex).IsSet() || Child->HasUserFocusedDescendants(InUserIndex);
+					if (bHasUserFocus)
+					{
+						ensureMsgf(!OutFocusedWidget.IsValid(), TEXT("This should only be set once"));
+						OutFocusedWidget = Child.ToSharedPtr();
+					}
+				}
+				else
+				{
+					AppendFocusableChildren(OutChildren, Child, OutFocusedWidget, InUserIndex);
+				}
+			}
+		}
+	}
+}
+}
 
 void SScrollBox::FSlot::Construct(const FChildren& SlotOwner, FSlotArguments&& InArgs)
 {
@@ -35,8 +67,9 @@ void SScrollPanel::PrivateRegisterAttributes(FSlateAttributeInitializer& Attribu
 	SScrollBox::FSlot::RegisterAttributes(Initializer);
 }
 
-void SScrollPanel::Construct(const FArguments& InArgs, TArray<SScrollBox::FSlot::FSlotArguments> InSlots)
+void SScrollPanel::Construct(const FArguments& InArgs, TArray<SScrollBox::FSlot::FSlotArguments> InSlots, TSharedRef<SScrollBox> InScrollBox)
 {
+	OwningScrollBox = InScrollBox;
 	PhysicalOffset = 0;
 	Children.AddSlots(MoveTemp(InSlots));
 	Orientation = InArgs._Orientation;
@@ -60,6 +93,118 @@ void SScrollPanel::OnArrangeChildren(const FGeometry& AllottedGeometry, FArrange
 	}
 }
 
+
+SScrollPanel::FArrangedWidgetsCache SScrollPanel::GetArrangedWidgetsForFocus(const FGeometry& MyGeometry, uint32 InUserIndex) const
+{
+	FArrangedWidgetsCache ReturnValue;
+
+	ReturnValue.UserIndex = InUserIndex;
+
+	// Find the Focused Child Widget, if it exists
+	AppendFocusableChildren(ReturnValue.Children, const_cast<SScrollPanel*>(this)->AsShared(), ReturnValue.UserFocusedWidget, InUserIndex);
+	if (ReturnValue.UserFocusedWidget.IsValid())
+	{
+		// Get FArrangedWidget for `UserFocusedWidget`
+		TMap<TSharedRef<SWidget>, FArrangedWidget> FocusToArranged;
+		if (FindChildGeometries(MyGeometry, { ReturnValue.UserFocusedWidget.ToSharedRef() }, FocusToArranged))
+		{
+			for (const TPair<TSharedRef<SWidget>, FArrangedWidget>& It : FocusToArranged)
+			{
+				ReturnValue.WidgetsToArranged.Emplace(It.Key, It.Value);
+			}
+		}
+	}
+
+	// Get FArrangedWidget for `ScrollBox`
+	TSharedPtr<SScrollBox> PinnedScrollBox = OwningScrollBox.IsValid() ? OwningScrollBox.Pin() : nullptr;
+	TSharedPtr<SWidget> ScrollBoxParent = PinnedScrollBox.IsValid() ? PinnedScrollBox->GetParentWidget() : nullptr;
+	if (ScrollBoxParent.IsValid())
+	{
+		FArrangedChildren ArrangedChildren(EVisibility::Visible);
+		ScrollBoxParent->ArrangeChildren(MyGeometry, ArrangedChildren);
+
+		// Get FArrangedWidget for `ScrollBox`
+		const TSharedRef<const SWidget> SelfRef = PinnedScrollBox.ToSharedRef();
+		if (FArrangedWidget* FoundSelfArranged = ArrangedChildren.GetInternalArray().FindByPredicate([SelfRef](const FArrangedWidget& WidgetIt) { return SelfRef == WidgetIt.Widget; }))
+		{
+			ReturnValue.WidgetsToArranged.Emplace(FoundSelfArranged->Widget, *FoundSelfArranged);
+		}
+	}
+
+	return ReturnValue;
+}
+
+TSharedPtr<SWidget> SScrollPanel::GetNextWidgetForFocus(
+	const FGeometry& MyGeometry,
+	const EUINavigation InNavigation,
+	const FArrangedWidgetsCache& InArrangedWidgets,
+	FNavigationMethod* InNavigationMethod,
+	TArray<FNavigationMethod::FDebugWidgetResult>& OutDebugResults) const
+{
+	TSharedPtr<SWidget> ReturnWidget;
+
+	FHittestGrid HitTestGrid;
+
+	if (InArrangedWidgets.UserFocusedWidget.IsValid())
+	{
+		const TSharedPtr<FNavigationMetaData> FocusWidgetNavMetaData = InArrangedWidgets.UserFocusedWidget->GetMetaData<FNavigationMetaData>();
+		const TSharedPtr<SWidget> ExplicitNextWidget = FocusWidgetNavMetaData.IsValid() ? FocusWidgetNavMetaData->GetFocusRecipient(InNavigation).Pin() : nullptr;
+		if (ExplicitNextWidget.IsValid())
+		{
+			return ExplicitNextWidget;
+		}
+	}
+
+	const TSharedPtr<const SScrollBox> OwningScrollBoxWidget = OwningScrollBox.IsValid() ? OwningScrollBox.Pin() : nullptr;
+
+	// If we have the FArrangedWidget for the focused widget and the owning scroll box 
+	const bool bFoundAllArrangedWidgets = OwningScrollBoxWidget.IsValid() && InArrangedWidgets.UserFocusedWidget.IsValid() &&
+		InArrangedWidgets.WidgetsToArranged.Contains(InArrangedWidgets.UserFocusedWidget.ToSharedRef()) && InArrangedWidgets.WidgetsToArranged.Contains(OwningScrollBoxWidget.ToSharedRef());
+	if (InNavigationMethod && bFoundAllArrangedWidgets)
+	{
+		// Use `this` widget's navigation rule to navigate
+		const FGeometry& ScrollBoxGeometry = OwningScrollBoxWidget->GetTickSpaceGeometry();
+		const FSlateRect Rect = ScrollBoxGeometry.GetRenderBoundingRect();
+
+		HitTestGrid.SetHittestArea(Rect.GetTopLeft(), OwningScrollBoxWidget->GetScrollPanelContentSize());
+
+		for (const TSharedRef<SWidget>& ScrollBoxChild : InArrangedWidgets.Children)
+		{
+			HitTestGrid.AddWidget(&ScrollBoxChild.Get(), 0, 0, FSlateInvalidationWidgetSortOrder());
+		}
+
+		check(InArrangedWidgets.UserFocusedWidget.IsValid());
+		InNavigationMethod->Initialize(&HitTestGrid, nullptr);
+			
+		const FNavigationReply FocusedNavigationReply = InArrangedWidgets.UserFocusedWidget->OnNavigation(MyGeometry, InNavigation);
+		ReturnWidget = InNavigationMethod->FindNextFocusableWidget(
+			InArrangedWidgets.WidgetsToArranged[InArrangedWidgets.UserFocusedWidget.ToSharedRef()], InNavigation, FocusedNavigationReply, InArrangedWidgets.WidgetsToArranged[OwningScrollBoxWidget.ToSharedRef()], InArrangedWidgets.UserIndex);
+	}
+
+	return ReturnWidget;
+}
+int32 SScrollPanel::OnPaint(const FPaintArgs& Args, const FGeometry& AllottedGeometry, const FSlateRect& MyCullingRect, FSlateWindowElementList& OutDrawElements, int32 LayerId, const FWidgetStyle& InWidgetStyle, bool bParentEnabled) const
+{
+	TSharedPtr<FNavigationMetaData> NavMetaData = OwningScrollBox.IsValid() ? OwningScrollBox.Pin()->GetMetaData<FNavigationMetaData>() : nullptr;
+	TInstancedStruct<FNavigationMethod> EmptyInstance;
+	TInstancedStruct<FNavigationMethod>& NavigationMethodInstance = NavMetaData.IsValid() ? NavMetaData->GetMutableNavigationMethodStruct() : EmptyInstance;
+	FNavigationMethod* NavigationMethod = NavigationMethodInstance.IsValid() ? NavigationMethodInstance.GetMutablePtr() : nullptr;
+
+	if (NavigationMethod)
+	{
+		// Use the content size of the scrollbox to navigate to offscreen rendered widgets
+		const FVector2D ScrollPanelSize = FVector2D::Max(AllottedGeometry.Size, GetArrangedSize());
+		const FGeometry ScrollBoxGeometry = AllottedGeometry.MakeChild(ScrollPanelSize, FSlateLayoutTransform());
+		const FSlateRect BoundingRect =  OwningScrollBox.IsValid() ? OwningScrollBox.Pin()->GetHitTestBoundingRect() : GetHitTestBoundingRect();
+
+		return SPanel::OnPaint(Args, ScrollBoxGeometry, BoundingRect, OutDrawElements, LayerId, InWidgetStyle, bParentEnabled);
+	}
+	else
+	{
+		return SPanel::OnPaint(Args, AllottedGeometry, MyCullingRect, OutDrawElements, LayerId, InWidgetStyle, bParentEnabled);
+	}
+}
+
 FVector2D SScrollPanel::ComputeDesiredSize(float) const
 {
 	FVector2D ThisDesiredSize = FVector2D::ZeroVector;
@@ -161,7 +306,7 @@ void SScrollBox::Construct( const FArguments& InArgs )
 		bScrollBarIsExternal = false;
 	}
 
-	SAssignNew(ScrollPanel, SScrollPanel, MoveTemp(const_cast<TArray<FSlot::FSlotArguments>&>(InArgs._Slots)))
+	SAssignNew(ScrollPanel, SScrollPanel, MoveTemp(const_cast<TArray<FSlot::FSlotArguments>&>(InArgs._Slots)), StaticCastSharedRef<SScrollBox>(AsShared()))
 		.Clipping(InArgs._Clipping)
 		.Orientation(Orientation)
 		.BackPadScrolling(BackPadScrolling)
@@ -1132,6 +1277,30 @@ void SScrollBox::OnMouseCaptureLost(const FCaptureLostEvent& CaptureLostEvent)
 }
 
 FNavigationReply SScrollBox::OnNavigation(const FGeometry& MyGeometry, const FNavigationEvent& InNavigationEvent)
+{
+	TSharedPtr<FNavigationMetaData> NavMetaData = GetMetaData<FNavigationMetaData>();
+	TInstancedStruct<FNavigationMethod> EmptyInstance;
+	TInstancedStruct<FNavigationMethod>& NavigationMethodInstance = NavMetaData.IsValid() ? NavMetaData->GetMutableNavigationMethodStruct() : EmptyInstance;
+	FNavigationMethod* NavigationMethod = NavigationMethodInstance.IsValid() ? NavigationMethodInstance.GetMutablePtr() : nullptr;
+
+	if (NavigationMethod)
+	{
+		SScrollPanel::FArrangedWidgetsCache ArrangedWidgets = ScrollPanel->GetArrangedWidgetsForFocus(MyGeometry, InNavigationEvent.GetUserIndex());
+		TArray<FNavigationMethod::FDebugWidgetResult> DebugResults;
+		const EUINavigation NavigationType = InNavigationEvent.GetNavigationType();
+
+		TSharedPtr<SWidget> NextFocusableChild = ScrollPanel->GetNextWidgetForFocus(MyGeometry, NavigationType, ArrangedWidgets, NavigationMethod, DebugResults);
+		if (NextFocusableChild.IsValid())
+		{
+			InternalScrollDescendantIntoView(MyGeometry, NextFocusableChild, false, NavigationDestination, NavigationScrollPadding);
+			return FNavigationReply::Explicit(NextFocusableChild);
+		}
+	}
+
+	return OnNavigation_Default(MyGeometry, InNavigationEvent);
+}
+
+FNavigationReply SScrollBox::OnNavigation_Default(const FGeometry& MyGeometry, const FNavigationEvent& InNavigationEvent)
 {
 	TSharedPtr<SWidget> FocusedChild;
 	int32 FocusedChildIndex = -1;
@@ -1214,6 +1383,29 @@ FNavigationReply SScrollBox::OnNavigation(const FGeometry& MyGeometry, const FNa
 	return SCompoundWidget::OnNavigation(MyGeometry, InNavigationEvent);
 }
 
+FSlateRect SScrollBox::GetHitTestBoundingRect() const
+{
+	const FGeometry WidgetGeo = GetPaintSpaceGeometry();
+	FSlateRect RenderRect = WidgetGeo.GetLayoutBoundingRect();
+
+	const FVector2f ContentSize = GetScrollPanelContentSize();
+	const FVector2f TransformedContentSize = WidgetGeo.GetAccumulatedLayoutTransform().TransformVector(ContentSize);
+
+	RenderRect.Bottom = RenderRect.Top + TransformedContentSize.Y;
+	RenderRect.Right = RenderRect.Left + TransformedContentSize.X;
+
+	if (Orientation == EOrientation::Orient_Vertical)
+	{
+		RenderRect.Top -= DesiredScrollOffset;
+	}
+	else
+	{
+		RenderRect.Left -= DesiredScrollOffset;
+	}
+
+	return RenderRect;
+}
+
 FReply SScrollBox::OnFocusReceived(const FGeometry& MyGeometry, const FFocusEvent& InFocusEvent)
 {
 	if (OnScrollBoxFocusReceived.IsBound())
diff --git a/Engine/Source/Runtime/Slate/Public/Widgets/Layout/SScrollBox.h b/Engine/Source/Runtime/Slate/Public/Widgets/Layout/SScrollBox.h
--- a/Engine/Source/Runtime/Slate/Public/Widgets/Layout/SScrollBox.h
+++ b/Engine/Source/Runtime/Slate/Public/Widgets/Layout/SScrollBox.h
@@ -9,6 +9,7 @@
 #include "Styling/SlateColor.h"
 #include "Input/CursorReply.h"
 #include "Input/Reply.h"
+#include "Input/NavigationMethod.h"
 #include "Input/NavigationReply.h"
 #include "Widgets/SWidget.h"
 #include "Widgets/SPanel.h"
@@ -394,8 +395,12 @@ public:
 	SLATE_API virtual FReply OnFocusReceived(const FGeometry& MyGeometry, const FFocusEvent& InFocusEvent) override;
 	SLATE_API virtual void OnFocusLost(const FFocusEvent& InFocusEvent) override;
 	virtual bool SupportsKeyboardFocus() const override { return bIsFocusable; }
+	SLATE_API virtual FSlateRect GetHitTestBoundingRect() const override;
 	// End of SWidget interface
 
+private:
+	FNavigationReply OnNavigation_Default(const FGeometry& MyGeometry, const FNavigationEvent& InNavigationEvent);
+
 protected:
 	SLATE_API void OnClippingChanged();
 
@@ -645,7 +650,7 @@ public:
 	{
 	}
 
-	void Construct(const FArguments& InArgs, TArray<SScrollBox::FSlot::FSlotArguments> InSlots);
+	void Construct(const FArguments& InArgs, TArray<SScrollBox::FSlot::FSlotArguments> InSlots, TSharedRef<SScrollBox> InScrollBox);
 
 public:
 
@@ -661,6 +666,7 @@ public:
 	}
 
 	virtual void OnArrangeChildren(const FGeometry& AllottedGeometry, FArrangedChildren& ArrangedChildren) const override;
+	virtual int32 OnPaint(const FPaintArgs& Args, const FGeometry& AllottedGeometry, const FSlateRect& MyCullingRect, FSlateWindowElementList& OutDrawElements, int32 LayerId, const FWidgetStyle& InWidgetStyle, bool bParentEnabled) const override;
 
 	virtual FChildren* GetChildren() override
 	{
@@ -675,13 +681,32 @@ public:
 	float PhysicalOffset;
 	TPanelChildren<SScrollBox::FSlot> Children;
 
+public:
+	struct FArrangedWidgetsCache
+	{
+		uint32 UserIndex;
+		TArray<TSharedRef<SWidget>> Children;
+		TSharedPtr<SWidget> UserFocusedWidget;
+		TMap<TSharedRef<const SWidget>, FArrangedWidget> WidgetsToArranged;
+	};
+
+public:
+	TSharedPtr<SWidget> GetNextWidgetForFocus(
+		const FGeometry& MyGeometry,
+		const EUINavigation InNavigation,
+		const FArrangedWidgetsCache& InArrangedWidgets,
+		FNavigationMethod* InNavigationMethod,
+		TArray<FNavigationMethod::FDebugWidgetResult>& OutDebugResults) const;
+
+	FArrangedWidgetsCache GetArrangedWidgetsForFocus(const FGeometry& MyGeometry, uint32 InUserIndex) const;
+
 protected:
 	// Begin SWidget overrides.
 	virtual FVector2D ComputeDesiredSize(float) const override;
 	// End SWidget overrides.
 
 private:
-
+	TWeakPtr<SScrollBox> OwningScrollBox;
 	EOrientation Orientation;
 	bool BackPadScrolling;
 	bool FrontPadScrolling;
diff --git a/Engine/Source/Runtime/SlateCore/Private/Input/HittestGrid.cpp b/Engine/Source/Runtime/SlateCore/Private/Input/HittestGrid.cpp
--- a/Engine/Source/Runtime/SlateCore/Private/Input/HittestGrid.cpp
+++ b/Engine/Source/Runtime/SlateCore/Private/Input/HittestGrid.cpp
@@ -699,15 +699,14 @@ TSharedPtr<SWidget> FHittestGrid::FindNextFocusableWidgetExperimental(const FArr
 	return NextFocusedWidget;
 }
 
-FSlateRect FHittestGrid::GetWidgetRenderBoundingRect(const TSharedPtr<SWidget>& Widget) const
+FSlateRect FHittestGrid::GetWidgetRenderBoundingRect(const TSharedPtr<const SWidget>& Widget) const
 {
 	if (!Widget.IsValid())
 	{
 		return FSlateRect();
 	}
-	FGeometry WidgetGeo = Widget->GetPaintSpaceGeometry();
-	WidgetGeo.AppendTransform(FSlateLayoutTransform(-GridWindowOrigin));
-	return WidgetGeo.GetRenderBoundingRect();
+
+	return Widget->GetHitTestBoundingRect().OffsetBy(-GridWindowOrigin);
 }
 
 FIntPoint FHittestGrid::GetCellCoordinate(UE::Slate::FDeprecateVector2DParameter Position) const
diff --git a/Engine/Source/Runtime/SlateCore/Private/Widgets/SWidget.cpp b/Engine/Source/Runtime/SlateCore/Private/Widgets/SWidget.cpp
--- a/Engine/Source/Runtime/SlateCore/Private/Widgets/SWidget.cpp
+++ b/Engine/Source/Runtime/SlateCore/Private/Widgets/SWidget.cpp
@@ -643,7 +643,11 @@ void SWidget::OnFinishedKeyInput()
 
 FNavigationReply SWidget::OnNavigation(const FGeometry& MyGeometry, const FNavigationEvent& InNavigationEvent)
 {
-	EUINavigation Type = InNavigationEvent.GetNavigationType();
+	return OnNavigation(MyGeometry, InNavigationEvent.GetNavigationType());
+}
+
+FNavigationReply SWidget::OnNavigation(const FGeometry& MyGeometry, const EUINavigation Type)
+{
 	TSharedPtr<FNavigationMetaData> NavigationMetaData = GetMetaData<FNavigationMetaData>();
 	if (NavigationMetaData.IsValid())
 	{
@@ -1427,6 +1431,11 @@ FSlateRect SWidget::CalculateCullingAndClippingRules(const FGeometry& AllottedGe
 	return IncomingCullingRect;
 }
 
+FSlateRect SWidget::GetHitTestBoundingRect() const
+{
+	return GetPaintSpaceGeometry().GetLayoutBoundingRect();
+}
+
 int32 SWidget::Paint(const FPaintArgs& Args, const FGeometry& AllottedGeometry, const FSlateRect& MyCullingRect, FSlateWindowElementList& OutDrawElements, int32 LayerId, const FWidgetStyle& InWidgetStyle, bool bParentEnabled) const
 {
 	const EWidgetUpdateFlags PreviousUpdateFlag = UpdateFlags;
diff --git a/Engine/Source/Runtime/SlateCore/Public/Input/HittestGrid.h b/Engine/Source/Runtime/SlateCore/Public/Input/HittestGrid.h
--- a/Engine/Source/Runtime/SlateCore/Public/Input/HittestGrid.h
+++ b/Engine/Source/Runtime/SlateCore/Public/Input/HittestGrid.h
@@ -82,7 +82,7 @@ public:
 	UE::Slate::FDeprecateVector2DResult GetGridWindowOrigin() const { return GridWindowOrigin; }
 
 	/** Get the screen space rectangle that corresponds to the painted area of the widget **/
-	SLATECORE_API FSlateRect GetWidgetRenderBoundingRect(const TSharedPtr<SWidget>& Widget) const;
+	SLATECORE_API FSlateRect GetWidgetRenderBoundingRect(const TSharedPtr<const SWidget>& Widget) const;
 
 	/** Returns true if the CellCoordinates are within the grid **/
 	SLATECORE_API bool IsValidCellCoord(const FIntPoint& CellCoord) const;
diff --git a/Engine/Source/Runtime/SlateCore/Public/Types/NavigationMetaData.h b/Engine/Source/Runtime/SlateCore/Public/Types/NavigationMetaData.h
--- a/Engine/Source/Runtime/SlateCore/Public/Types/NavigationMetaData.h
+++ b/Engine/Source/Runtime/SlateCore/Public/Types/NavigationMetaData.h
@@ -128,6 +128,11 @@ public:
 		NavigationMethodStruct = InNavigationMethodStruct;
 	}
 
+	TInstancedStruct<FNavigationMethod>& GetMutableNavigationMethodStruct()
+	{
+		return NavigationMethodStruct;
+	}
+
 	const TInstancedStruct<FNavigationMethod>& GetNavigationMethodStruct() const
 	{
 		return NavigationMethodStruct;
diff --git a/Engine/Source/Runtime/SlateCore/Public/Widgets/SWidget.h b/Engine/Source/Runtime/SlateCore/Public/Widgets/SWidget.h
--- a/Engine/Source/Runtime/SlateCore/Public/Widgets/SWidget.h
+++ b/Engine/Source/Runtime/SlateCore/Public/Widgets/SWidget.h
@@ -639,6 +639,7 @@ public:
 	 * @param InNavigationEvent	The navigation event generated
 	 */
 	SLATECORE_API virtual FNavigationReply OnNavigation(const FGeometry& MyGeometry, const FNavigationEvent& InNavigationEvent);
+	SLATECORE_API FNavigationReply OnNavigation(const FGeometry& MyGeometry, const EUINavigation Type);
 
 	/**
 	 * Called when the mouse is moved over the widget's window, to determine if we should report whether
@@ -706,6 +707,11 @@ public:
 	 */
 	SLATECORE_API FSlateRect CalculateCullingAndClippingRules(const FGeometry& AllottedGeometry, const FSlateRect& MyCullingRect, bool& bClipToBounds, bool& bAlwaysClip, bool& bIntersectClipBounds) const;
 
+	/**
+	 * @return the SlateRect that FHittestGrid would use to test for hit bounds
+	 */
+	SLATECORE_API virtual FSlateRect GetHitTestBoundingRect() const;
+
 	bool HasAnyUpdateFlags(EWidgetUpdateFlags FlagsToCheck) const
 	{
 		return EnumHasAnyFlags(UpdateFlags, FlagsToCheck);
