Make the unreal way a bit less ... Unreal

The editor layout now properly restore the position and scale of the level editor window position and the other windows managed by the tab manager.
Note this still broken when using a DPI a scaling. That fix will be in another cl.

So this was caused by two issue:
1. We didn't have a mechanism to store and restore the main frame window in the layout system.
2. For the other windows we stored the restorable size instead of the actual size of the window even when it's not maximized. Yeah in some OS a window can have a restorable state without being minimized or maximized.

#Synced-CL 50202233
#Synced-CL 50202233

diff --git a/Engine/Source/Runtime/Slate/Private/Framework/Docking/SDockingArea.cpp b/Engine/Source/Runtime/Slate/Private/Framework/Docking/SDockingArea.cpp
--- a/Engine/Source/Runtime/Slate/Private/Framework/Docking/SDockingArea.cpp
+++ b/Engine/Source/Runtime/Slate/Private/Framework/Docking/SDockingArea.cpp
@@ -14,6 +14,46 @@
 
 #define LOCTEXT_NAMESPACE "SDockingArea"
 
+namespace UE::Slate::Private
+{
+	void SetAreaNodeWindowPlacementValues(TSharedPtr<FTabManager::FArea>& AreaToSet, const SWindow& InWindow)
+	{
+		bool bIsWindowMaximized = InWindow.IsWindowMaximized();
+		FSlateRect WindowRect;
+		if (bIsWindowMaximized)
+		{
+			WindowRect = InWindow.GetNonMaximizedRectInScreen();
+		}
+		else
+		{
+			WindowRect = InWindow.GetRectInScreen();
+		}
+
+		// In order to restore SWindows to their correct size, we need to save areas as 
+		// client area sizes, since the Constructor for SWindow uses a client size
+		if (!InWindow.HasOSWindowBorder())
+		{
+			const FMargin WindowBorder = InWindow.GetWindowBorderSize();
+			WindowRect.Right -= WindowBorder.Left + WindowBorder.Right;
+			WindowRect.Bottom -= WindowBorder.Top + WindowBorder.Bottom;
+		}
+
+		// Remove DPI Scale when saving layout so that the saved size is DPI independent
+		float DPIScale = FPlatformApplicationMisc::GetDPIScaleFactorAtPoint(WindowRect.Left, WindowRect.Top);
+
+		if (AreaToSet)
+		{
+			AreaToSet->SetWindowSize(WindowRect.GetSize() / DPIScale);
+		}
+		else
+		{
+			AreaToSet = FTabManager::NewArea(WindowRect.GetSize() / DPIScale);
+		}
+		AreaToSet->SetWindow(FVector2D(WindowRect.Left, WindowRect.Top) / DPIScale, bIsWindowMaximized);
+	}
+}
+
+
 void SDockingArea::Construct( const FArguments& InArgs, const TSharedRef<FTabManager>& InTabManager, const TSharedRef<FTabManager::FArea>& PersistentNode )
 {
 	MyTabManager = InTabManager;
@@ -351,22 +391,7 @@ TSharedPtr<FTabManager::FLayoutNode> SDockingArea::GatherPersistentLayout() cons
 	TSharedPtr<SWindow> ParentWindow = ParentWindowPtr.Pin();
 	if ( ParentWindow.IsValid() && bManageParentWindow )
 	{
-		FSlateRect WindowRect = ParentWindow->GetNonMaximizedRectInScreen();
-
-		// In order to restore SWindows to their correct size, we need to save areas as 
-		// client area sizes, since the Constructor for SWindow uses a client size
-		if (!ParentWindow->HasOSWindowBorder())
-		{
-			const FMargin WindowBorder = ParentWindow->GetWindowBorderSize();
-			WindowRect.Right -= WindowBorder.Left + WindowBorder.Right;
-			WindowRect.Bottom -= WindowBorder.Top + WindowBorder.Bottom;
-		}
-
-		// Remove DPI Scale when saving layout so that the saved size is DPI independent
-		float DPIScale = FPlatformApplicationMisc::GetDPIScaleFactorAtPoint(WindowRect.Left, WindowRect.Top);
-
-		PersistentNode = FTabManager::NewArea( WindowRect.GetSize() / DPIScale );
-		PersistentNode->SetWindow( FVector2D( WindowRect.Left, WindowRect.Top ) / DPIScale, ParentWindow->IsWindowMaximized() );
+		UE::Slate::Private::SetAreaNodeWindowPlacementValues(PersistentNode, *ParentWindow.Get());
 	}
 	else
 	{
diff --git a/Engine/Source/Runtime/Slate/Private/Framework/Docking/SDockingArea.h b/Engine/Source/Runtime/Slate/Private/Framework/Docking/SDockingArea.h
--- a/Engine/Source/Runtime/Slate/Private/Framework/Docking/SDockingArea.h
+++ b/Engine/Source/Runtime/Slate/Private/Framework/Docking/SDockingArea.h
@@ -15,6 +15,11 @@
 
 class STabSidebar;
 
+namespace UE::Slate::Private
+{
+	void SetAreaNodeWindowPlacementValues(TSharedPtr<FTabManager::FArea>& AreaToSet, const SWindow& InWindow);
+}
+
 /** List of tabs that should be in each sidebar */
 struct FSidebarTabLists
 {
diff --git a/Engine/Source/Runtime/Slate/Private/Framework/Docking/TabManager.cpp b/Engine/Source/Runtime/Slate/Private/Framework/Docking/TabManager.cpp
--- a/Engine/Source/Runtime/Slate/Private/Framework/Docking/TabManager.cpp
+++ b/Engine/Source/Runtime/Slate/Private/Framework/Docking/TabManager.cpp
@@ -194,28 +194,28 @@ TSharedRef<FTabManager::FLayoutNode> FTabManager::FLayout::NewFromString_Helper(
 {
 	struct local
 	{
-		static FTabManager::FArea::EWindowPlacement PlacementFromString( const FString& AsString )
+		static FTabManager::FArea::EWindowPlacementInternal PlacementFromString( const FString& AsString )
 		{
-			static const FString Placement_NoWindow_Str = "Placement_NoWindow";
-			static const FString Placement_Automatic_Str = "Placement_Automatic";
-			static const FString Placement_Specified_Str = "Placement_Specified";
-
-			if (AsString == Placement_NoWindow_Str)
+			if (AsString == GET_ENUMERATOR_NAME_STRING_VIEW_CHECKED(FTabManager::FArea::EWindowPlacementInternal, Placement_NoWindow))
+			{
+				return FTabManager::FArea::EWindowPlacementInternal::Placement_NoWindow;
+			}
+			else if (AsString == GET_ENUMERATOR_NAME_STRING_VIEW_CHECKED(FTabManager::FArea::EWindowPlacementInternal, Placement_Automatic))
 			{
-				return FTabManager::FArea::Placement_NoWindow;
+				return FTabManager::FArea::EWindowPlacementInternal::Placement_Automatic;
 			}
-			else if (AsString == Placement_Automatic_Str)
+			else if (AsString == GET_ENUMERATOR_NAME_STRING_VIEW_CHECKED(FTabManager::FArea::EWindowPlacementInternal, Placement_Specified))
 			{
-				return FTabManager::FArea::Placement_Automatic;
+				return FTabManager::FArea::EWindowPlacementInternal::Placement_Specified;
 			}
-			else if (AsString == Placement_Specified_Str)
+			else if (AsString == GET_ENUMERATOR_NAME_STRING_VIEW_CHECKED(FTabManager::FArea::EWindowPlacementInternal, Placement_ParentSpecified))
 			{
-				return FTabManager::FArea::Placement_Specified;
+				return FTabManager::FArea::EWindowPlacementInternal::Placement_ParentSpecified;
 			}
 			else
 			{
 				ensureMsgf(false, TEXT("Invalid placement mode."));
-				return FTabManager::FArea::Placement_Automatic;
+				return FTabManager::FArea::EWindowPlacementInternal::Placement_Automatic;
 			}
 		}
 
@@ -247,17 +247,17 @@ TSharedRef<FTabManager::FLayoutNode> FTabManager::FLayout::NewFromString_Helper(
 	{
 		TSharedPtr<FTabManager::FArea> NewArea;
 
-		FTabManager::FArea::EWindowPlacement WindowPlacement = local::PlacementFromString( JsonObject->GetStringField(TEXT("WindowPlacement")) );
+		FTabManager::FArea::EWindowPlacementInternal WindowPlacement = local::PlacementFromString( JsonObject->GetStringField(TEXT("WindowPlacement")) );
 		switch( WindowPlacement )
 		{
 			default:
-			case FTabManager::FArea::Placement_NoWindow:
+			case FTabManager::FArea::EWindowPlacementInternal::Placement_NoWindow:
 			{
 				NewArea = FTabManager::NewPrimaryArea();
 			}
 			break;
 
-			case FTabManager::FArea::Placement_Automatic:
+			case FTabManager::FArea::EWindowPlacementInternal::Placement_Automatic:
 			{
 				FVector2D WindowSize;
 
@@ -268,7 +268,8 @@ TSharedRef<FTabManager::FLayoutNode> FTabManager::FLayout::NewFromString_Helper(
 			}
 			break;
 
-			case FTabManager::FArea::Placement_Specified:
+			case FTabManager::FArea::EWindowPlacementInternal::Placement_ParentSpecified:
+			case FTabManager::FArea::EWindowPlacementInternal::Placement_Specified:
 			{
 				FVector2D WindowPosition = FVector2D::ZeroVector;
 				FVector2D WindowSize;
@@ -283,6 +284,11 @@ TSharedRef<FTabManager::FLayoutNode> FTabManager::FLayout::NewFromString_Helper(
 
 				NewArea = FTabManager::NewArea( WindowSize );
 				NewArea->SetWindow( WindowPosition, bIsMaximized );
+
+				if (WindowPlacement == FTabManager::FArea::EWindowPlacementInternal::Placement_ParentSpecified)
+				{
+					NewArea->WindowPlacement = FTabManager::FArea::EWindowPlacementInternal::Placement_ParentSpecified;
+				}
 			}
 			break;
 		}
@@ -474,19 +480,28 @@ TSharedRef<FJsonObject> FTabManager::FLayout::PersistToString_Helper(const TShar
 		JsonObj->SetStringField( TEXT("Type"), TEXT("Area") );
 		JsonObj->SetStringField( TEXT("Orientation"), (NodeAsArea->GetOrientation() == Orient_Horizontal) ? TEXT("Orient_Horizontal") : TEXT("Orient_Vertical")  );
 
-		if ( NodeAsArea->WindowPlacement == FArea::Placement_Automatic )
+		if ( NodeAsArea->WindowPlacement == FArea::EWindowPlacementInternal::Placement_Automatic )
 		{
-			JsonObj->SetStringField( TEXT("WindowPlacement"), TEXT("Placement_Automatic") );
+			JsonObj->SetStringField( TEXT("WindowPlacement"), GET_ENUMERATOR_NAME_STRING_CHECKED(FArea::EWindowPlacementInternal, Placement_Automatic) );
 			JsonObj->SetNumberField( TEXT("WindowSize_X"), NodeAsArea->UnscaledWindowSize.X );
 			JsonObj->SetNumberField( TEXT("WindowSize_Y"), NodeAsArea->UnscaledWindowSize.Y );
 		}
-		else if (NodeAsArea->WindowPlacement == FArea::Placement_NoWindow)
+		else if (NodeAsArea->WindowPlacement == FArea::EWindowPlacementInternal::Placement_NoWindow)
 		{
-			JsonObj->SetStringField( TEXT("WindowPlacement"), TEXT("Placement_NoWindow") );
+			JsonObj->SetStringField( TEXT("WindowPlacement"), GET_ENUMERATOR_NAME_STRING_CHECKED(FArea::EWindowPlacementInternal, Placement_NoWindow) );
 		}
-		else if ( NodeAsArea->WindowPlacement == FArea::Placement_Specified )
+		else if ( NodeAsArea->WindowPlacement == FArea::EWindowPlacementInternal::Placement_Specified 
+			|| NodeAsArea->WindowPlacement == FArea::EWindowPlacementInternal::Placement_ParentSpecified )
 		{
-			JsonObj->SetStringField( TEXT("WindowPlacement"), TEXT("Placement_Specified") );
+			if (NodeAsArea->WindowPlacement == FArea::EWindowPlacementInternal::Placement_Specified)
+			{
+				JsonObj->SetStringField(TEXT("WindowPlacement"), GET_ENUMERATOR_NAME_STRING_CHECKED(FArea::EWindowPlacementInternal, Placement_Specified) );
+			}
+			else
+			{
+				JsonObj->SetStringField(TEXT("WindowPlacement"), GET_ENUMERATOR_NAME_STRING_CHECKED(FArea::EWindowPlacementInternal, Placement_ParentSpecified) );
+			}
+
 			JsonObj->SetNumberField( TEXT("WindowPosition_X"), NodeAsArea->UnscaledWindowPosition.X );
 			JsonObj->SetNumberField( TEXT("WindowPosition_Y"), NodeAsArea->UnscaledWindowPosition.Y );
 			JsonObj->SetNumberField( TEXT("WindowSize_X"), NodeAsArea->UnscaledWindowSize.X );
@@ -1036,10 +1051,20 @@ TSharedRef<FTabManager::FLayout> FTabManager::PersistLayout() const
 		if ( PersistedNode.IsValid() )
 		{
 			PersistentLayout->AddArea( PersistedNode.ToSharedRef() );
-			if (PersistedNode->WindowPlacement == FArea::Placement_NoWindow)
+			if (PersistedNode->WindowPlacement == FArea::EWindowPlacementInternal::Placement_NoWindow)
 			{
 				ensure( !PersistentLayout->PrimaryArea.IsValid() );
 				PersistentLayout->PrimaryArea = PersistedNode;
+
+				if (this == &FGlobalTabmanager::Get().Get())
+				{
+					// Convert the area placement to use the ParentSpecified mode so that the global that manager can restore the main frame window
+					if (TSharedPtr<SWindow> ParentWindow = ChildDockingArea->GetParentWindow())
+					{
+						UE::Slate::Private::SetAreaNodeWindowPlacementValues(PersistedNode, *ParentWindow.Get());
+						PersistedNode->WindowPlacement = FArea::EWindowPlacementInternal::Placement_ParentSpecified;
+					}
+				}
 			}
 		}
 	}
@@ -1134,7 +1159,9 @@ TSharedPtr<SWidget> FTabManager::RestoreFrom(const TSharedRef<FLayout>& Layout,
 		const TSharedRef<FArea>& ThisArea = Layout->Areas[AreaIndex];
 		// Set all InvalidTab tabs to OpenedTab so the Editor tries to load them. All non-recognized tabs will be set to InvalidTab later.
 		SetTabsTo(ThisArea, ETabState::OpenedTab, ETabState::InvalidTab);
-		const bool bIsPrimaryArea = ThisArea->WindowPlacement == FArea::Placement_NoWindow;
+
+		// Use in the placement has a fallback for old layout where the primary area might not have been saved
+		const bool bIsPrimaryArea = ThisArea == Layout->GetPrimaryArea() || ThisArea->WindowPlacement == FArea::EWindowPlacementInternal::Placement_NoWindow;
 
 		// If this is the primary area containing the LevelEditor based on the HomeScreen CVar it will either force the HomeScreen and LevelEditor position (and always spawn the HomeScreen) or just the LevelEditor position
 		if (bIsPrimaryArea)
@@ -1172,6 +1199,13 @@ TSharedPtr<SWidget> FTabManager::RestoreFrom(const TSharedRef<FLayout>& Layout,
 			if (bIsPrimaryArea && RestoredDockArea.IsValid() && ensure(!PrimaryDockArea.IsValid()))
 			{
 				PrimaryDockArea = RestoredDockArea;
+
+				if (ThisArea->WindowPlacement == FArea::EWindowPlacementInternal::Placement_Specified && this == &FGlobalTabmanager::Get().Get())
+				{
+					// The global manager will restore the main area window if is was specified
+					TSharedPtr<SWindow> MainWindow = RestoredDockArea->GetParentWindow();
+
+				}
 			}
 		}
 	}
@@ -1859,7 +1893,7 @@ TSharedPtr<SDockingTabStack> FTabManager::AttemptToOpenTab( const FTabId& Closed
 		// If this is not the global tab manager and the tab is a NomadTab in a floating window, then remove it from the collapsed area.
 		if (FGlobalTabmanager::Get() != SharedThis(this) &&
 			NomadTabSpawner->Contains(ClosedTabId.TabType) &&
-			CollapsedAreaWithMatchingTab->WindowPlacement != FTabManager::FArea::EWindowPlacement::Placement_NoWindow)
+			CollapsedAreaWithMatchingTab->WindowPlacement != FTabManager::FArea::EWindowPlacementInternal::Placement_NoWindow)
 		{
 			RemoveTabFromCollapsedAreas(TabMatcher);
 		}
@@ -2124,8 +2158,9 @@ TSharedPtr<SDockingNode> FTabManager::RestoreArea_Helper(const TSharedRef<FLayou
 	}
 	else if ( NodeAsArea.IsValid() )
 	{
-		const bool bSplitterIsDockArea = NodeAsArea.IsValid();
-		const bool bDockNeedsNewWindow = NodeAsArea.IsValid() && (NodeAsArea->WindowPlacement != FArea::Placement_NoWindow);
+		const bool bDockNeedsNewWindow = NodeAsArea.IsValid()
+			&& !(NodeAsArea->WindowPlacement == FArea::EWindowPlacementInternal::Placement_NoWindow 
+				|| NodeAsArea->WindowPlacement == FArea::EWindowPlacementInternal::Placement_ParentSpecified);
 
 		TSharedPtr<SDockingArea> NewDockAreaWidget;
 
@@ -2133,10 +2168,9 @@ TSharedPtr<SDockingNode> FTabManager::RestoreArea_Helper(const TSharedRef<FLayou
 		{
 			// The layout node we are restoring is a dock area.
 			// It needs a new window into which it will land.
-
 			const bool bIsChildWindow = ParentWindow.IsValid();
 
-			const bool bAutoPlacement = (NodeAsArea->WindowPlacement == FArea::Placement_Automatic);
+			const bool bAutoPlacement = (NodeAsArea->WindowPlacement == FArea::EWindowPlacementInternal::Placement_Automatic);
 			TSharedRef<SWindow> NewWindow = (bAutoPlacement)
 				? SNew(SWindow)
 					.AutoCenter( EAutoCenter::PreferredWorkArea )
@@ -2199,6 +2233,19 @@ TSharedPtr<SDockingNode> FTabManager::RestoreArea_Helper(const TSharedRef<FLayou
 		}
 		else
 		{
+			if (NodeAsArea->WindowPlacement == FArea::EWindowPlacementInternal::Placement_ParentSpecified)
+			{
+				if (ParentWindow)
+				{
+					UE::Slate::FDeprecateVector2DResult WindowSize = ParentWindow->GetWindowSizeFromClientSize(NodeAsArea->UnscaledWindowSize, ParentWindow->GetDPIScaleFactor());
+					ParentWindow->ReshapeWindow(NodeAsArea->UnscaledWindowPosition, WindowSize);
+					if (NodeAsArea->bIsMaximized)
+					{
+						ParentWindow->Maximize();
+					}
+				}
+			}
+
 			TArray<TSharedRef<SDockingNode>> DockingNodes;
 			if (CanRestoreSplitterContent(DockingNodes, NodeAsArea.ToSharedRef(), ParentWindow, OutSidebarTabs, OutputCanBeNullptr))
 			{
diff --git a/Engine/Source/Runtime/Slate/Public/Framework/Docking/TabManager.h b/Engine/Source/Runtime/Slate/Public/Framework/Docking/TabManager.h
--- a/Engine/Source/Runtime/Slate/Public/Framework/Docking/TabManager.h
+++ b/Engine/Source/Runtime/Slate/Public/Framework/Docking/TabManager.h
@@ -693,13 +693,16 @@ class FTabManager : public TSharedFromThis<FTabManager>
 		{
 				friend class FTabManager;
 		
-			public:
+			protected:
 				/** An enum which specifies how content appears within a window. */
-				enum EWindowPlacement
+				enum class EWindowPlacementInternal : uint8
 				{
 					/** The content is docked within a primary window. */
 					Placement_NoWindow,
 
+					/** The content is docked within a primary window, but determine the placement of the primary window */
+					Placement_ParentSpecified,
+
 					/** The content is docked within a floating window that is initially automatically positioned. */
 					Placement_Automatic,
 
@@ -707,6 +710,20 @@ class FTabManager : public TSharedFromThis<FTabManager>
 					Placement_Specified
 				};
 
+			public:
+
+				enum UE_DEPRECATED(5.8, "This enum will removed from the engine") EWindowPlacement
+				{
+					/** The content is docked within a primary window. */
+					Placement_NoWindow,
+
+					/** The content is docked within a floating window that is initially automatically positioned. */
+					Placement_Automatic,
+
+					/** The content is docked within a floating window with positioning specified by an FArea. */
+					Placement_Specified
+				};
+	
 				TSharedRef<FArea> Split( TSharedRef<FLayoutNode> InNode )
 				{
 					ChildNodes.Add(InNode);
@@ -726,9 +743,15 @@ class FTabManager : public TSharedFromThis<FTabManager>
 					return SharedThis(this);
 				}
 
+				TSharedRef<FArea> SetWindowSize( FVector2D InSize )
+				{
+					UnscaledWindowSize = InSize;
+					return SharedThis(this);
+				}
+
 				TSharedRef<FArea> SetWindow( FVector2D InPosition, bool IsMaximized )
 				{
-					WindowPlacement = Placement_Specified;
+					WindowPlacement = EWindowPlacementInternal::Placement_Specified;
 					UnscaledWindowPosition = InPosition;
 					bIsMaximized = IsMaximized;
 					return SharedThis(this);
@@ -777,19 +800,19 @@ class FTabManager : public TSharedFromThis<FTabManager>
 			*/
 			bool DefinesPositionallySpecifiedFloatingWindow()const
 			{
-				return WindowPlacement == Placement_Specified;
+				return WindowPlacement == EWindowPlacementInternal::Placement_Specified;
 			}
 
 			protected:
 				FArea( const float InWidth, const float InHeight )
-				: WindowPlacement(Placement_Automatic)
+				: WindowPlacement(EWindowPlacementInternal::Placement_Automatic)
 				, UnscaledWindowPosition(FVector2D(0.0, 0.0))
 				, UnscaledWindowSize(InWidth, InHeight)
 				, bIsMaximized( false )
 				{
 				}
 
-				EWindowPlacement WindowPlacement;
+				EWindowPlacementInternal WindowPlacement;
 				FVector2D UnscaledWindowPosition;
 				FVector2D UnscaledWindowSize;
 				bool bIsMaximized;
@@ -928,7 +951,7 @@ class FTabManager : public TSharedFromThis<FTabManager>
 		static TSharedRef<FArea> NewPrimaryArea()
 		{
 			TSharedRef<FArea>Area = MakeShareable( new FArea(0,0) );
-			Area->WindowPlacement = FArea::Placement_NoWindow;
+			Area->WindowPlacement = FArea::EWindowPlacementInternal::Placement_NoWindow;
 			return Area;
 		}
 
