Enabling multi view rendering support and Xcode view debugger support for Quixel Bridge and other standalone windows with webview-backed widgets

#Synced-CL 50368161

diff --git a/Engine/Plugins/Bridge/Source/Bridge/Private/UI/BrowserBinding.cpp b/Engine/Plugins/Bridge/Source/Bridge/Private/UI/BrowserBinding.cpp
--- a/Engine/Plugins/Bridge/Source/Bridge/Private/UI/BrowserBinding.cpp
+++ b/Engine/Plugins/Bridge/Source/Bridge/Private/UI/BrowserBinding.cpp
@@ -279,6 +279,11 @@ void UBrowserBinding::ShowLoginDialog(FString LoginUrl, FString ResponseCodeUrl)
 			.Title(FText::FromString(TEXT("Login")))
 			.ClientSize(FVector2D(450, 700))
 			.SupportsMaximize(false)
+#if PLATFORM_MAC && ALPHA_BLENDED_WINDOWS
+			// this is needed to allow Slate-rendered Metal views to show up on XCode's view debugger
+			// and for multi-view rendering to function properly
+			.SupportsTransparency(EWindowTransparency::PerPixel)
+#endif
 			.SupportsMinimize(false)		
 			[
 				SNew(SVerticalBox)
diff --git a/Engine/Source/Runtime/Slate/Private/Framework/Docking/FDockingDragOperation.cpp b/Engine/Source/Runtime/Slate/Private/Framework/Docking/FDockingDragOperation.cpp
--- a/Engine/Source/Runtime/Slate/Private/Framework/Docking/FDockingDragOperation.cpp
+++ b/Engine/Source/Runtime/Slate/Private/Framework/Docking/FDockingDragOperation.cpp
@@ -349,6 +349,11 @@ void FDockingDragOperation::DroppedOntoNothing()
 	TSharedPtr<SWindow> NewWindowParent = MyTabManager->GetPrivateApi().GetParentWindow();
 
 	TSharedRef<SWindow> NewWindow = SNew(SWindow)
+#if PLATFORM_MAC && ALPHA_BLENDED_WINDOWS
+		// this is needed to allow Slate-rendered Metal views to show up on XCode's view debugger
+		// and for multi-view rendering to function properly
+		.SupportsTransparency(EWindowTransparency::PerPixel)
+#endif
 		.Title(FGlobalTabmanager::Get()->GetApplicationTitle())
 		.AutoCenter(EAutoCenter::None)
 		// Divide out scale, it is already factored into position
diff --git a/Engine/Source/Runtime/Slate/Private/Framework/Docking/STabPanelDrawer.cpp b/Engine/Source/Runtime/Slate/Private/Framework/Docking/STabPanelDrawer.cpp
--- a/Engine/Source/Runtime/Slate/Private/Framework/Docking/STabPanelDrawer.cpp
+++ b/Engine/Source/Runtime/Slate/Private/Framework/Docking/STabPanelDrawer.cpp
@@ -412,6 +412,11 @@ void STabPanelDrawer::OpenTabInNewWindow() const
 	TSharedPtr<SWindow> NewWindowParent = MyTabManager->GetPrivateApi().GetParentWindow();
 
 	TSharedRef<SWindow> NewWindow = SNew(SWindow)
+#if PLATFORM_MAC && ALPHA_BLENDED_WINDOWS
+		// this is needed to allow Slate-rendered Metal views to show up on XCode's view debugger
+		// and for multi-view rendering to function properly
+		.SupportsTransparency(EWindowTransparency::PerPixel)
+#endif
 		.Title(FGlobalTabmanager::Get()->GetApplicationTitle())
 		.AutoCenter(EAutoCenter::None)
 		// Divide out scale, it is already factored into position
diff --git a/Engine/Source/Runtime/Slate/Private/Framework/Docking/TabManager.cpp b/Engine/Source/Runtime/Slate/Private/Framework/Docking/TabManager.cpp
--- a/Engine/Source/Runtime/Slate/Private/Framework/Docking/TabManager.cpp
+++ b/Engine/Source/Runtime/Slate/Private/Framework/Docking/TabManager.cpp
@@ -1786,6 +1786,11 @@ TSharedPtr<SDockTab> FTabManager::InvokeTab_Internal(const FTabId& TabId, bool b
 		TSharedPtr<SWindow> NewWindowParent = GetPrivateApi().GetParentWindow();
 
 		TSharedRef<SWindow> NewWindow = SNew(SWindow)
+#if PLATFORM_MAC && ALPHA_BLENDED_WINDOWS
+			// this is needed to allow Slate-rendered Metal views to show up on XCode's view debugger
+			// and for multi-view rendering to function properly
+			.SupportsTransparency(EWindowTransparency::PerPixel)
+#endif
 			.Title(FGlobalTabmanager::Get()->GetApplicationTitle())
 			.AutoCenter(EAutoCenter::None)
 			// Divide out scale, it is already factored into position
@@ -2176,12 +2181,22 @@ TSharedPtr<SDockingNode> FTabManager::RestoreArea_Helper(const TSharedRef<FLayou
 					.AutoCenter( EAutoCenter::PreferredWorkArea )
 					.ClientSize( NodeAsArea->UnscaledWindowSize )
 					.CreateTitleBar( false )
+#if PLATFORM_MAC && ALPHA_BLENDED_WINDOWS
+					// this is needed to allow Slate-rendered Metal views to show up on XCode's view debugger
+					// and for multi-view rendering to function properly
+					.SupportsTransparency(EWindowTransparency::PerPixel)
+#endif
 					.IsInitiallyMaximized( NodeAsArea->bIsMaximized )
 				: SNew(SWindow)
 					.AutoCenter( EAutoCenter::None )
 					.ScreenPosition( NodeAsArea->UnscaledWindowPosition )
 					.ClientSize( NodeAsArea->UnscaledWindowSize )
 					.CreateTitleBar( false )
+#if PLATFORM_MAC && ALPHA_BLENDED_WINDOWS
+					// this is needed to allow Slate-rendered Metal views to show up on XCode's view debugger
+					// and for multi-view rendering to function properly
+					.SupportsTransparency(EWindowTransparency::PerPixel)
+#endif
 					.IsInitiallyMaximized( NodeAsArea->bIsMaximized );
 
 			// Set a default title; restoring the splitter content may override this if it activates a tab
@@ -3795,6 +3810,11 @@ void FGlobalTabmanager::OnTabRelocated( const TSharedRef<SDockTab>& RelocatedTab
 				if ( OldChildWindow.IsValid() )
 				{
 					TSharedRef<SWindow> NewChildWindow = SNew(SWindow)
+#if PLATFORM_MAC && ALPHA_BLENDED_WINDOWS
+					// this is needed to allow Slate-rendered Metal views to show up on XCode's view debugger
+					// and for multi-view rendering to function properly
+					.SupportsTransparency(EWindowTransparency::PerPixel)
+#endif
 					.AutoCenter(EAutoCenter::None)
 					.ScreenPosition(OldChildWindow->GetPositionInScreen())
 					.ClientSize(OldChildWindow->GetSizeInScreen())
