Salte Dynamic Invalidation - STableRow

Adding invalidation support for dragging operations

#Synced-CL 50145007

diff --git a/Engine/Plugins/Media/MediaViewer/Source/MediaViewer/Private/Widgets/SMediaViewerLibraryGroup.cpp b/Engine/Plugins/Media/MediaViewer/Source/MediaViewer/Private/Widgets/SMediaViewerLibraryGroup.cpp
--- a/Engine/Plugins/Media/MediaViewer/Source/MediaViewer/Private/Widgets/SMediaViewerLibraryGroup.cpp
+++ b/Engine/Plugins/Media/MediaViewer/Source/MediaViewer/Private/Widgets/SMediaViewerLibraryGroup.cpp
@@ -188,12 +188,12 @@ void SMediaViewerLibraryGroup::OnDragEnter(const FGeometry& InMyGeometry, const
 		return;
 	}
 
-	ItemDropZone = EItemDropZone::OntoItem;
+	SetItemDropZone(EItemDropZone::OntoItem);
 }
 
 FReply SMediaViewerLibraryGroup::OnDrop(const FGeometry& InMyGeometry, const FDragDropEvent& InEvent)
 {
-	ItemDropZone = TOptional<EItemDropZone>();
+	SetItemDropZone(TOptional<EItemDropZone>());
 
 	TSharedPtr<IMediaViewerLibrary> Library = LibraryWeak.Pin();
 
diff --git a/Engine/Plugins/Media/MediaViewer/Source/MediaViewer/Private/Widgets/SMediaViewerLibraryItem.cpp b/Engine/Plugins/Media/MediaViewer/Source/MediaViewer/Private/Widgets/SMediaViewerLibraryItem.cpp
--- a/Engine/Plugins/Media/MediaViewer/Source/MediaViewer/Private/Widgets/SMediaViewerLibraryItem.cpp
+++ b/Engine/Plugins/Media/MediaViewer/Source/MediaViewer/Private/Widgets/SMediaViewerLibraryItem.cpp
@@ -411,7 +411,7 @@ void SMediaViewerLibraryItem::OnLibraryItemDropped(const IMediaViewerLibrary::FG
 	{
 		Library->MoveItemWithinGroup(
 			InDroppedGroupItem,
-			ItemDropZone.Get(EItemDropZone::OntoItem) == EItemDropZone::AboveItem
+			GetItemDropZone().Get(EItemDropZone::OntoItem) == EItemDropZone::AboveItem
 				? MyGroupIndex
 				: MyGroupIndex + 1
 		);
@@ -421,7 +421,7 @@ void SMediaViewerLibraryItem::OnLibraryItemDropped(const IMediaViewerLibrary::FG
 		Library->MoveItemToGroup(
 			InDroppedGroupItem,
 			GroupItem.GroupId,
-			ItemDropZone.Get(EItemDropZone::OntoItem) == EItemDropZone::AboveItem
+			GetItemDropZone().Get(EItemDropZone::OntoItem) == EItemDropZone::AboveItem
 				? MyGroupIndex
 				: MyGroupIndex + 1
 		);
@@ -436,7 +436,7 @@ void SMediaViewerLibraryItem::OnLibraryItemDropped(const IMediaViewerLibrary::FG
 		Library->AddItemToGroup(
 			Library->GetItem(InDroppedGroupItem.ItemId).ToSharedRef(),
 			GroupItem.GroupId,
-			ItemDropZone.Get(EItemDropZone::OntoItem) == EItemDropZone::AboveItem
+			GetItemDropZone().Get(EItemDropZone::OntoItem) == EItemDropZone::AboveItem
 				? MyGroupIndex
 				: MyGroupIndex + 1
 		);
@@ -474,15 +474,15 @@ void SMediaViewerLibraryItem::OnDragEnter(const FGeometry& InMyGeometry, const F
 	{
 		if (Mouse < (Top + 5))
 		{
-			ItemDropZone = EItemDropZone::AboveItem;
+			SetItemDropZone(EItemDropZone::AboveItem);
 		}
 		else if (Mouse > (Bottom - 5))
 		{
-			ItemDropZone = EItemDropZone::BelowItem;
+			SetItemDropZone(EItemDropZone::BelowItem);
 		}
 		else
 		{
-			ItemDropZone = EItemDropZone::OntoItem;
+			SetItemDropZone(EItemDropZone::OntoItem);
 		}
 	}
 }
@@ -500,7 +500,7 @@ FReply SMediaViewerLibraryItem::OnDrop(const FGeometry& InMyGeometry, const FDra
 		OnAssetsDropped(AssetDragDrop->GetAssets());
 	}
 
-	ItemDropZone = TOptional<EItemDropZone>();
+	SetItemDropZone(TOptional<EItemDropZone>());
 
 	return FReply::Handled();
 }
@@ -556,7 +556,7 @@ void SMediaViewerLibraryItem::OnAssetsDropped(TConstArrayView<FAssetData> InAsse
 	}
 	
 	IMediaViewerModule& Module = IMediaViewerModule::Get();	
-	const EItemDropZone DropZone = ItemDropZone.Get(EItemDropZone::OntoItem);
+	const EItemDropZone DropZone = GetItemDropZone().Get(EItemDropZone::OntoItem);
 
 	for (int32 Index = InAssetData.Num() - 1; Index >= 0; --Index)
 	{
diff --git a/Engine/Plugins/Runtime/AudioWidgets/Source/AudioWidgets/Private/SDragReorderableTileView.h b/Engine/Plugins/Runtime/AudioWidgets/Source/AudioWidgets/Private/SDragReorderableTileView.h
--- a/Engine/Plugins/Runtime/AudioWidgets/Source/AudioWidgets/Private/SDragReorderableTileView.h
+++ b/Engine/Plugins/Runtime/AudioWidgets/Source/AudioWidgets/Private/SDragReorderableTileView.h
@@ -181,15 +181,15 @@ namespace AudioWidgets
 		virtual FReply OnDragOver(const FGeometry& MyGeometry, const FDragDropEvent& DragDropEvent) override
 		{
 			// Setting STableRow::ItemDropZone TOptional will enable painting the drop indicator:
-			Base::ItemDropZone = CanAcceptDrop(MyGeometry, DragDropEvent);
+			Base::SetItemDropZone(CanAcceptDrop(MyGeometry, DragDropEvent));
 
 			return FReply::Handled();
 		}
 
 		virtual FReply OnDrop(const FGeometry& MyGeometry, const FDragDropEvent& DragDropEvent) override
 		{
-			// Clear STableRow::ItemDropZone TOptional to no longer paint the drop indicator:
-			Base::ItemDropZone.Reset();
+			// Clear ItemDropZone to no longer paint the drop indicator:
+			Base::SetItemDropZone(TOptional<EItemDropZone>());
 
 			if (TOptional<EItemDropZone> AcceptableDropZone = CanAcceptDrop(MyGeometry, DragDropEvent))
 			{
diff --git a/Engine/Source/Editor/ContentBrowser/Private/SPathView.h b/Engine/Source/Editor/ContentBrowser/Private/SPathView.h
--- a/Engine/Source/Editor/ContentBrowser/Private/SPathView.h
+++ b/Engine/Source/Editor/ContentBrowser/Private/SPathView.h
@@ -641,8 +641,7 @@ private:
 template<>
 struct TWidgetTypeTraits<class SFavoritePathView>
 {
-	// Disable invalidation otherwise drop indicators are not drawn
-	static constexpr bool SupportsInvalidation() { return false; }
+	static constexpr bool SupportsInvalidation() { return true; }
 };
 
 /**
diff --git a/Engine/Source/Runtime/Slate/Public/Widgets/Views/STableRow.h b/Engine/Source/Runtime/Slate/Public/Widgets/Views/STableRow.h
--- a/Engine/Source/Runtime/Slate/Public/Widgets/Views/STableRow.h
+++ b/Engine/Source/Runtime/Slate/Public/Widgets/Views/STableRow.h
@@ -87,6 +87,8 @@ class STableRow : public ITableRow, public SBorder
 {
 	static_assert(TIsValidListItem<ItemType>::Value, "Item type T must be UObjectBase*, TObjectPtr<>, TWeakObjectPtr<>, TSharedRef<>, or TSharedPtr<>.");
 
+	SLATE_DECLARE_WIDGET(STableRow, SBorder)
+	static inline SLATE_IMPLEMENT_WIDGET(STableRow)
 public:
 	/** Delegate signature for querying whether this FDragDropEvent will be handled by the drop target of type ItemType. */
 	DECLARE_DELEGATE_RetVal_ThreeParams(TOptional<EItemDropZone>, FOnCanAcceptDrop, const FDragDropEvent&, EItemDropZone, ItemType);
@@ -356,15 +358,15 @@ public:
 	}
 	int32 PaintDropIndicator( const FPaintArgs& Args, const FGeometry& AllottedGeometry, const FSlateRect& MyCullingRect, FSlateWindowElementList& OutDrawElements, int32 LayerId, const FWidgetStyle& InWidgetStyle, bool bParentEnabled ) const
 	{
-		if (ItemDropZone.IsSet())
+		if (GetItemDropZone().IsSet())
 		{
 			if (PaintDropIndicatorEvent.IsBound())
 			{
-				return PaintDropIndicatorEvent.Execute(ItemDropZone.GetValue(), Args, AllottedGeometry, MyCullingRect, OutDrawElements, LayerId, InWidgetStyle, bParentEnabled);
+				return PaintDropIndicatorEvent.Execute(GetItemDropZone().GetValue(), Args, AllottedGeometry, MyCullingRect, OutDrawElements, LayerId, InWidgetStyle, bParentEnabled);
 			}
 			else
 			{
-				return OnPaintDropIndicator(ItemDropZone.GetValue(), Args, AllottedGeometry, MyCullingRect, OutDrawElements, LayerId, InWidgetStyle, bParentEnabled);
+				return OnPaintDropIndicator(GetItemDropZone().GetValue(), Args, AllottedGeometry, MyCullingRect, OutDrawElements, LayerId, InWidgetStyle, bParentEnabled);
 			}
 		}
 
@@ -737,7 +739,7 @@ public:
 
 	virtual void OnDragLeave(FDragDropEvent const& DragDropEvent) override
 	{
-		ItemDropZone = TOptional<EItemDropZone>();
+		SetItemDropZone(TOptional<EItemDropZone>());
 
 		if (OnDragLeave_Handler.IsBound())
 		{
@@ -774,7 +776,7 @@ public:
 			const FVector2f LocalPointerPos = MyGeometry.AbsoluteToLocal(DragDropEvent.GetScreenSpacePosition());
 			const EItemDropZone ItemHoverZone = ZoneFromPointerPosition(LocalPointerPos, MyGeometry.GetLocalSize(), OwnerTable->Private_GetOrientation());
 
-			ItemDropZone = [ItemHoverZone, DragDropEvent, this]()
+			SetItemDropZone([ItemHoverZone, DragDropEvent, this]()
 			{
 				TSharedRef< ITypedTableView<ItemType> > OwnerTable = OwnerTablePtr.Pin().ToSharedRef();
 				if (const TObjectPtrWrapTypeOf<ItemType>* MyItemPtr = GetItemForThis(OwnerTable))
@@ -783,7 +785,7 @@ public:
 				}
 
 				return TOptional<EItemDropZone>();
-			}();
+			}());
 
 			return FReply::Handled();
 		}
@@ -803,7 +805,7 @@ public:
 				const TSharedRef< ITypedTableView<ItemType> > OwnerTable = OwnerTablePtr.Pin().ToSharedRef();
 
 				// A drop finishes the drag/drop operation, so we are no longer providing any feedback.
-				ItemDropZone = TOptional<EItemDropZone>();
+				SetItemDropZone(TOptional<EItemDropZone>());
 
 				// Find item associated with this widget.
 				if (const TObjectPtrWrapTypeOf<ItemType>* MyItemPtr = GetItemForThis(OwnerTable))
@@ -1127,7 +1129,8 @@ public:
 		: IndexInList(0)
 		, bShowSelection(true)
 		, SignalSelectionMode( ETableRowSignalSelectionMode::Deferred )
-	{ 
+		, ItemDropZoneAttribute(*this, TOptional<EItemDropZone>())
+	{
 #if WITH_ACCESSIBILITY
 		// As the contents of table rows could be anything,
 		// Ideally, somebody would assign a custom label to each table row with non-accessible content.
@@ -1234,6 +1237,14 @@ protected:
 		return nullptr;
 	}
 
+protected:
+	const TOptional<EItemDropZone>& GetItemDropZone() const { return ItemDropZoneAttribute.Get(); }
+	void SetItemDropZone(TOptional<EItemDropZone> InItemDropZone) { ItemDropZoneAttribute.Assign(*this, MoveTemp(InItemDropZone)); }
+	TSlateAttributeRef<TOptional<EItemDropZone>> GetItemDropZoneAttribute() const
+	{
+		return TSlateAttributeRef<TOptional<EItemDropZone>>{SharedThis(this), ItemDropZoneAttribute};
+	}
+
 protected:
 
 	/** The list that owns this Selectable */
@@ -1266,8 +1277,11 @@ protected:
 	/** Optional delegate for painting drop indicators */
 	FOnPaintDropIndicator PaintDropIndicatorEvent;
 
+#if WITH_EDITORONLY_DATA
 	/** Are we currently dragging/dropping over this item? */
+	UE_DEPRECATED(5.8, "Use GetItemDropZone / SetItemDropZone / GetItemDropZoneAttribute.")
 	TOptional<EItemDropZone> ItemDropZone;
+#endif
 
 	/** Delegate triggered when a user starts to drag a list item */
 	FOnDragDetected OnDragDetected_Handler;
@@ -1299,8 +1313,17 @@ protected:
 
 private:
 	bool bShowWires;
+
+	/** Are we currently dragging/dropping over this item? */
+	TSlateAttribute<TOptional<EItemDropZone>> ItemDropZoneAttribute;
 };
 
+template<typename ItemType>
+inline void STableRow<ItemType>::PrivateRegisterAttributes(FSlateAttributeInitializer& AttributeInitializer)
+{
+	SLATE_ADD_MEMBER_ATTRIBUTE_DEFINITION(AttributeInitializer, ItemDropZoneAttribute, EInvalidateWidgetReason::Paint);
+}
+
 
 template<typename ItemType>
 class SMultiColumnTableRow;
