Slate Dynamic Invalidation - SEditableText/SMultilineEditableText and SlateEditableTextLayout

This change gives SlateEditableTextLayout its own nested attributes to properly handle invalidation when the value of attribute changes, ensuring that the right callbacks are triggered.

#Synced-CL 50145007

diff --git a/Engine/Source/Runtime/Slate/Private/Widgets/Input/SEditableText.cpp b/Engine/Source/Runtime/Slate/Private/Widgets/Input/SEditableText.cpp
--- a/Engine/Source/Runtime/Slate/Private/Widgets/Input/SEditableText.cpp
+++ b/Engine/Source/Runtime/Slate/Private/Widgets/Input/SEditableText.cpp
@@ -12,8 +12,44 @@
 #include "Widgets/Accessibility/SlateAccessibleWidgets.h"
 #endif
 
+SLATE_IMPLEMENT_WIDGET(SEditableText)
+void SEditableText::PrivateRegisterAttributes(FSlateAttributeInitializer& AttributeInitializer)
+{
+	struct FInvalidation
+	{
+		static void SynchronizeTextStyle(SWidget& OwningWidget)
+		{
+			SEditableText& EditableText = static_cast<SEditableText&>(OwningWidget);
+			EditableText.SynchronizeTextStyle();
+		}
+		static void UpdateMarshaller(SWidget& OwningWidget)
+		{
+			SEditableText& EditableText = static_cast<SEditableText&>(OwningWidget);
+			EditableText.PlainTextMarshaller->MakeDirty();
+		}
+	};
+
+	FSlateEditableTextLayout::RegisterNestedAttributes<&SEditableText::SlateEditableTextLayout>(AttributeInitializer);
+
+	SLATE_ADD_MEMBER_ATTRIBUTE_DEFINITION(AttributeInitializer, MinDesiredWidthAttribute, EInvalidateWidgetReason::Layout);
+	SLATE_ADD_MEMBER_ATTRIBUTE_DEFINITION(AttributeInitializer, bIsPasswordAttribute, EInvalidateWidgetReason::Layout)
+		.OnValueChanged(FSlateAttributeDescriptor::FAttributeValueChangedDelegate::CreateStatic(FInvalidation::UpdateMarshaller));
+	SLATE_ADD_MEMBER_ATTRIBUTE_DEFINITION(AttributeInitializer, FontAttribute, EInvalidateWidgetReason::Layout)
+		.OnValueChanged(FSlateAttributeDescriptor::FAttributeValueChangedDelegate::CreateStatic(FInvalidation::SynchronizeTextStyle));
+
+	SLATE_ADD_MEMBER_ATTRIBUTE_DEFINITION(AttributeInitializer, ColorAndOpacityAttribute, EInvalidateWidgetReason::Paint)
+		.OnValueChanged(FSlateAttributeDescriptor::FAttributeValueChangedDelegate::CreateStatic(FInvalidation::SynchronizeTextStyle));
+	SLATE_ADD_MEMBER_ATTRIBUTE_DEFINITION(AttributeInitializer, BackgroundImageSelectedAttribute, EInvalidateWidgetReason::Layout)
+		.OnValueChanged(FSlateAttributeDescriptor::FAttributeValueChangedDelegate::CreateStatic(FInvalidation::SynchronizeTextStyle));
+}
+
 SEditableText::SEditableText()
 	: SlateEditableTextLayout(*this)
+	, FontAttribute(*this)
+	, ColorAndOpacityAttribute(*this)
+	, BackgroundImageSelectedAttribute(*this)
+	, bIsPasswordAttribute(*this, false)
+	, MinDesiredWidthAttribute(*this, 0.0f)
 {
 #if WITH_ACCESSIBILITY
 	AccessibleBehavior = EAccessibleBehavior::Auto;
@@ -23,8 +59,6 @@ SEditableText::SEditableText()
 
 void SEditableText::Construct( const FArguments& InArgs )
 {
-	SetIsReadOnly(InArgs._IsReadOnly);
-
 	bIsCaretMovedWhenGainFocus = InArgs._IsCaretMovedWhenGainFocus;
 	bSelectAllTextWhenFocused = InArgs._SelectAllTextWhenFocused;
 	bRevertTextOnEscape = InArgs._RevertTextOnEscape;
@@ -71,7 +105,7 @@ void SEditableText::Construct( const FArguments& InArgs )
 		SlateEditableTextLayout.SetOnAllFontFacesFinishLoadingDelegate(InArgs._OnAllFontFacesFinishLoading);
 	}
 
-	SetText(InArgs._Text);
+	SetIsReadOnly(InArgs._IsReadOnly);
 	SetIsPassword(InArgs._IsPassword);
 	SetHintText(InArgs._HintText);
 	SetHintTextOpacity(InArgs._HintTextOpacity);
@@ -105,23 +139,23 @@ bool SEditableText::SetEditableText(const FText& InNewText)
 
 void SEditableText::SetFont(TAttribute<FSlateFontInfo> InNewFont)
 {
-	Font = MoveTemp(InNewFont);
-
-	Invalidate(EInvalidateWidgetReason::Layout);
+	bIsFontAttrSet = InNewFont.IsSet();
+	FontAttribute.Assign(*this, MoveTemp(InNewFont));
 }
 
 FSlateFontInfo SEditableText::GetFont() const
 {
-	return Font.Get(SlateEditableTextLayout.GetTextStyle().Font);
+	return bIsFontAttrSet ? FontAttribute.Get() : SlateEditableTextLayout.GetTextStyle().Font;
 }
 
 void SEditableText::SetTextStyle( const FEditableTextStyle& InNewTextStyle )
 {
-	Font = InNewTextStyle.Font;
-	ColorAndOpacity = InNewTextStyle.ColorAndOpacity;
-	BackgroundImageSelected = &InNewTextStyle.BackgroundImageSelected;
-
-	Invalidate(EInvalidateWidgetReason::Layout);
+	bIsFontAttrSet = true;
+	FontAttribute.Assign(*this, InNewTextStyle.Font);
+	bIsColorAndOpacityAttrSet = true;
+	ColorAndOpacityAttribute.Assign(*this, InNewTextStyle.ColorAndOpacity);
+	bIsBackgroundImageSelectedAttrSet = true;
+	BackgroundImageSelectedAttribute.Assign(*this, &InNewTextStyle.BackgroundImageSelected);
 }
 
 void SEditableText::SetTextBlockStyle(const FTextBlockStyle* InTextStyle)
@@ -129,7 +163,6 @@ void SEditableText::SetTextBlockStyle(const FTextBlockStyle* InTextStyle)
 	if (InTextStyle)
 	{
 		SlateEditableTextLayout.SetTextStyle(*InTextStyle);
-		Invalidate(EInvalidateWidgetReason::Layout); //Using Layout as changing text block size can affect the size.
 	}
 }
 
@@ -161,7 +194,7 @@ void SEditableText::CacheDesiredSize(float LayoutScaleMultiplier)
 FVector2D SEditableText::ComputeDesiredSize(float LayoutScaleMultiplier) const
 {
 	FVector2D TextLayoutSize = SlateEditableTextLayout.ComputeDesiredSize(LayoutScaleMultiplier);
-	TextLayoutSize.X = FMath::Max(TextLayoutSize.X, MinDesiredWidth.Get(0.0f));
+	TextLayoutSize.X = FMath::Max(TextLayoutSize.X, MinDesiredWidthAttribute.Get());
 	return TextLayoutSize;
 }
 
@@ -304,14 +337,7 @@ bool SEditableText::IsInteractable() const
 bool SEditableText::ComputeVolatility() const
 {
 	return SWidget::ComputeVolatility()
-		|| HasKeyboardFocus()
-		|| SlateEditableTextLayout.ComputeVolatility()
-		|| Font.IsBound()
-		|| ColorAndOpacity.IsBound()
-		|| BackgroundImageSelected.IsBound()
-		|| bIsReadOnly.IsBound()
-		|| bIsPassword.IsBound()
-		|| MinDesiredWidth.IsBound();
+		|| SlateEditableTextLayout.ComputeVolatility();
 }
 
 void SEditableText::SetHintText(TAttribute<FText> InHintText )
@@ -346,28 +372,30 @@ FText SEditableText::GetSearchText() const
 
 void SEditableText::SetIsReadOnly(TAttribute<bool> InIsReadOnly)
 {
-	bIsReadOnly = MoveTemp(InIsReadOnly);
+	SlateEditableTextLayout.SetIsReadOnly(MoveTemp(InIsReadOnly));
 }
 
 void SEditableText::SetIsPassword(TAttribute<bool> InIsPassword)
 {
 	PlainTextMarshaller->SetIsPassword(InIsPassword);
-	bIsPassword = MoveTemp(InIsPassword);
+	bIsPasswordAttribute.Assign(*this, MoveTemp(InIsPassword));
 }
 
 void SEditableText::SetColorAndOpacity(TAttribute<FSlateColor> Color)
 {
-	ColorAndOpacity = MoveTemp(Color);
+	bIsColorAndOpacityAttrSet = Color.IsSet();
+	ColorAndOpacityAttribute.Assign(*this, MoveTemp(Color));
 }
 
 void SEditableText::SetBackgroundImageSelected(TAttribute<const FSlateBrush*> InBackgroundImageSelected)
 {
-	BackgroundImageSelected = MoveTemp(InBackgroundImageSelected);
+	bIsBackgroundImageSelectedAttrSet = InBackgroundImageSelected.IsSet();
+	BackgroundImageSelectedAttribute.Assign(*this, MoveTemp(InBackgroundImageSelected));
 }
 
 void SEditableText::SetMinDesiredWidth(TAttribute<float> InMinDesiredWidth)
 {
-	MinDesiredWidth = MoveTemp(InMinDesiredWidth);
+	MinDesiredWidthAttribute.Assign(*this, MoveTemp(InMinDesiredWidth));
 }
 
 void SEditableText::SetIsCaretMovedWhenGainFocus(TAttribute<bool> InIsCaretMovedWhenGainFocus)
@@ -512,9 +540,9 @@ void SEditableText::SynchronizeTextStyle()
 	FTextBlockStyle NewTextStyle = SlateEditableTextLayout.GetTextStyle();
 
 	// Sync from the font override
-	if (Font.IsSet())
+	if (bIsFontAttrSet)
 	{
-		const FSlateFontInfo& NewFontInfo = Font.Get();
+		const FSlateFontInfo& NewFontInfo = FontAttribute.Get();
 		if (!NewTextStyle.Font.IsIdenticalTo(NewFontInfo))
 		{
 			NewTextStyle.Font = NewFontInfo;
@@ -523,9 +551,9 @@ void SEditableText::SynchronizeTextStyle()
 	}
 
 	// Sync from the color override
-	if (ColorAndOpacity.IsSet())
+	if (bIsColorAndOpacityAttrSet)
 	{
-		const FSlateColor& NewColorAndOpacity = ColorAndOpacity.Get();
+		const FSlateColor& NewColorAndOpacity = ColorAndOpacityAttribute.Get();
 		if (NewTextStyle.ColorAndOpacity != NewColorAndOpacity)
 		{
 			NewTextStyle.ColorAndOpacity = NewColorAndOpacity;
@@ -534,9 +562,9 @@ void SEditableText::SynchronizeTextStyle()
 	}
 
 	// Sync from the highlight shape override
-	if (BackgroundImageSelected.IsSet())
+	if (bIsBackgroundImageSelectedAttrSet)
 	{
-		const FSlateBrush* NewSelectionBrush = BackgroundImageSelected.Get();
+		const FSlateBrush* NewSelectionBrush = BackgroundImageSelectedAttribute.Get();
 		if (NewSelectionBrush && NewTextStyle.HighlightShape != *NewSelectionBrush)
 		{
 			NewTextStyle.HighlightShape = *NewSelectionBrush;
@@ -551,14 +579,19 @@ void SEditableText::SynchronizeTextStyle()
 	}
 }
 
+bool SEditableText::IsWidgetConstructed() const
+{
+	return IsConstructed();
+}
+
 bool SEditableText::IsTextReadOnly() const
 {
-	return bIsReadOnly.Get(false);
+	return SlateEditableTextLayout.GetIsReadOnly();
 }
 
 bool SEditableText::IsTextPassword() const
 {
-	return bIsPassword.Get(false);
+	return bIsPasswordAttribute.Get();
 }
 
 bool SEditableText::IsMultiLineTextEdit() const
@@ -722,6 +755,26 @@ float SEditableText::UpdateAndClampVerticalScrollBar(const float InViewOffset, c
 	return 0.0f;
 }
 
+ETextJustify::Type SEditableText::GetJustification() const
+{
+	return SlateEditableTextLayout.GetJustification();
+}
+
+FSlateColor SEditableText::GetColorAndOpacity() const
+{
+	return ColorAndOpacityAttribute.Get();
+}
+
+const FSlateBrush* SEditableText::GetBackgroundImageSelected() const
+{
+	return BackgroundImageSelectedAttribute.Get();
+}
+
+float SEditableText::GetMinDesiredWidth() const
+{
+	return MinDesiredWidthAttribute.Get();
+}
+
 #if WITH_ACCESSIBILITY
 TSharedRef<FSlateAccessibleWidget> SEditableText::CreateAccessibleWidget()
 {
diff --git a/Engine/Source/Runtime/Slate/Private/Widgets/Input/SEditableTextBox.cpp b/Engine/Source/Runtime/Slate/Private/Widgets/Input/SEditableTextBox.cpp
--- a/Engine/Source/Runtime/Slate/Private/Widgets/Input/SEditableTextBox.cpp
+++ b/Engine/Source/Runtime/Slate/Private/Widgets/Input/SEditableTextBox.cpp
@@ -28,7 +28,6 @@ SEditableTextBox::~SEditableTextBox() = default;
 void SEditableTextBox::Construct( const FArguments& InArgs )
 {
 	check (InArgs._Style);
-	SetStyle(InArgs._Style);
 
 	PaddingOverride = InArgs._Padding;
 	FontOverride = InArgs._Font;
@@ -97,6 +96,8 @@ void SEditableTextBox::Construct( const FArguments& InArgs )
 		]
 	);
 
+	SetStyle(InArgs._Style);
+
 	ErrorReporting = InArgs._ErrorReporting;
 	if ( ErrorReporting.IsValid() )
 	{
@@ -132,12 +133,7 @@ void SEditableTextBox::SetStyle(const FEditableTextBoxStyle* InStyle)
 
 void SEditableTextBox::SetTextBlockStyle(const FTextBlockStyle* InTextStyle)
 {
-	// The Construct() function will call this before EditableText exists,
-	// so we need a guard here to ignore that function call.
-	if (EditableText.IsValid())
-	{
-		EditableText->SetTextBlockStyle(InTextStyle);
-	}
+	EditableText->SetTextBlockStyle(InTextStyle);
 }
 
 void SEditableTextBox::SetText(TAttribute<FText> InNewText)
@@ -215,10 +211,7 @@ void SEditableTextBox::SetOverflowPolicy(const TOptional<ETextOverflowPolicy>& I
 
 void SEditableTextBox::SetFontFacesLoadingPaintPolicy(EFontFacesLoadingPaintPolicy InFontFacesLoadingPaintPolicy)
 {
-	if (EditableText)
-	{
-		EditableText->SetFontFacesLoadingPaintPolicy(InFontFacesLoadingPaintPolicy);
-	}
+	EditableText->SetFontFacesLoadingPaintPolicy(InFontFacesLoadingPaintPolicy);
 }
 
 bool SEditableTextBox::AnyTextSelected() const
diff --git a/Engine/Source/Runtime/Slate/Private/Widgets/Input/SMultiLineEditableTextBox.cpp b/Engine/Source/Runtime/Slate/Private/Widgets/Input/SMultiLineEditableTextBox.cpp
--- a/Engine/Source/Runtime/Slate/Private/Widgets/Input/SMultiLineEditableTextBox.cpp
+++ b/Engine/Source/Runtime/Slate/Private/Widgets/Input/SMultiLineEditableTextBox.cpp
@@ -34,7 +34,6 @@ namespace
 void SMultiLineEditableTextBox::Construct(const FArguments& InArgs)
 {
 	check (InArgs._Style);
-	SetStyle(InArgs._Style);
 
 	PaddingOverride = InArgs._Padding;
 	HScrollBarPaddingOverride = InArgs._HScrollBarPadding;
@@ -163,6 +162,8 @@ void SMultiLineEditableTextBox::Construct(const FArguments& InArgs)
 		]
 	);
 
+	SetStyle(InArgs._Style);
+
 	ErrorReporting = InArgs._ErrorReporting;
 	if ( ErrorReporting.IsValid() )
 	{
@@ -215,12 +216,7 @@ void SMultiLineEditableTextBox::SetStyle(const FEditableTextBoxStyle* InStyle)
 
 void SMultiLineEditableTextBox::SetTextStyle(const FTextBlockStyle* InTextStyle)
 {
-	// The Construct() function will call this before EditableText exists,
-	// so we need a guard here to ignore that function call.
-	if (EditableText.IsValid())
-	{
-		EditableText->SetTextStyle(InTextStyle);
-	}
+	EditableText->SetTextStyle(InTextStyle);
 }
 
 FMargin SMultiLineEditableTextBox::DeterminePadding() const
@@ -427,12 +423,9 @@ void SMultiLineEditableTextBox::SetIsReadOnly(TAttribute<bool> InIsReadOnly)
 	EditableText->SetIsReadOnly(MoveTemp(InIsReadOnly));
 }
 
-void SMultiLineEditableTextBox::SetFontFacesLoadingPaintPolicy( EFontFacesLoadingPaintPolicy InFontFacesLoadingPaintPolicy)
+void SMultiLineEditableTextBox::SetFontFacesLoadingPaintPolicy(EFontFacesLoadingPaintPolicy InFontFacesLoadingPaintPolicy)
 {
-	if (EditableText)
-	{
-		EditableText->SetFontFacesLoadingPaintPolicy(InFontFacesLoadingPaintPolicy);
-	}
+	EditableText->SetFontFacesLoadingPaintPolicy(InFontFacesLoadingPaintPolicy);
 }
 
 void SMultiLineEditableTextBox::SetError(const FText& InError)
diff --git a/Engine/Source/Runtime/Slate/Private/Widgets/Text/SMultiLineEditableText.cpp b/Engine/Source/Runtime/Slate/Private/Widgets/Text/SMultiLineEditableText.cpp
--- a/Engine/Source/Runtime/Slate/Private/Widgets/Text/SMultiLineEditableText.cpp
+++ b/Engine/Source/Runtime/Slate/Private/Widgets/Text/SMultiLineEditableText.cpp
@@ -14,19 +14,35 @@
 #include "Types/ReflectionMetadata.h"
 #include "Types/TrackedMetaData.h"
 
+SLATE_IMPLEMENT_WIDGET(SMultiLineEditableText)
+void SMultiLineEditableText::PrivateRegisterAttributes(FSlateAttributeInitializer& AttributeInitializer)
+{
+	struct FInvalidation
+	{
+		static void SynchronizeTextStyle(SWidget& OwningWidget)
+		{
+			SMultiLineEditableText& MultiLineEditableText = static_cast<SMultiLineEditableText&>(OwningWidget);
+			MultiLineEditableText.SynchronizeTextStyle();
+		}
+	};
+
+	FSlateEditableTextLayout::RegisterNestedAttributes<&SMultiLineEditableText::SlateEditableTextLayout>(AttributeInitializer);
+
+	SLATE_ADD_MEMBER_ATTRIBUTE_DEFINITION(AttributeInitializer, FontAttribute, EInvalidateWidgetReason::Layout)
+		.OnValueChanged(FSlateAttributeDescriptor::FAttributeValueChangedDelegate::CreateStatic(FInvalidation::SynchronizeTextStyle));
+}
+
 SMultiLineEditableText::SMultiLineEditableText()
 	: SlateEditableTextLayout(*this)
 	, bSelectAllTextWhenFocused(false)
-	, bIsReadOnly(false)
 	, AmountScrolledWhileRightMouseDown(0.0f)
 	, bIsSoftwareCursor(false)
+	, FontAttribute(*this)
 {
 }
 
 void SMultiLineEditableText::Construct( const FArguments& InArgs )
 {
-	bIsReadOnly = InArgs._IsReadOnly;
-
 	OnIsTypedCharValid = InArgs._OnIsTypedCharValid;
 	OnTextChangedCallback = InArgs._OnTextChanged;
 	OnTextCommittedCallback = InArgs._OnTextCommitted;
@@ -67,19 +83,13 @@ void SMultiLineEditableText::Construct( const FArguments& InArgs )
 		VScrollBar->SetOnUserScrolled(FOnUserScrolled::CreateSP(this, &SMultiLineEditableText::OnVScrollBarMoved));
 	}
 
-	FTextBlockStyle TextStyle = *InArgs._TextStyle;
-	if (InArgs._Font.IsSet() || InArgs._Font.IsBound())
-	{
-		TextStyle.SetFont(InArgs._Font.Get());
-	}
-
 	TSharedPtr<ITextLayoutMarshaller> Marshaller = InArgs._Marshaller;
 	if (!Marshaller.IsValid())
 	{
 		Marshaller = FPlainTextLayoutMarshaller::Create();
 	}
 
-	SlateEditableTextLayout.Construct(InArgs._Text, TextStyle, InArgs._TextShapingMethod, InArgs._TextFlowDirection, InArgs._CreateSlateTextLayout, Marshaller.ToSharedRef(), Marshaller.ToSharedRef());
+	SlateEditableTextLayout.Construct(InArgs._Text, *InArgs._TextStyle, InArgs._TextShapingMethod, InArgs._TextFlowDirection, InArgs._CreateSlateTextLayout, Marshaller.ToSharedRef(), Marshaller.ToSharedRef());
 	SlateEditableTextLayout.SetHintText(InArgs._HintText);
 	SlateEditableTextLayout.SetSearchText(InArgs._SearchText);
 	SlateEditableTextLayout.SetTextWrapping(InArgs._WrapTextAt, InArgs._AutoWrapText, InArgs._WrappingPolicy);
@@ -94,6 +104,9 @@ void SMultiLineEditableText::Construct( const FArguments& InArgs )
 		SlateEditableTextLayout.SetOnAllFontFacesFinishLoadingDelegate(InArgs._OnAllFontFacesFinishLoading);
 	}
 
+	SetIsReadOnly(InArgs._IsReadOnly);
+	SetFont(InArgs._Font);
+
 	// build context menu extender
 	MenuExtender = MakeShareable(new FExtender);
 	MenuExtender->AddMenuExtension("EditText", EExtensionHook::Before, TSharedPtr<FUICommandList>(), InArgs._ContextMenuExtender);
@@ -177,15 +190,13 @@ void SMultiLineEditableText::SetTextStyle(const FTextBlockStyle* InTextStyle)
 
 void SMultiLineEditableText::SetFont(TAttribute<FSlateFontInfo> InNewFont)
 {
-	FTextBlockStyle TextStyle = SlateEditableTextLayout.GetTextStyle();
-	TextStyle.SetFont(InNewFont.Get(TextStyle.Font));
-	SlateEditableTextLayout.SetTextStyle(TextStyle);
+	bIsFontAttrSet = InNewFont.IsSet();
+	FontAttribute.Assign(*this, MoveTemp(InNewFont));
 }
 
 FSlateFontInfo SMultiLineEditableText::GetFont() const
 {
-	FTextBlockStyle TextStyle = SlateEditableTextLayout.GetTextStyle();
-	return TextStyle.Font;
+	return bIsFontAttrSet ? FontAttribute.Get() : SlateEditableTextLayout.GetTextStyle().Font;
 }
 
 void SMultiLineEditableText::SetTextShapingMethod(const TOptional<ETextShapingMethod>& InTextShapingMethod)
@@ -250,7 +261,7 @@ void SMultiLineEditableText::SetVirtualKeyboardDismissAction(TAttribute<EVirtual
 
 void SMultiLineEditableText::SetIsReadOnly(TAttribute<bool> InIsReadOnly)
 {
-	bIsReadOnly = MoveTemp(InIsReadOnly);
+	SlateEditableTextLayout.SetIsReadOnly(MoveTemp(InIsReadOnly));
 }
 
 void SMultiLineEditableText::SetSelectAllTextWhenFocused(TAttribute<bool> InSelectAllTextWhenFocused)
@@ -295,9 +306,14 @@ void SMultiLineEditableText::OnVScrollBarMoved(const float InScrollOffsetFractio
 	OnVScrollBarUserScrolled.ExecuteIfBound(InScrollOffsetFraction);
 }
 
+bool SMultiLineEditableText::IsWidgetConstructed() const
+{
+	return IsConstructed();
+}
+
 bool SMultiLineEditableText::IsTextReadOnly() const
 {
-	return bIsReadOnly.Get(false);
+	return SlateEditableTextLayout.GetIsReadOnly();
 }
 
 bool SMultiLineEditableText::IsTextPassword() const
@@ -637,6 +653,30 @@ void SMultiLineEditableText::ForceScroll(int32 UserIndex, float ScrollAxisMagnit
 	OnMouseWheel(CachedGeom, FPointerEvent(UserIndex, 0, ScrollPos, ScrollPos, PressedKeys, EKeys::Invalid, ScrollAxisMagnitude, FModifierKeysState()));
 }
 
+void SMultiLineEditableText::SynchronizeTextStyle()
+{
+	// Has the style used for this editable text changed?
+	bool bTextStyleChanged = false;
+	FTextBlockStyle NewTextStyle = SlateEditableTextLayout.GetTextStyle();
+
+	// Sync from the font override
+	if (bIsFontAttrSet)
+	{
+		const FSlateFontInfo& NewFontInfo = FontAttribute.Get();
+		if (!NewTextStyle.Font.IsIdenticalTo(NewFontInfo))
+		{
+			NewTextStyle.Font = NewFontInfo;
+			bTextStyleChanged = true;
+		}
+	}
+
+	if (bTextStyleChanged)
+	{
+		SlateEditableTextLayout.SetTextStyle(NewTextStyle);
+		SlateEditableTextLayout.ForceRefreshTextLayout(SlateEditableTextLayout.GetEditableText());
+	}
+}
+
 void SMultiLineEditableText::Tick( const FGeometry& AllottedGeometry, const double InCurrentTime, const float InDeltaTime )
 {
 	SlateEditableTextLayout.Tick(AllottedGeometry, InCurrentTime, InDeltaTime);
@@ -677,6 +717,7 @@ int32 SMultiLineEditableText::OnPaint( const FPaintArgs& Args, const FGeometry&
 
 void SMultiLineEditableText::CacheDesiredSize(float LayoutScaleMultiplier)
 {
+	SynchronizeTextStyle();
 	SlateEditableTextLayout.CacheDesiredSize(LayoutScaleMultiplier);
 	SWidget::CacheDesiredSize(LayoutScaleMultiplier);
 }
@@ -859,9 +900,7 @@ bool SMultiLineEditableText::IsInteractable() const
 bool SMultiLineEditableText::ComputeVolatility() const
 {
 	return SWidget::ComputeVolatility()
-		|| HasKeyboardFocus()
-		|| SlateEditableTextLayout.ComputeVolatility()
-		|| bIsReadOnly.IsBound();
+		|| SlateEditableTextLayout.ComputeVolatility();
 }
 
 bool SMultiLineEditableText::IsRightClickScrolling() const
diff --git a/Engine/Source/Runtime/Slate/Private/Widgets/Text/SlateEditableTextLayout.cpp b/Engine/Source/Runtime/Slate/Private/Widgets/Text/SlateEditableTextLayout.cpp
--- a/Engine/Source/Runtime/Slate/Private/Widgets/Text/SlateEditableTextLayout.cpp
+++ b/Engine/Source/Runtime/Slate/Private/Widgets/Text/SlateEditableTextLayout.cpp
@@ -83,27 +83,21 @@ void FSlateEditableTextLayout::Construct(TAttribute<FText> InInitialText, FTextB
 	HintTextStyle = TextStyle;
 	TextLayout = CreateSlateTextLayout.Execute(&OwnerSlateWidget.Get(), TextStyle);
 
-	WrapTextAt = 0.0f;
-	AutoWrapText = false;
-	WrappingPolicy = ETextWrappingPolicy::DefaultWrapping;
-
-	Margin = FMargin(0);
-	Justification = ETextJustify::Left;
-	LineHeightPercentage = 1.0f;
 	DebugSourceInfo = FString();
 
 	SearchCase = ESearchCase::IgnoreCase;
 
 	GraphemeBreakIterator = GSlateUseSharedBreakIterator ? FBreakIterator::GetCharacterBoundaryIterator() : FBreakIterator::CreateCharacterBoundaryIterator();
 
-	BoundText = MoveTemp(InInitialText);
+	BoundText.Assign(OwnerSlateWidget.Get(), MoveTemp(InInitialText), ESlateAttributeBindAction::Initialize);
+	PreviousBoundText = BoundText.Get();
 
 	// Set the initial text - the same as SetText, but doesn't call OnTextChanged as that can cause a crash during construction
 	{
 		const bool bIsPassword = OwnerWidget->IsTextPassword();
 		TextLayout->SetIsPassword(bIsPassword);
 
-		const FText& InitialTextToSet = BoundText.Get(FText::GetEmpty());
+		const FText& InitialTextToSet = BoundText.Get();
 		SetEditableText(InitialTextToSet, true);
 
 		// Update the cached BoundText value to prevent it triggering another SetEditableText update again next Tick
@@ -143,6 +137,7 @@ void FSlateEditableTextLayout::Construct(TAttribute<FText> InInitialText, FTextB
 	bTextChangedByVirtualKeyboard = false;
 	bTextCommittedByVirtualKeyboard = false;
 	bSelectionChangedExternally = false;
+	bForceBoundTextReview = false;
 	VirtualKeyboardTextCommitType = ETextCommit::Default;
 
 	CachedSize = FVector2f::ZeroVector;
@@ -222,38 +217,14 @@ FSlateEditableTextLayout::~FSlateEditableTextLayout()
 
 void FSlateEditableTextLayout::SetText(TAttribute<FText> InText)
 {
-	const FText PreviousText = BoundText.Get(FText::GetEmpty());
-	BoundText = MoveTemp(InText);
-	const FText NewText = BoundText.Get(FText::GetEmpty());
-
-	// We need to force an update if the text doesn't match the editable text, as the editable 
-	// text may not match the current bound text since it may have been changed by the user
-	const FText EditableText = GetEditableText();
-	const bool bForceRefresh = !EditableText.ToString().Equals(NewText.ToString(), ESearchCase::CaseSensitive);
-
-	// Only emit the "text changed" event if the text has actually been changed
-	const bool bHasTextChanged = OwnerWidget->GetSlateWidget()->HasAnyUserFocus().IsSet()
-		? !NewText.ToString().Equals(EditableText.ToString(), ESearchCase::CaseSensitive)
-		: !NewText.ToString().Equals(PreviousText.ToString(), ESearchCase::CaseSensitive);
-
-	if (RefreshImpl(&NewText, bForceRefresh))
-	{
-		// Make sure we move the cursor to the end of the new text if we had keyboard focus
-		if (OwnerWidget->GetSlateWidget()->HasAnyUserFocus().IsSet() && !bWasFocusedByLastMouseDown)
-		{
-			JumpTo(ETextLocation::EndOfDocument, ECursorAction::MoveCursor);
-		}
+	bForceBoundTextReview = true;
+	BoundText.Assign(OwnerWidget->GetSlateWidget().Get(), MoveTemp(InText), FText::GetEmpty(), ESlateAttributeBindAction::Initialize);
 
-		// Let outsiders know that the text content has been changed
-		if (bHasTextChanged)
-		{
-			OwnerWidget->OnTextChanged(NewText);
-		}
-	}
-
-	if (bHasTextChanged || BoundText.IsBound())
+	if (bForceBoundTextReview)
 	{
-		OwnerWidget->GetSlateWidget()->Invalidate(EInvalidateWidget::LayoutAndVolatility);
+		// We would only arrive here if OnBoundTextChanged did not trigger (because the BoundText technically did not change),
+		// but we still want to make sure the EditableText is aligning so we call it manually
+		OnBoundTextChanged();
 	}
 }
 
@@ -266,41 +237,35 @@ FText FSlateEditableTextLayout::GetText() const
 {
 	SLATE_CROSS_THREAD_CHECK();
 
-	return BoundText.Get(FText::GetEmpty());
+	return BoundText.Get();
 }
 
 void FSlateEditableTextLayout::SetHintText(TAttribute<FText> InHintText)
 {
-	HintText = MoveTemp(InHintText);
+	const TSharedRef<SWidget> OwnerSlateWidget = OwnerWidget->GetSlateWidget();
+
+	HintText.Assign(OwnerSlateWidget.Get(), MoveTemp(InHintText), FText::GetEmpty(), ESlateAttributeBindAction::Initialize);
 
 	// If we have hint text that is either non-empty or bound to a delegate, we'll also need to make the hint text layout
-	if (HintText.IsBound() || !HintText.Get(FText::GetEmpty()).IsEmpty())
+	if (HintText.IsBound(OwnerSlateWidget.Get()) || !HintText.Get().IsEmpty())
 	{
-		HintTextLayout = MakeUnique<FSlateTextBlockLayout>(OwnerWidget->GetSlateWidgetPtr().Get(), HintTextStyle, TextLayout->GetTextShapingMethod(), TextLayout->GetTextFlowDirection(), CreateSlateTextLayout, HintMarshaller.ToSharedRef(), nullptr);
+		HintTextLayout = MakeUnique<FSlateTextBlockLayout>(&OwnerSlateWidget.Get(), HintTextStyle, TextLayout->GetTextShapingMethod(), TextLayout->GetTextFlowDirection(), CreateSlateTextLayout, HintMarshaller.ToSharedRef(), nullptr);
 		HintTextLayout->SetDebugSourceInfo(DebugSourceInfo);
 	}
 	else
 	{
 		HintTextLayout.Reset();
 	}
-
-	OwnerWidget->GetSlateWidget()->Invalidate(EInvalidateWidget::LayoutAndVolatility);
 }
 
 FText FSlateEditableTextLayout::GetHintText() const
 {
-	return HintText.Get(FText::GetEmpty());
+	return HintText.Get();
 }
 
 void FSlateEditableTextLayout::SetSearchText(TAttribute<FText> InSearchText)
 {
-	const FText& SearchTextToSet = InSearchText.Get(FText::GetEmpty());
-
-	BoundSearchText = MoveTemp(InSearchText);
-	BoundSearchTextLastTick = FTextSnapshot(SearchTextToSet);
-
-	BeginSearch(SearchTextToSet);
-	OwnerWidget->GetSlateWidget()->Invalidate(EInvalidateWidget::LayoutAndVolatility);
+	BoundSearchText.Assign(OwnerWidget->GetSlateWidget().Get(), MoveTemp(InSearchText));
 }
 
 FText FSlateEditableTextLayout::GetSearchText() const
@@ -337,12 +302,22 @@ const FTextBlockStyle& FSlateEditableTextLayout::GetTextStyle() const
 
 void FSlateEditableTextLayout::SetHintTextOpacity(TAttribute<float> InOpacity)
 {
-	HintTextOpacity = MoveTemp(InOpacity);
+	HintTextOpacity.Assign(OwnerWidget->GetSlateWidget().Get(), MoveTemp(InOpacity));
 }
 
 float FSlateEditableTextLayout::GetHintTextOpacity() const
 {
-	return HintTextOpacity.Get(0.35f);
+	return HintTextOpacity.Get();
+}
+
+void FSlateEditableTextLayout::SetIsReadOnly(TAttribute<bool> InIsReadOnly)
+{
+	bIsReadOnlyAttribute.Assign(OwnerWidget->GetSlateWidget().Get(), MoveTemp(InIsReadOnly));
+}
+
+bool FSlateEditableTextLayout::GetIsReadOnly() const
+{
+	return bIsReadOnlyAttribute.Get();
 }
 
 void FSlateEditableTextLayout::SetCursorBrush(const TAttribute<const FSlateBrush*>& InCursorBrush)
@@ -492,68 +467,51 @@ void FSlateEditableTextLayout::SetTextFlowDirection(const TOptional<ETextFlowDir
 
 void FSlateEditableTextLayout::SetTextWrapping(TAttribute<float> InWrapTextAt, TAttribute<bool> InAutoWrapText, TAttribute<ETextWrappingPolicy> InWrappingPolicy)
 {
-	WrapTextAt = MoveTemp(InWrapTextAt);
-	AutoWrapText = MoveTemp(InAutoWrapText);
-	WrappingPolicy = MoveTemp(InWrappingPolicy);
+	const TSharedRef<SWidget> OwnerSlateWidget = OwnerWidget->GetSlateWidget();
 
-	OwnerWidget->GetSlateWidget()->Invalidate(EInvalidateWidget::LayoutAndVolatility);
+	WrapTextAt.Assign(OwnerSlateWidget.Get(), MoveTemp(InWrapTextAt));
+	AutoWrapText.Assign(OwnerSlateWidget.Get(), MoveTemp(InAutoWrapText));
+	WrappingPolicy.Assign(OwnerSlateWidget.Get(), MoveTemp(InWrappingPolicy));
 }
 
 void FSlateEditableTextLayout::SetWrapTextAt(TAttribute<float> InWrapTextAt)
 {
-	WrapTextAt = MoveTemp(InWrapTextAt);
-
-	OwnerWidget->GetSlateWidget()->Invalidate(EInvalidateWidget::LayoutAndVolatility);
+	WrapTextAt.Assign(OwnerWidget->GetSlateWidget().Get(), MoveTemp(InWrapTextAt));
 }
 
 void FSlateEditableTextLayout::SetAutoWrapText(TAttribute<bool> InAutoWrapText)
 {
-	AutoWrapText = MoveTemp(InAutoWrapText);
-
-	OwnerWidget->GetSlateWidget()->Invalidate(EInvalidateWidget::LayoutAndVolatility);
+	AutoWrapText.Assign(OwnerWidget->GetSlateWidget().Get(), MoveTemp(InAutoWrapText));
 }
 
 void FSlateEditableTextLayout::SetWrappingPolicy(TAttribute<ETextWrappingPolicy> InWrappingPolicy)
 {
-	WrappingPolicy = MoveTemp(InWrappingPolicy);
-
-	OwnerWidget->GetSlateWidget()->Invalidate(EInvalidateWidget::LayoutAndVolatility);
+	WrappingPolicy.Assign(OwnerWidget->GetSlateWidget().Get(), MoveTemp(InWrappingPolicy));
 }
 
 void FSlateEditableTextLayout::SetMargin(TAttribute<FMargin> InMargin)
 {
-	Margin = MoveTemp(InMargin);
-
-	OwnerWidget->GetSlateWidget()->Invalidate(EInvalidateWidget::LayoutAndVolatility);
+	Margin.Assign(OwnerWidget->GetSlateWidget().Get(), MoveTemp(InMargin));
 }
 
 void FSlateEditableTextLayout::SetJustification(TAttribute<ETextJustify::Type> InJustification)
 {
-	if (!Justification.IdenticalTo(InJustification))
-	{
-		Justification = MoveTemp(InJustification);
-
-		OwnerWidget->GetSlateWidget()->Invalidate(EInvalidateWidget::LayoutAndVolatility);
-	}
+	Justification.Assign(OwnerWidget->GetSlateWidget().Get(), MoveTemp(InJustification));
 }
 
 ETextJustify::Type FSlateEditableTextLayout::GetJustification() const
 {
-	return Justification.Get(ETextJustify::Left);
+	return Justification.Get();
 }
 
 void FSlateEditableTextLayout::SetLineHeightPercentage(TAttribute<float> InLineHeightPercentage)
 {
-	LineHeightPercentage = MoveTemp(InLineHeightPercentage);
-
-	OwnerWidget->GetSlateWidget()->Invalidate(EInvalidateWidget::LayoutAndVolatility);
+	LineHeightPercentage.Assign(OwnerWidget->GetSlateWidget().Get(), MoveTemp(InLineHeightPercentage));
 }
 
 void FSlateEditableTextLayout::SetApplyLineHeightToBottomLine(TAttribute<bool> InApplyLineHeightToBottomLine)
 {
-	ApplyLineHeightToBottomLine = MoveTemp(InApplyLineHeightToBottomLine);
-
-	OwnerWidget->GetSlateWidget()->Invalidate(EInvalidateWidget::LayoutAndVolatility);
+	ApplyLineHeightToBottomLine.Assign(OwnerWidget->GetSlateWidget().Get(), MoveTemp(InApplyLineHeightToBottomLine));
 }
 
 void FSlateEditableTextLayout::SetOverflowPolicy(const TOptional<ETextOverflowPolicy>& InOverflowPolicy)
@@ -567,7 +525,7 @@ void FSlateEditableTextLayout::SetOverflowPolicy(const TOptional<ETextOverflowPo
 			HintTextLayout->SetTextOverflowPolicy(OverflowPolicyOverride);
 		}
 
-		OwnerWidget->GetSlateWidget()->Invalidate(EInvalidateWidget::LayoutAndVolatility);
+		OwnerWidget->GetSlateWidget()->Invalidate(EInvalidateWidget::Layout);
 	}
 }
 
@@ -623,7 +581,7 @@ void FSlateEditableTextLayout::EnableTextInputMethodContext()
 
 bool FSlateEditableTextLayout::Refresh()
 {
-	const FText& TextToSet = BoundText.Get(FText::GetEmpty());
+	const FText& TextToSet = BoundText.Get();
 	return RefreshImpl(&TextToSet);
 }
 
@@ -642,7 +600,7 @@ bool FSlateEditableTextLayout::RefreshImpl(const FText* InTextToSet, const bool
 		if (bForce || !BoundTextLastTick.IsDisplayStringEqualTo(*InTextToSet))
 		{
 			// The source text has changed, so update the internal editable text
-			bHasSetText = SetEditableText(*InTextToSet, true);
+			bHasSetText = SetEditableText(*InTextToSet, bForce);
 		}
 
 		// Update this even if the text is lexically identical, as it will update the pointer compared by IdenticalTo for the next Tick
@@ -850,7 +808,7 @@ float FSlateEditableTextLayout::GetComputedWrappingWidth() const
 
 bool FSlateEditableTextLayout::GetAutoWrapText() const
 {
-	return AutoWrapText.Get(false);
+	return AutoWrapText.Get();
 }
 
 bool FSlateEditableTextLayout::HandleFocusReceived(const FFocusEvent& InFocusEvent)
@@ -912,7 +870,7 @@ bool FSlateEditableTextLayout::HandleFocusReceived(const FFocusEvent& InFocusEve
 	// Force update of last cursor interaction time to ensure the caret blinks at the correct frequency regardless of how focus is set.
 	CursorInfo.UpdateLastCursorInteractionTime();
 
-	OwnerWidget->OnBeginTextEdit(BoundText.Get(FText::GetEmpty()));
+	OwnerWidget->OnBeginTextEdit(BoundText.Get());
 	return true;
 }
 
@@ -1884,7 +1842,7 @@ void FSlateEditableTextLayout::SelectAllText()
 	const FTextLocation NewCursorPosition = FTextLocation(NumberOfLines - 1, Lines[NumberOfLines - 1].Text->Len());
 	CursorInfo.SetCursorLocationAndCalculateAlignment(*TextLayout, NewCursorPosition);
 	UpdateCursorHighlight();
-	OwnerWidget->GetSlateWidget()->Invalidate(EInvalidateWidget::LayoutAndVolatility);
+	OwnerWidget->GetSlateWidget()->Invalidate(EInvalidateWidget::Layout);
 }
 
 void FSlateEditableTextLayout::SelectWordAt(const FGeometry& MyGeometry, const UE::Slate::FDeprecateVector2DParameter& ScreenSpacePosition)
@@ -2730,7 +2688,6 @@ void FSlateEditableTextLayout::UpdateCursorHighlight()
 	const bool bIsComposing = TextInputMethodContext->IsComposing();
 	const bool bHasSelection = SelectionLocation != CursorInteractionPosition;
 	const bool bHasSearch = !SearchText.IsEmpty();
-	const bool bIsReadOnly = OwnerWidget->IsTextReadOnly();
 
 	if (bHasSearch)
 	{
@@ -2827,7 +2784,7 @@ void FSlateEditableTextLayout::UpdateCursorHighlight()
 		}
 	}
 
-	if (bHasKeyboardFocus && !bIsReadOnly)
+	if (bHasKeyboardFocus)
 	{
 		// The cursor mode uses the literal position rather than the interaction position
 		const FTextLocation CursorPosition = CursorInfo.GetCursorLocation();
@@ -3431,37 +3388,38 @@ void FSlateEditableTextLayout::Redo()
 
 void FSlateEditableTextLayout::SaveText(const FText& TextToSave)
 {
+	const TSharedRef<SWidget> OwnerSlateWidget = OwnerWidget->GetSlateWidget();
 	// Don't set text if the text attribute has a 'getter' binding on it, otherwise we'd blow away
 	// that binding.  If there is a getter binding, then we'll assume it will provide us with
 	// updated text after we've fired our 'text changed' callbacks
-	if (!BoundText.IsBound())
+	if (!BoundText.IsBound(OwnerSlateWidget.Get()))
 	{
-		BoundText.Set(TextToSave);
+		BoundText.Assign(OwnerSlateWidget.Get(), TextToSave);
 	}
 }
 
 void FSlateEditableTextLayout::LoadText()
 {
+	const TSharedRef<SWidget> OwnerSlateWidget = OwnerWidget->GetSlateWidget();
 	// We only need to do this if we're bound to a delegate, otherwise the text layout will already be up-to-date
 	// either from Construct, or a call to SetText
-	if (BoundText.IsBound())
+	if (BoundText.IsBound(OwnerSlateWidget.Get()))
 	{
-		SetText(BoundText);
+		bForceBoundTextReview = true;
+		BoundText.UpdateNow(OwnerSlateWidget.Get());
+		if (bForceBoundTextReview)
+		{
+			// We would only arrive here if OnBoundTextChanged did not trigger (because the BoundText technically did not change),
+			// but we still want to make sure the EditableText is aligning so we call it manually
+			OnBoundTextChanged();
+		}
 		TextLayout->UpdateIfNeeded();
 	}
 }
 
 bool FSlateEditableTextLayout::ComputeVolatility() const
 {
-	return BoundText.IsBound()
-		|| HintText.IsBound()
-		|| BoundSearchText.IsBound()
-		|| WrapTextAt.IsBound()
-		|| AutoWrapText.IsBound()
-		|| WrappingPolicy.IsBound()
-		|| Margin.IsBound()
-		|| Justification.IsBound()
-		|| LineHeightPercentage.IsBound();
+	return OwnerWidget->GetSlateWidget()->HasAnyUserFocus().IsSet() && CursorLineHighlighter->GetDrawCursor();
 }
 
 void FSlateEditableTextLayout::UpdateTextChangedByVirtualKeyboard()
@@ -3529,7 +3487,7 @@ void FSlateEditableTextLayout::Tick(const FGeometry& AllottedGeometry, const dou
 
 	// Update the search before we process the next PositionToScrollIntoView
 	{
-		const FText& SearchTextToSet = BoundSearchText.Get(FText::GetEmpty());
+		const FText& SearchTextToSet = BoundSearchText.Get();
 		if (!BoundSearchTextLastTick.IdenticalTo(SearchTextToSet))
 		{
 			// The pointer used by the bound text has changed, however the text may still be the same - check that now
@@ -3549,7 +3507,7 @@ void FSlateEditableTextLayout::Tick(const FGeometry& AllottedGeometry, const dou
 	// If we're auto-wrapping, we need to hide the scrollbars until the first valid auto-wrap has been performed
 	// If we don't do this, then we can get some nasty layout shuffling as the scrollbars appear for one frame and then vanish again
 	// We also hide the scrollbars for non-multi-line text widgets
-	const EVisibility ScrollBarVisibilityOverride = ((AutoWrapText.Get(false) && CachedSize.IsZero()) || !OwnerWidget->IsMultiLineTextEdit()) ? EVisibility::Collapsed : EVisibility::Visible;
+	const EVisibility ScrollBarVisibilityOverride = ((AutoWrapText.Get() && CachedSize.IsZero()) || !OwnerWidget->IsMultiLineTextEdit()) ? EVisibility::Collapsed : EVisibility::Visible;
 
 	// Try and make sure that the line containing the cursor is in view
 	if (PositionToScrollIntoView.IsSet())
@@ -3647,27 +3605,27 @@ void FSlateEditableTextLayout::CacheDesiredSize(float LayoutScaleMultiplier)
 	const float CaretWidth = FTextEditHelper::CalculateCaretWidth(FontMaxCharHeight);
 
 	// Get the wrapping width and font to see if they have changed
-	float WrappingWidth = WrapTextAt.Get(0.0f);
+	float WrappingWidth = WrapTextAt.Get();
 
 	// Text wrapping can either be used defined (WrapTextAt), automatic (AutoWrapText), or a mixture of both
 	// Take whichever has the smallest value (>1)
-	if (AutoWrapText.Get(false) && CachedSize.X >= 1.0f)
+	if (AutoWrapText.Get() && CachedSize.X >= 1.0f)
 	{
 		WrappingWidth = (WrappingWidth >= 1.0f) ? FMath::Min(WrappingWidth, CachedSize.X) : CachedSize.X;
 	}
 
 	// Append the caret width to the margin to make sure it doesn't get clipped
-	FMargin MarginValue = Margin.Get(FMargin(0));
+	FMargin MarginValue = Margin.Get();
 	MarginValue.Left += CaretWidth;
 	MarginValue.Right += CaretWidth;
 
 	TextLayout->SetScale(LayoutScaleMultiplier);
 	TextLayout->SetWrappingWidth(WrappingWidth);
-	TextLayout->SetWrappingPolicy(WrappingPolicy.Get(ETextWrappingPolicy::DefaultWrapping));
+	TextLayout->SetWrappingPolicy(WrappingPolicy.Get());
 	TextLayout->SetMargin(MarginValue);
-	TextLayout->SetLineHeightPercentage(LineHeightPercentage.Get(1.0f));
-	TextLayout->SetApplyLineHeightToBottomLine(ApplyLineHeightToBottomLine.Get(false));
-	TextLayout->SetJustification(Justification.Get(ETextJustify::Left));
+	TextLayout->SetLineHeightPercentage(LineHeightPercentage.Get());
+	TextLayout->SetApplyLineHeightToBottomLine(ApplyLineHeightToBottomLine.Get());
+	TextLayout->SetJustification(Justification.Get());
 	TextLayout->SetVisibleRegion(FVector2D(CachedSize), FVector2D(ScrollOffset) * TextLayout->GetScale());
 	TextLayout->UpdateIfNeeded();
 }
@@ -3679,7 +3637,7 @@ FVector2D FSlateEditableTextLayout::ComputeDesiredSize(float LayoutScaleMultipli
 	const float FontMaxCharHeight = FTextEditHelper::GetFontHeight(TextStyle.Font);
 	const float CaretWidth = FTextEditHelper::CalculateCaretWidth(FontMaxCharHeight);
 
-	const float WrappingWidth = WrapTextAt.Get(0.0f);
+	const float WrappingWidth = WrapTextAt.Get();
 	float DesiredWidth = 0.0f;
 	float DesiredHeight = 0.0f;
 
@@ -3687,12 +3645,12 @@ FVector2D FSlateEditableTextLayout::ComputeDesiredSize(float LayoutScaleMultipli
 	if (TextLayout->IsEmpty() && HintTextLayout.IsValid())
 	{
 		// Append the caret width to the margin to mimic what happens to the main text layout
-		FMargin MarginValue = Margin.Get(FMargin(0));
+		FMargin MarginValue = Margin.Get();
 		MarginValue.Left += CaretWidth;
 		MarginValue.Right += CaretWidth;
 
 		const FVector2D HintTextSize = HintTextLayout->ComputeDesiredSize(
-			FSlateTextBlockLayout::FWidgetDesiredSizeArgs(HintText.Get(FText::GetEmpty()), FText::GetEmpty(), WrappingWidth, AutoWrapText.Get(false), WrappingPolicy.Get(ETextWrappingPolicy::DefaultWrapping), ETextTransformPolicy::None, MarginValue, LineHeightPercentage.Get(1.0f), ApplyLineHeightToBottomLine.Get(false), Justification.Get(ETextJustify::Left)),
+			FSlateTextBlockLayout::FWidgetDesiredSizeArgs(HintText.Get(), FText::GetEmpty(), WrappingWidth, AutoWrapText.Get(), WrappingPolicy.Get(), ETextTransformPolicy::None, MarginValue, LineHeightPercentage.Get(), ApplyLineHeightToBottomLine.Get(), Justification.Get()),
 			LayoutScaleMultiplier, HintTextStyle
 			);
 
@@ -3834,11 +3792,13 @@ void FSlateEditableTextLayout::FVirtualKeyboardEntry::SetTextFromVirtualKeyboard
 {
 	check(IsInGameThread());
 
+	const TSharedRef<SWidget> OwnerSlateWidget = OwnerLayout->OwnerWidget->GetSlateWidget();
+
 	// Only set the text if the text attribute doesn't have a getter binding (otherwise it would be blown away).
 	// If it is bound, we'll assume that OnTextCommitted will handle the update.
-	if (!OwnerLayout->BoundText.IsBound())
+	if (!OwnerLayout->BoundText.IsBound(OwnerSlateWidget.Get()))
 	{
-		OwnerLayout->BoundText.Set(InNewText);
+		OwnerLayout->BoundText.Assign(OwnerSlateWidget.Get(), InNewText);
 	}
 
 	// Update the internal editable text
@@ -4263,4 +4223,59 @@ void FSlateEditableTextLayout::SetFontFacesLoadingPaintPolicy(EFontFacesLoadingP
 	{
 		TextLayout->SetFontFacesLoadingPaintPolicy(InFontFacesLoadingPaintPolicy);
 	}
-}
\ No newline at end of file
+}
+
+void FSlateEditableTextLayout::OnBoundTextChanged()
+{
+	const TSharedRef<SWidget> Widget = OwnerWidget->GetSlateWidget();
+
+	// During construction, we do not want to trigger events.
+	// If we have focus, the bound text might have changed but the EditableText prevails and we don't want to replace it
+	if (!OwnerWidget->IsWidgetConstructed() || (Widget->HasAnyUserFocus().IsSet() && !bForceBoundTextReview))
+	{
+		bForceBoundTextReview = false;
+		return;
+	}
+
+	const FText NewBoundText = BoundText.Get();
+
+	// We need to force an update if the text doesn't match the editable text, as the editable
+	// text may not match the current bound text since it may have been changed by the user
+	const FText EditableText = GetEditableText();
+	const bool bForceRefresh = !EditableText.ToString().Equals(NewBoundText.ToString(), ESearchCase::CaseSensitive);
+
+	// Only emit the "text changed" event if the text has actually been changed
+	const bool bHasTextChanged = Widget->HasAnyUserFocus().IsSet()
+		? !NewBoundText.ToString().Equals(EditableText.ToString(), ESearchCase::CaseSensitive)
+		: !NewBoundText.ToString().Equals(PreviousBoundText.ToString(), ESearchCase::CaseSensitive);
+
+	if (RefreshImpl(&NewBoundText, bForceRefresh))
+	{
+		// Make sure we move the cursor to the end of the new text if we had keyboard focus, but only if we set it via SetText or LoadText
+		if (bForceBoundTextReview && Widget->HasAnyUserFocus().IsSet() && !bWasFocusedByLastMouseDown)
+		{
+			JumpTo(ETextLocation::EndOfDocument, ECursorAction::MoveCursor);
+		}
+
+		// Let outsiders know that the text content has been changed
+		if (bHasTextChanged)
+		{
+			OwnerWidget->OnTextChanged(NewBoundText);
+		}
+	}
+
+	PreviousBoundText = NewBoundText;
+	bForceBoundTextReview = false;
+}
+
+void FSlateEditableTextLayout::OnBoundSearchTextChanged(SWidget& Widget, FSlateEditableTextLayout& SlateTextLayout)
+{
+	const FText& SearchTextToSet = SlateTextLayout.BoundSearchText.Get();
+	SlateTextLayout.BoundSearchTextLastTick = FTextSnapshot(SearchTextToSet);
+	SlateTextLayout.BeginSearch(SearchTextToSet);
+}
+
+void FSlateEditableTextLayout::OnIsReadOnlyChanged(SWidget& Widget, FSlateEditableTextLayout& SlateTextLayout)
+{
+	SlateTextLayout.CursorLineHighlighter->SetDrawCursor(!SlateTextLayout.bIsReadOnlyAttribute.Get());
+}
diff --git a/Engine/Source/Runtime/Slate/Private/Widgets/Text/SlateEditableTextTypes.cpp b/Engine/Source/Runtime/Slate/Private/Widgets/Text/SlateEditableTextTypes.cpp
--- a/Engine/Source/Runtime/Slate/Private/Widgets/Text/SlateEditableTextTypes.cpp
+++ b/Engine/Source/Runtime/Slate/Private/Widgets/Text/SlateEditableTextTypes.cpp
@@ -105,6 +105,11 @@ void FCursorLineHighlighter::SetCursorBrush(const TAttribute<const FSlateBrush*>
 
 int32 FCursorLineHighlighter::OnPaint(const FPaintArgs& Args, const FTextLayout::FLineView& Line, const FVector2D Offset, const float Width, const FTextBlockStyle& DefaultStyle, const FGeometry& AllottedGeometry, const FSlateRect& MyCullingRect, FSlateWindowElementList& OutDrawElements, int32 LayerId, const FWidgetStyle& InWidgetStyle, bool bParentEnabled) const
 {
+	if (!bDrawCursor)
+	{
+		return LayerId;
+	}
+
 	const FVector2D Location(Line.Offset.X + Offset.X, Line.Offset.Y + Offset.Y);
 	const FVector2D Size(Width, Line.TextHeight);
 
diff --git a/Engine/Source/Runtime/Slate/Public/Widgets/Input/SEditableText.h b/Engine/Source/Runtime/Slate/Public/Widgets/Input/SEditableText.h
--- a/Engine/Source/Runtime/Slate/Public/Widgets/Input/SEditableText.h
+++ b/Engine/Source/Runtime/Slate/Public/Widgets/Input/SEditableText.h
@@ -34,11 +34,18 @@ enum class ETextShapingMethod : uint8;
 
 DECLARE_DELEGATE_TwoParams(FOnCursorMovedWithSelection, const FTextLocation&, const FTextSelection&);
 
+template <>
+struct TWidgetTypeTraits<class SEditableText>
+{
+	static constexpr bool SupportsInvalidation() { return true; }
+};
+
 /**
  * Editable text widget
  */
 class SEditableText : public SWidget, public ISlateEditableTextWidget
 {
+	SLATE_DECLARE_WIDGET_API(SEditableText, SWidget, SLATE_API)
 public:
 	SLATE_BEGIN_ARGS( SEditableText )
 		: _Text()
@@ -459,6 +466,7 @@ protected:
 
 public:
 	//~ Begin ISlateEditableTextWidget Interface
+	SLATE_API virtual bool IsWidgetConstructed() const override;
 	SLATE_API virtual bool IsTextReadOnly() const override;
 	SLATE_API virtual bool IsTextPassword() const override;
 	SLATE_API virtual bool IsMultiLineTextEdit() const override;
@@ -492,6 +500,37 @@ protected:
 	SLATE_API virtual float UpdateAndClampVerticalScrollBar(const float InViewOffset, const float InViewFraction, const EVisibility InVisibilityOverride) override;
 	//~ End ISlateEditableTextWidget Interface
 
+protected:
+	SLATE_API ETextJustify::Type GetJustification() const;
+
+	TSlateAttributeRef<FSlateFontInfo> GetFontAttribute() const
+	{
+		return TSlateAttributeRef<FSlateFontInfo>{SharedThis(this), FontAttribute};
+	}
+
+	SLATE_API FSlateColor GetColorAndOpacity() const;
+	TSlateAttributeRef<FSlateColor> GetColorAndOpacityAttribute() const
+	{
+		return TSlateAttributeRef<FSlateColor>{SharedThis(this), ColorAndOpacityAttribute};
+	}
+
+	SLATE_API const FSlateBrush* GetBackgroundImageSelected() const;
+	TSlateAttributeRef<const FSlateBrush*> GetBackgroundImageSelectedAttribute() const
+	{
+		return TSlateAttributeRef<const FSlateBrush*>{SharedThis(this), BackgroundImageSelectedAttribute};
+	}
+
+	TSlateAttributeRef<bool> GetIsPasswordAttribute() const
+	{
+		return TSlateAttributeRef<bool>{SharedThis(this), bIsPasswordAttribute};
+	}
+
+	SLATE_API float GetMinDesiredWidth() const;
+	TSlateAttributeRef<float> GetMinDesiredWidthAttribute() const
+	{
+		return TSlateAttributeRef<float>{SharedThis(this), MinDesiredWidthAttribute};
+	}
+
 protected:
 	/** Text marshaller used by the editable text layout */
 	TSharedPtr<FPlainTextLayoutMarshaller> PlainTextMarshaller;
@@ -504,20 +543,27 @@ protected:
 	/** The text layout that deals with the editable text */
 	FSlateEditableTextLayout SlateEditableTextLayout;
 
+#if WITH_EDITORONLY_DATA
 	/** The font used to draw the text */
-	TAttribute<FSlateFontInfo> Font;
+	UE_DEPRECATED(5.8, "Use SetFont / GetFont / GetFontAttribute")
+	TSlateDeprecatedTAttribute<FSlateFontInfo> Font;
 
 	/** Text color and opacity */
-	TAttribute<FSlateColor> ColorAndOpacity;
+	UE_DEPRECATED(5.8, "Use SetColorAndOpacity / GetColorAndOpacity / GetColorAndOpacityAttribute")
+	TSlateDeprecatedTAttribute<FSlateColor> ColorAndOpacity;
 
 	/** Background image for the selected text */
-	TAttribute<const FSlateBrush*> BackgroundImageSelected;
+	UE_DEPRECATED(5.8, "Use SetBackgroundImageSelected / GetBackgroundImageSelected / GetBackgroundImageSelectedAttribute")
+	TSlateDeprecatedTAttribute<const FSlateBrush*> BackgroundImageSelected;
 
 	/** Sets whether this text box can actually be modified interactively by the user */
-	TAttribute<bool> bIsReadOnly;
+	UE_DEPRECATED(5.8, "Use SetIsReadOnly / IsTextReadOnly")
+	TSlateDeprecatedTAttribute<bool> bIsReadOnly;
 
 	/** Sets whether this text box is for storing a password */
-	TAttribute<bool> bIsPassword;
+	UE_DEPRECATED(5.8, "Use SetIsPassword / IsTextPassword / GetIsPasswordAttribute")
+	TSlateDeprecatedTAttribute<bool> bIsPassword;
+#endif
 
 	/** Workaround as we loose focus when the auto completion closes. */
 	TAttribute<bool> bIsCaretMovedWhenGainFocus;
@@ -561,8 +607,11 @@ protected:
 	/** Called whenever the user's selection and/or cursor position of the text has changed. */
 	FOnCursorMovedWithSelection OnCursorMovedWithSelectionCallback;
 
+#if WITH_EDITORONLY_DATA
 	/** Prevents the editable text from being smaller than desired in certain cases (e.g. when it is empty) */
-	TAttribute<float> MinDesiredWidth;
+	UE_DEPRECATED(5.8, "Use SetMinDesiredWidth / GetMinDesiredWidth / GetMinDesiredWidthAttribute")
+	TSlateDeprecatedTAttribute<float> MinDesiredWidth;
+#endif
 
 	/** Menu extender for right-click context menu */
 	TSharedPtr<FExtender> MenuExtender;
@@ -590,4 +639,23 @@ protected:
 
 	/** Additional options used by the virtual keyboard summoned by this widget */
 	FVirtualKeyboardOptions VirtualKeyboardOptions;
+
+private:
+	/** The font used to draw the text */
+	TSlateAttribute<FSlateFontInfo> FontAttribute;
+	bool bIsFontAttrSet = false;
+
+	/** Text color and opacity */
+	TSlateAttribute<FSlateColor> ColorAndOpacityAttribute;
+	bool bIsColorAndOpacityAttrSet = false;
+
+	/** Background image for the selected text */
+	TSlateAttribute<const FSlateBrush*> BackgroundImageSelectedAttribute;
+	bool bIsBackgroundImageSelectedAttrSet = false;
+
+	/** Sets whether this text box is for storing a password */
+	TSlateAttribute<bool> bIsPasswordAttribute;
+
+	/** Prevents the editable text from being smaller than desired in certain cases (e.g. when it is empty) */
+	TSlateAttribute<float> MinDesiredWidthAttribute;
 };
diff --git a/Engine/Source/Runtime/Slate/Public/Widgets/Input/SEditableTextBox.h b/Engine/Source/Runtime/Slate/Public/Widgets/Input/SEditableTextBox.h
--- a/Engine/Source/Runtime/Slate/Public/Widgets/Input/SEditableTextBox.h
+++ b/Engine/Source/Runtime/Slate/Public/Widgets/Input/SEditableTextBox.h
@@ -23,6 +23,12 @@ class SHorizontalBox;
 enum class ETextFlowDirection : uint8;
 enum class ETextShapingMethod : uint8;
 
+template <>
+struct TWidgetTypeTraits<class SEditableTextBox>
+{
+	static constexpr bool SupportsInvalidation() { return true; }
+};
+
 /**
  * Editable text box widget
  */
diff --git a/Engine/Source/Runtime/Slate/Public/Widgets/Input/SMultiLineEditableTextBox.h b/Engine/Source/Runtime/Slate/Public/Widgets/Input/SMultiLineEditableTextBox.h
--- a/Engine/Source/Runtime/Slate/Public/Widgets/Input/SMultiLineEditableTextBox.h
+++ b/Engine/Source/Runtime/Slate/Public/Widgets/Input/SMultiLineEditableTextBox.h
@@ -27,6 +27,12 @@ enum class ETextShapingMethod : uint8;
 
 #if WITH_FANCY_TEXT
 
+template <>
+struct TWidgetTypeTraits<class SMultiLineEditableTextBox>
+{
+	static constexpr bool SupportsInvalidation() { return true; }
+};
+
 /**
  * Editable text box widget
  */
@@ -514,7 +520,7 @@ protected:
 protected:
 
 	/** Editable text widget */
-	TSharedPtr< SMultiLineEditableText > EditableText;
+	TSharedPtr<SMultiLineEditableText> EditableText;
 
 	/** Padding (overrides style) */
 	TAttribute<FMargin> PaddingOverride;
@@ -544,7 +550,7 @@ protected:
 	TAttribute<int32> MaximumLength;
 
 	/** Whether to disable the context menu */
-	TAttribute< bool > AllowContextMenu;
+	TAttribute<bool> AllowContextMenu;
 
 	/** Whether to select work on mouse double click */
 	TAttribute<bool> bSelectWordOnMouseDoubleClick;
diff --git a/Engine/Source/Runtime/Slate/Public/Widgets/Text/ISlateEditableTextWidget.h b/Engine/Source/Runtime/Slate/Public/Widgets/Text/ISlateEditableTextWidget.h
--- a/Engine/Source/Runtime/Slate/Public/Widgets/Text/ISlateEditableTextWidget.h
+++ b/Engine/Source/Runtime/Slate/Public/Widgets/Text/ISlateEditableTextWidget.h
@@ -201,6 +201,9 @@ private:
 class ISlateEditableTextWidget
 {
 public:
+	/** Is the widget construction completed (did we called and returned from the Construct() function) */
+	virtual bool IsWidgetConstructed() const = 0;
+
 	/** Is the text currently read-only? */
 	virtual bool IsTextReadOnly() const = 0;
 
diff --git a/Engine/Source/Runtime/Slate/Public/Widgets/Text/SMultiLineEditableText.h b/Engine/Source/Runtime/Slate/Public/Widgets/Text/SMultiLineEditableText.h
--- a/Engine/Source/Runtime/Slate/Public/Widgets/Text/SMultiLineEditableText.h
+++ b/Engine/Source/Runtime/Slate/Public/Widgets/Text/SMultiLineEditableText.h
@@ -38,9 +38,16 @@ enum class ETextShapingMethod : uint8;
 class ITextLayoutMarshaller;
 class FSlateEditableTextLayout;
 
+template <>
+struct TWidgetTypeTraits<class SMultiLineEditableText>
+{
+	static constexpr bool SupportsInvalidation() { return true; }
+};
+
 /** An editable text widget that supports multiple lines and soft word-wrapping. */
 class SMultiLineEditableText : public SWidget, public ISlateEditableTextWidget
 {
+	SLATE_DECLARE_WIDGET_API(SMultiLineEditableText, SWidget, SLATE_API)
 public:
 
 	/** Used to merge multiple text edit transactions within a scope */
@@ -79,7 +86,6 @@ public:
 		, _LineHeightPercentage( 1.0f )
 		, _ApplyLineHeightToBottomLine( true )
 		, _Justification( ETextJustify::Left )
-		, _IsReadOnly(false)
 		, _OnTextChanged()
 		, _OnTextCommitted()
 		, _AllowMultiLine(true)
@@ -487,6 +493,10 @@ public:
 	 */
 	SLATE_API void ForceScroll(int32 UserIndex, float ScrollAxisMagnitude);
 
+protected:
+	/** Synchronize the text style currently set (including from overrides) and update the text layout if required */
+	SLATE_API void SynchronizeTextStyle();
+
 protected:
 	//~ Begin SWidget Interface
 	SLATE_API virtual void Tick( const FGeometry& AllottedGeometry, const double InCurrentTime, const float InDeltaTime ) override;
@@ -520,6 +530,7 @@ protected:
 
 public:
 	//~ Begin ISlateEditableTextWidget Interface
+	SLATE_API virtual bool IsWidgetConstructed() const override;
 	SLATE_API virtual bool IsTextReadOnly() const override;
 	SLATE_API virtual bool IsTextPassword() const override;
 	SLATE_API virtual bool IsMultiLineTextEdit() const override;
@@ -560,6 +571,12 @@ protected:
 	/** Called to end an undoable editable text transaction, marked as protected for use with FScopedEditableTextTransaction only */
 	SLATE_API void EndEditTransaction();
 
+protected:
+	TSlateAttributeRef<FSlateFontInfo> GetFontAttribute() const
+	{
+		return TSlateAttributeRef<FSlateFontInfo>{SharedThis(this), FontAttribute};
+	}
+
 protected:
 #if WITH_EDITORONLY_DATA
 	/** The text layout that deals with the editable text */
@@ -590,8 +607,11 @@ protected:
 	/** Whether the context menu can be opened */
 	TAttribute<bool> bAllowContextMenu;
 
+#if WITH_EDITORONLY_DATA
 	/** Sets whether this text box can actually be modified interactively by the user */
-	TAttribute<bool> bIsReadOnly;
+	UE_DEPRECATED(5.8, "Use SetIsReadOnly / IsTextReadOnly")
+	TSlateDeprecatedTAttribute<bool> bIsReadOnly;
+#endif
 
 	/** Delegate to call before a context menu is opened */
 	FOnContextMenuOpening OnContextMenuOpening;
@@ -655,6 +675,11 @@ protected:
 
 	/** The message action to take when the virtual keyboard is dismissed by the user */
 	TAttribute<EVirtualKeyboardDismissAction> VirtualKeyboardDismissAction;
+
+private:
+	/** The font used to draw the text */
+	TSlateAttribute<FSlateFontInfo> FontAttribute;
+	bool bIsFontAttrSet = false;
 };
 
 #endif //WITH_FANCY_TEXT
diff --git a/Engine/Source/Runtime/Slate/Public/Widgets/Text/SlateEditableTextLayout.h b/Engine/Source/Runtime/Slate/Public/Widgets/Text/SlateEditableTextLayout.h
--- a/Engine/Source/Runtime/Slate/Public/Widgets/Text/SlateEditableTextLayout.h
+++ b/Engine/Source/Runtime/Slate/Public/Widgets/Text/SlateEditableTextLayout.h
@@ -19,6 +19,7 @@
 #include "Widgets/Text/SlateEditableTextTypes.h"
 #include "Framework/Text/SlateTextLayoutFactory.h"
 #include "GenericPlatform/ITextInputMethodSystem.h"
+#include "Types/SlateNestedAttribute.h"
 
 class FArrangedChildren;
 class FExtender;
@@ -33,9 +34,32 @@ enum class ETextShapingMethod : uint8;
 enum class EFontFacesLoadingPaintPolicy : uint8;
 
 /** Class to handle the cached layout of SEditableText/SMultiLineEditableText by proxying around a FTextLayout */
-class FSlateEditableTextLayout
+class FSlateEditableTextLayout : public FSlateNestedAttributeContainer
 {
 public:
+	template<auto InPointerToAttributeOwner>
+	static void RegisterNestedAttributes(FSlateAttributeInitializer& AttributeInitializer)
+	{
+		SLATE_ADD_NESTED_MEMBER_ATTRIBUTE_DEFINITION(AttributeInitializer, InPointerToAttributeOwner, BoundText, EInvalidateWidgetReason::Layout)
+			.OnValueChanged_Static(OnBoundTextChanged);
+		SLATE_ADD_NESTED_MEMBER_ATTRIBUTE_DEFINITION(AttributeInitializer, InPointerToAttributeOwner, HintText, EInvalidateWidgetReason::Layout);
+		SLATE_ADD_NESTED_MEMBER_ATTRIBUTE_DEFINITION(AttributeInitializer, InPointerToAttributeOwner, HintTextOpacity, EInvalidateWidgetReason::Paint);
+		SLATE_ADD_NESTED_MEMBER_ATTRIBUTE_DEFINITION(AttributeInitializer, InPointerToAttributeOwner, BoundSearchText, EInvalidateWidgetReason::Layout)
+			.OnValueChanged_Static(OnBoundSearchTextChanged);
+
+		SLATE_ADD_NESTED_MEMBER_ATTRIBUTE_DEFINITION(AttributeInitializer, InPointerToAttributeOwner, bIsReadOnlyAttribute, EInvalidateWidgetReason::PaintAndVolatility)
+			.OnValueChanged_Static(OnIsReadOnlyChanged);
+
+		SLATE_ADD_NESTED_MEMBER_ATTRIBUTE_DEFINITION(AttributeInitializer, InPointerToAttributeOwner, WrapTextAt, EInvalidateWidgetReason::Layout);
+		SLATE_ADD_NESTED_MEMBER_ATTRIBUTE_DEFINITION(AttributeInitializer, InPointerToAttributeOwner, AutoWrapText, EInvalidateWidgetReason::Layout);
+		SLATE_ADD_NESTED_MEMBER_ATTRIBUTE_DEFINITION(AttributeInitializer, InPointerToAttributeOwner, WrappingPolicy, EInvalidateWidgetReason::Layout);
+
+		SLATE_ADD_NESTED_MEMBER_ATTRIBUTE_DEFINITION(AttributeInitializer, InPointerToAttributeOwner, Margin, EInvalidateWidgetReason::Layout);
+
+		SLATE_ADD_NESTED_MEMBER_ATTRIBUTE_DEFINITION(AttributeInitializer, InPointerToAttributeOwner, Justification, EInvalidateWidgetReason::Layout);
+		SLATE_ADD_NESTED_MEMBER_ATTRIBUTE_DEFINITION(AttributeInitializer, InPointerToAttributeOwner, LineHeightPercentage, EInvalidateWidgetReason::Layout);
+		SLATE_ADD_NESTED_MEMBER_ATTRIBUTE_DEFINITION(AttributeInitializer, InPointerToAttributeOwner, ApplyLineHeightToBottomLine, EInvalidateWidgetReason::Layout);
+	}
 
 	template<typename RealWidgetType>
 	UE_DEPRECATED(5.8, "Use the other base constructor and call Construct during widget construction")
@@ -49,9 +73,21 @@ public:
 	template<typename RealWidgetType>
 	FSlateEditableTextLayout(RealWidgetType& InOwnerWidget)
 		: OwnerWidget(&InOwnerWidget)
+		, BoundText(InOwnerWidget, FText::GetEmpty())
 		, bWasPasswordLastTick(false)
+		, HintText(InOwnerWidget, FText::GetEmpty())
+		, HintTextOpacity(InOwnerWidget, 0.35f)
+		, BoundSearchText(InOwnerWidget, FText::GetEmpty())
 		, SearchCase(ESearchCase::IgnoreCase)
 		, CurrentSearchResultIndex(0)
+		, bIsReadOnlyAttribute(InOwnerWidget, false)
+		, WrapTextAt(InOwnerWidget, 0.0f)
+		, AutoWrapText(InOwnerWidget, false)
+		, WrappingPolicy(InOwnerWidget, ETextWrappingPolicy::DefaultWrapping)
+		, Margin(InOwnerWidget, FMargin(0))
+		, Justification(InOwnerWidget, ETextJustify::Left)
+		, LineHeightPercentage(InOwnerWidget, 1.0f)
+		, ApplyLineHeightToBottomLine(InOwnerWidget, false)
 		, bHasRegisteredTextInputMethodContext(false)
 		, PreferredCursorScreenOffsetInLine(0)
 		, CurrentUndoLevel(0)
@@ -63,6 +99,7 @@ public:
 		, bTextCommittedByVirtualKeyboard(false)
 		, VirtualKeyboardTextCommitType()
 		, bSelectionChangedExternally(false)
+		, bForceBoundTextReview(false)
 		, ExternalSelectionStart(0)
 		, ExternalSelectionEnd(0)
 	{
@@ -82,19 +119,22 @@ public:
 
 	SLATE_API void SetSearchText(TAttribute<FText> InSearchText);
 	SLATE_API FText GetSearchText() const;
-	
+
 	/** Get the index of the search result (0 if none) */
 	SLATE_API int32 GetSearchResultIndex() const;
-	
+
 	/** Get the total number of search results (0 if none) */
 	SLATE_API int32 GetNumSearchResults() const;
 
 	SLATE_API void SetTextStyle(const FTextBlockStyle& InTextStyle);
 	SLATE_API const FTextBlockStyle& GetTextStyle() const;
-	
+
 	SLATE_API void SetHintTextOpacity(TAttribute<float> InOpacity);
 	SLATE_API float GetHintTextOpacity() const;
 
+	SLATE_API void SetIsReadOnly(TAttribute<bool> InIsReadOnly);
+	SLATE_API bool GetIsReadOnly() const;
+
 	/** Get the number of Text Lines */
 	UE_DEPRECATED(5.8, "Use the const version instead.")
 	int32 GetTextLineCount()
@@ -481,7 +521,23 @@ public:
 
 	/** Sets the paint policy to use for text if the font faces are still async loading. */
 	SLATE_API void SetFontFacesLoadingPaintPolicy(EFontFacesLoadingPaintPolicy InFontFacesLoadingPaintPolicy);
+
 private:
+	/** Function called when the BoundText attribute changes. */
+	static void OnBoundTextChanged(SWidget& Widget, FSlateEditableTextLayout& SlateTextLayout)
+	{
+		SlateTextLayout.OnBoundTextChanged();
+	}
+
+	/** Function called when the BoundText attribute changes or called manually from SetText. Will call the Widget's OnTextChanged */
+	void OnBoundTextChanged();
+
+	/** Function called when the BoundSearchText attribute changes */
+	static void OnBoundSearchTextChanged(SWidget& Widget, FSlateEditableTextLayout& SlateTextLayout);
+
+	/** Function called when the bIsReadOnlyAttribute attribute changes */
+	static void OnIsReadOnlyChanged(SWidget& Widget, FSlateEditableTextLayout& SlateTextLayout);
+
 	/** Insert the given text at the current cursor position, correctly taking into account new line characters */
 	void InsertTextAtCursorImpl(const FString& InString);
 
@@ -612,7 +668,8 @@ private:
 	mutable FTextBlockStyle HintTextStyle;
 
 	/** The text displayed in this text block */
-	TAttribute<FText> BoundText;
+	TSlateNestedAttribute<FText> BoundText;
+	FText PreviousBoundText;
 
 	/** The state of BoundText last Tick() (used to allow updates when the text is changed) */
 	FTextSnapshot BoundTextLastTick;
@@ -621,13 +678,13 @@ private:
 	bool bWasPasswordLastTick;
 
 	/** The text that appears when there is no text in the text box */
-	TAttribute<FText> HintText;
+	TSlateNestedAttribute<FText> HintText;
 
 	/** Opacity of the hint text */
-	TAttribute<float> HintTextOpacity;
+	TSlateNestedAttribute<float> HintTextOpacity;
 
 	/** The text to be searched for */
-	TAttribute<FText> BoundSearchText;
+	TSlateNestedAttribute<FText> BoundSearchText;
 
 	/** The state of BoundSearchText last Tick() (used to allow updates when the text is changed) */
 	FTextSnapshot BoundSearchTextLastTick;
@@ -644,26 +701,29 @@ private:
 	/** The active search result index */ 
 	int32 CurrentSearchResultIndex;
 
+	/** Sets whether this text box can actually be modified interactively by the user */
+	TSlateNestedAttribute<bool> bIsReadOnlyAttribute;
+
 	/** Whether text wraps onto a new line when it's length exceeds this width; if this value is zero or negative, no wrapping occurs. */
-	TAttribute<float> WrapTextAt;
+	TSlateNestedAttribute<float> WrapTextAt;
 
 	/** True if we're wrapping text automatically based on the computed horizontal space for this widget */
-	TAttribute<bool> AutoWrapText;
+	TSlateNestedAttribute<bool> AutoWrapText;
 
 	/** The wrapping policy we're using */
-	TAttribute<ETextWrappingPolicy> WrappingPolicy;
+	TSlateNestedAttribute<ETextWrappingPolicy> WrappingPolicy;
 
 	/** The amount of blank space left around the edges of text area */
-	TAttribute<FMargin> Margin;
+	TSlateNestedAttribute<FMargin> Margin;
 
 	/** How the text should be aligned with the margin */
-	TAttribute<ETextJustify::Type> Justification;
+	TSlateNestedAttribute<ETextJustify::Type> Justification;
 
 	/** The amount to scale each lines height by */
-	TAttribute<float> LineHeightPercentage;
+	TSlateNestedAttribute<float> LineHeightPercentage;
 
 	/** Whether to leave extra space below the last line due to line height */
-	TAttribute<bool> ApplyLineHeightToBottomLine;
+	TSlateNestedAttribute<bool> ApplyLineHeightToBottomLine;
 
 	/** The information used to help identify who owns this text layout in the case of an error */
 	TAttribute<FString> DebugSourceInfo;
@@ -761,6 +821,9 @@ private:
 	/** Whether the cursor position has been changed externally*/
 	bool bSelectionChangedExternally;
 
+	/** Used in OnBoundTextChanged to assess if we need to force some text updates */
+	bool bForceBoundTextReview;
+
 	/** The boundaries of the external selection */
 	int ExternalSelectionStart;
 	int ExternalSelectionEnd;
diff --git a/Engine/Source/Runtime/Slate/Public/Widgets/Text/SlateEditableTextTypes.h b/Engine/Source/Runtime/Slate/Public/Widgets/Text/SlateEditableTextTypes.h
--- a/Engine/Source/Runtime/Slate/Public/Widgets/Text/SlateEditableTextTypes.h
+++ b/Engine/Source/Runtime/Slate/Public/Widgets/Text/SlateEditableTextTypes.h
@@ -143,6 +143,11 @@ public:
 
 	SLATE_API virtual int32 OnPaint(const FPaintArgs& Args, const FTextLayout::FLineView& Line, const FVector2D Offset, const float Width, const FTextBlockStyle& DefaultStyle, const FGeometry& AllottedGeometry, const FSlateRect& MyCullingRect, FSlateWindowElementList& OutDrawElements, int32 LayerId, const FWidgetStyle& InWidgetStyle, bool bParentEnabled) const override;
 
+	/** Set to false to skip drawing the cursor */
+	void SetDrawCursor(bool bInDrawCursor) { bDrawCursor = bInDrawCursor; }
+	/** Is the cursor drawn? */
+	bool GetDrawCursor() const { return bDrawCursor; }
+
 protected:
 	SLATE_API FCursorLineHighlighter(const FCursorInfo* InCursorInfo);
 
@@ -151,6 +156,9 @@ protected:
 
 	/** Brush used to draw the cursor */
 	TAttribute<const FSlateBrush*> CursorBrush;
+
+	/** Set to false to skip drawing the cursor */
+	bool bDrawCursor = true;
 };
 
 /** Run highlighter used to draw the composition range */
