[Fix] Fix for `FSlateEditableTextLayout` to properly keep track of the `PreviousText`.
The `OnTextChanged` event can be triggered when editing the text and not just when the `BoundText` has changed, so the `PreviousText` needs to account for these cases to properly trigger future `OnTextChanged` events.

#Synced-CL 50145007

diff --git a/Engine/Source/Runtime/Slate/Private/Widgets/Text/SlateEditableTextLayout.cpp b/Engine/Source/Runtime/Slate/Private/Widgets/Text/SlateEditableTextLayout.cpp
--- a/Engine/Source/Runtime/Slate/Private/Widgets/Text/SlateEditableTextLayout.cpp
+++ b/Engine/Source/Runtime/Slate/Private/Widgets/Text/SlateEditableTextLayout.cpp
@@ -90,7 +90,7 @@ void FSlateEditableTextLayout::Construct(TAttribute<FText> InInitialText, FTextB
 	GraphemeBreakIterator = GSlateUseSharedBreakIterator ? FBreakIterator::GetCharacterBoundaryIterator() : FBreakIterator::CreateCharacterBoundaryIterator();
 
 	BoundText.Assign(OwnerSlateWidget.Get(), MoveTemp(InInitialText), ESlateAttributeBindAction::Initialize);
-	PreviousBoundText = BoundText.Get();
+	PreviousText = BoundText.Get();
 
 	// Set the initial text - the same as SetText, but doesn't call OnTextChanged as that can cause a crash during construction
 	{
@@ -3219,7 +3219,7 @@ void FSlateEditableTextLayout::EndEditTransaction()
 		TextLayout->UpdateIfNeeded();
 
 		// Let outsiders know that the text content has been changed
-		OwnerWidget->OnTextChanged(EditedText);
+		BroadcastOnTextChanged(EditedText);
 
 		OwnerWidget->OnCursorMoved(CursorInfo.GetCursorInteractionLocation());
 
@@ -3320,7 +3320,7 @@ void FSlateEditableTextLayout::Undo()
 				if (SetEditableText(UndoState.Text))
 				{
 					// Let outsiders know that the text content has been changed
-					OwnerWidget->OnTextChanged(UndoState.Text);
+					BroadcastOnTextChanged(UndoState.Text);
 				}
 
 				CursorInfo = UndoState.CursorInfo.CreateUndo();
@@ -3361,7 +3361,7 @@ void FSlateEditableTextLayout::Redo()
 				if (SetEditableText(UndoState.Text))
 				{
 					// Let outsiders know that the text content has been changed
-					OwnerWidget->OnTextChanged(UndoState.Text);
+					BroadcastOnTextChanged(UndoState.Text);
 				}
 
 				CursorInfo.RestoreFromUndo(UndoState.CursorInfo);
@@ -3428,7 +3428,7 @@ void FSlateEditableTextLayout::UpdateTextChangedByVirtualKeyboard()
 	{
 		SetEditableText(VirtualKeyboardText);
 		// Let outsiders know that the text content has been changed
-		OwnerWidget->OnTextChanged(GetEditableText());
+		BroadcastOnTextChanged(GetEditableText());
 		bTextChangedByVirtualKeyboard = false;
 	}
 }
@@ -4053,7 +4053,7 @@ void FSlateEditableTextLayout::FTextInputMethodContext::SetTextInRange(const uin
 	{
 		OwnerLayout->SaveText(EditedText);
 		OwnerLayout->TextLayout->UpdateIfNeeded();
-		OwnerLayout->OwnerWidget->OnTextChanged(EditedText);
+		OwnerLayout->BroadcastOnTextChanged(EditedText);
 	}
 }
 
@@ -4225,6 +4225,14 @@ void FSlateEditableTextLayout::SetFontFacesLoadingPaintPolicy(EFontFacesLoadingP
 	}
 }
 
+void FSlateEditableTextLayout::BroadcastOnTextChanged(const FText& InText)
+{
+	// We update the PreviousText so future text change events will trigger as expected.
+	// This is needed because this OnTextChanged event can be broadcasted from a few places
+	PreviousText = InText;
+	OwnerWidget->OnTextChanged(InText);
+}
+
 void FSlateEditableTextLayout::OnBoundTextChanged()
 {
 	const TSharedRef<SWidget> Widget = OwnerWidget->GetSlateWidget();
@@ -4247,7 +4255,7 @@ void FSlateEditableTextLayout::OnBoundTextChanged()
 	// Only emit the "text changed" event if the text has actually been changed
 	const bool bHasTextChanged = Widget->HasAnyUserFocus().IsSet()
 		? !NewBoundText.ToString().Equals(EditableText.ToString(), ESearchCase::CaseSensitive)
-		: !NewBoundText.ToString().Equals(PreviousBoundText.ToString(), ESearchCase::CaseSensitive);
+		: !NewBoundText.ToString().Equals(PreviousText.ToString(), ESearchCase::CaseSensitive);
 
 	if (RefreshImpl(&NewBoundText, bForceRefresh))
 	{
@@ -4260,11 +4268,10 @@ void FSlateEditableTextLayout::OnBoundTextChanged()
 		// Let outsiders know that the text content has been changed
 		if (bHasTextChanged)
 		{
-			OwnerWidget->OnTextChanged(NewBoundText);
+			BroadcastOnTextChanged(NewBoundText);
 		}
 	}
 
-	PreviousBoundText = NewBoundText;
 	bForceBoundTextReview = false;
 }
 
diff --git a/Engine/Source/Runtime/Slate/Public/Widgets/Text/SlateEditableTextLayout.h b/Engine/Source/Runtime/Slate/Public/Widgets/Text/SlateEditableTextLayout.h
--- a/Engine/Source/Runtime/Slate/Public/Widgets/Text/SlateEditableTextLayout.h
+++ b/Engine/Source/Runtime/Slate/Public/Widgets/Text/SlateEditableTextLayout.h
@@ -523,6 +523,9 @@ public:
 	SLATE_API void SetFontFacesLoadingPaintPolicy(EFontFacesLoadingPaintPolicy InFontFacesLoadingPaintPolicy);
 
 private:
+	/** Function to call OnTextChanged on the owner widget and to properly set the PreviousText. */
+	void BroadcastOnTextChanged(const FText& InText);
+
 	/** Function called when the BoundText attribute changes. */
 	static void OnBoundTextChanged(SWidget& Widget, FSlateEditableTextLayout& SlateTextLayout)
 	{
@@ -669,7 +672,7 @@ private:
 
 	/** The text displayed in this text block */
 	TSlateNestedAttribute<FText> BoundText;
-	FText PreviousBoundText;
+	FText PreviousText;
 
 	/** The state of BoundText last Tick() (used to allow updates when the text is changed) */
 	FTextSnapshot BoundTextLastTick;
