[TreeView] Add recusive expansion/collapse with the arrow keys to match the shift-click expansion arrow behavior

#Synced-CL 50594661

diff --git a/Engine/Source/Runtime/Slate/Public/Widgets/Views/STreeView.h b/Engine/Source/Runtime/Slate/Public/Widgets/Views/STreeView.h
--- a/Engine/Source/Runtime/Slate/Public/Widgets/Views/STreeView.h
+++ b/Engine/Source/Runtime/Slate/Public/Widgets/Views/STreeView.h
@@ -401,7 +401,13 @@ public:
 					ItemType RangeSelectionEndItem = TListTypeTraits<ItemType>::NullableItemTypeConvertToItemType( this->SelectorItem );
 					int32 SelectionIndex = this->LinearizedItems.Find( RangeSelectionEndItem );
 
-					if ( Private_DoesItemHaveChildren(SelectionIndex) && Private_IsItemExpanded( RangeSelectionEndItem ) )
+					const bool bDoesItemHaveChildren = Private_DoesItemHaveChildren(SelectionIndex);
+					if( bDoesItemHaveChildren && InKeyEvent.GetModifierKeys().IsShiftDown() )
+					{
+						// Collapse the selected item and all of it's children
+						Private_OnExpanderArrowShiftClicked(RangeSelectionEndItem, false);
+					}
+					else if ( bDoesItemHaveChildren && Private_IsItemExpanded( RangeSelectionEndItem ) )
 					{
 						// Collapse the selected item
 						Private_SetItemExpansion(RangeSelectionEndItem, false);
@@ -434,7 +440,12 @@ public:
 					// Right only applies to items with children
 					if ( Private_DoesItemHaveChildren(SelectionIndex) )
 					{
-						if ( !Private_IsItemExpanded(RangeSelectionEndItem) )
+						if( InKeyEvent.GetModifierKeys().IsShiftDown() )
+						{
+							// Expand the selected item and all of it's children
+							Private_OnExpanderArrowShiftClicked(RangeSelectionEndItem, true);
+						}
+						else if ( !Private_IsItemExpanded(RangeSelectionEndItem) )
 						{
 							// Expand the selected item
 							Private_SetItemExpansion(RangeSelectionEndItem, true);
