Slate Dynamic Invalidation - SSpinBox and related widgets

#Synced-CL 50145007

diff --git a/Engine/Source/Developer/ToolWidgets/Public/SSearchableComboBox.h b/Engine/Source/Developer/ToolWidgets/Public/SSearchableComboBox.h
--- a/Engine/Source/Developer/ToolWidgets/Public/SSearchableComboBox.h
+++ b/Engine/Source/Developer/ToolWidgets/Public/SSearchableComboBox.h
@@ -1,4 +1,4 @@
-// Copyright Epic Games, Inc. All Rights Reservekd.
+// Copyright Epic Games, Inc. All Rights Reserved.
 
 #pragma once
 
@@ -13,6 +13,13 @@
 #include "Widgets/Input/SComboBox.h"
 
 class SEditableTextBox;
+
+template <>
+struct TWidgetTypeTraits<class SSearchableComboBox>
+{
+	static constexpr bool SupportsInvalidation() { return true; }
+};
+
 /**
  * A searchable text combo box
  */
diff --git a/Engine/Source/Editor/GraphEditor/Public/KismetPins/SVector2DTextBox.h b/Engine/Source/Editor/GraphEditor/Public/KismetPins/SVector2DTextBox.h
--- a/Engine/Source/Editor/GraphEditor/Public/KismetPins/SVector2DTextBox.h
+++ b/Engine/Source/Editor/GraphEditor/Public/KismetPins/SVector2DTextBox.h
@@ -9,6 +9,16 @@
 
 #define LOCTEXT_NAMESPACE "Vector2DTextBox"
 
+
+template <typename NumericType>
+class SVector2DTextBox;
+
+template <typename NumericType>
+struct TWidgetTypeTraits<SVector2DTextBox<NumericType>>
+{
+	static constexpr bool SupportsInvalidation() { return true; }
+};
+
 // Class implementation to create 2 editable text boxes to represent vector2D graph pin
 template <typename NumericType>
 class SVector2DTextBox : public SCompoundWidget
diff --git a/Engine/Source/Editor/GraphEditor/Public/KismetPins/SVector4TextBox.h b/Engine/Source/Editor/GraphEditor/Public/KismetPins/SVector4TextBox.h
--- a/Engine/Source/Editor/GraphEditor/Public/KismetPins/SVector4TextBox.h
+++ b/Engine/Source/Editor/GraphEditor/Public/KismetPins/SVector4TextBox.h
@@ -9,7 +9,18 @@
 
 #define LOCTEXT_NAMESPACE "VectorTextBox"
 
-// Class implementation to create 3 editable text boxes to represent vector/rotator graph pin
+
+template <typename NumericType>
+class SVector4TextBox;
+
+template <typename NumericType>
+struct TWidgetTypeTraits<SVector4TextBox<NumericType>>
+{
+	static constexpr bool SupportsInvalidation() { return true; }
+};
+
+
+// Class implementation to create 4 editable text boxes to represent vector/rotator graph pin
 template <typename NumericType>
 class SVector4TextBox : public SCompoundWidget
 {
diff --git a/Engine/Source/Editor/GraphEditor/Public/KismetPins/SVectorTextBox.h b/Engine/Source/Editor/GraphEditor/Public/KismetPins/SVectorTextBox.h
--- a/Engine/Source/Editor/GraphEditor/Public/KismetPins/SVectorTextBox.h
+++ b/Engine/Source/Editor/GraphEditor/Public/KismetPins/SVectorTextBox.h
@@ -9,6 +9,15 @@
 
 #define LOCTEXT_NAMESPACE "VectorTextBox"
 
+template <typename NumericType>
+class SVectorTextBox;
+
+template <typename NumericType>
+struct TWidgetTypeTraits<SVectorTextBox<NumericType>>
+{
+	static constexpr bool SupportsInvalidation() { return true; }
+};
+
 // Class implementation to create 3 editable text boxes to represent vector/rotator graph pin
 template <typename NumericType>
 class SVectorTextBox : public SCompoundWidget
diff --git a/Engine/Source/Runtime/Slate/Private/Widgets/Input/SSpinBox.cpp b/Engine/Source/Runtime/Slate/Private/Widgets/Input/SSpinBox.cpp
--- a/Engine/Source/Runtime/Slate/Private/Widgets/Input/SSpinBox.cpp
+++ b/Engine/Source/Runtime/Slate/Private/Widgets/Input/SSpinBox.cpp
@@ -30,16 +30,42 @@ float SpinBoxComputeExponentSliderFraction(float FractionFilled, float StartFrac
 }
 
 
+template<typename NumericType>
+inline void SSpinBox<NumericType>::PrivateRegisterAttributes(FSlateAttributeInitializer& AttributeInitializer)
+{
+	SLATE_ADD_MEMBER_ATTRIBUTE_DEFINITION(AttributeInitializer, bDragging, EInvalidateWidgetReason::Paint);
+	AttributeInitializer.OverrideInvalidationReason("EnabledState", FSlateAttributeDescriptor::FInvalidateWidgetReasonAttribute{EInvalidateWidgetReason::Paint});
+	AttributeInitializer.OverrideInvalidationReason("Hovered", FSlateAttributeDescriptor::FInvalidateWidgetReasonAttribute{EInvalidateWidgetReason::Paint});
+
+	SLATE_ADD_MEMBER_ATTRIBUTE_DEFINITION(AttributeInitializer, ValueAttribute, EInvalidateWidgetReason::Paint);
+	SLATE_ADD_MEMBER_ATTRIBUTE_DEFINITION(AttributeInitializer, Delta, EInvalidateWidgetReason::Paint);
+	SLATE_ADD_MEMBER_ATTRIBUTE_DEFINITION(AttributeInitializer, MinValue, EInvalidateWidgetReason::Paint);
+	SLATE_ADD_MEMBER_ATTRIBUTE_DEFINITION(AttributeInitializer, MaxValue, EInvalidateWidgetReason::Paint);
+	SLATE_ADD_MEMBER_ATTRIBUTE_DEFINITION(AttributeInitializer, MinSliderValue, EInvalidateWidgetReason::Paint);
+	SLATE_ADD_MEMBER_ATTRIBUTE_DEFINITION(AttributeInitializer, MaxSliderValue, EInvalidateWidgetReason::Paint);
+
+	SLATE_ADD_MEMBER_ATTRIBUTE_DEFINITION(AttributeInitializer, SliderExponent, EInvalidateWidgetReason::Paint);
+	SLATE_ADD_MEMBER_ATTRIBUTE_DEFINITION(AttributeInitializer, SliderExponentNeutralValue, EInvalidateWidgetReason::Paint);
+}
 
 template <typename NumericType>
 SSpinBox<NumericType>::SSpinBox()
+	: ValueAttribute(*this, 0)
+	, Delta(*this, 0)
+	, SliderExponent(*this, 1.f)
+	, SliderExponentNeutralValue(*this)
+	, MinValue(*this, TOptional<NumericType>{})
+	, MaxValue(*this, TOptional<NumericType>{})
+	, MinSliderValue(*this, TOptional<NumericType>{})
+	, MaxSliderValue(*this, TOptional<NumericType>{})
+	, bDragging(*this, false)
 {
 }
 
 template <typename NumericType>
 SSpinBox<NumericType>::~SSpinBox()
 {
-	if (bDragging || PointerDraggingSliderIndex != INDEX_NONE)
+	if (bDragging.Get() || PointerDraggingSliderIndex != INDEX_NONE)
 	{
 		CancelMouseCapture();
 	}
@@ -66,17 +92,17 @@ void SSpinBox<NumericType>::Construct(const FArguments& InArgs)
 		Interface->GetOnSettingChanged()->AddSP(this, &SSpinBox::ResetCachedValueString);
 	}
 
-	ValueAttribute = InArgs._Value;
+	ValueAttribute.Assign(*this, InArgs._Value, ESlateAttributeBindAction::Initialize);
 	OnValueChanged = InArgs._OnValueChanged;
 	OnValueCommitted = InArgs._OnValueCommitted;
 	OnBeginSliderMovement = InArgs._OnBeginSliderMovement;
 	OnEndSliderMovement = InArgs._OnEndSliderMovement;
 	MinDesiredWidth = InArgs._MinDesiredWidth;
 
-	MinValue = InArgs._MinValue;
-	MaxValue = InArgs._MaxValue;
-	MinSliderValue = (InArgs._MinSliderValue.Get().IsSet()) ? InArgs._MinSliderValue : MinValue;
-	MaxSliderValue = (InArgs._MaxSliderValue.Get().IsSet()) ? InArgs._MaxSliderValue : MaxValue;
+	MinValue.Assign(*this, InArgs._MinValue, ESlateAttributeBindAction::Initialize);
+	MaxValue.Assign(*this, InArgs._MaxValue, ESlateAttributeBindAction::Initialize);
+	MinSliderValue.Assign(*this, (InArgs._MinSliderValue.Get().IsSet()) ? InArgs._MinSliderValue : MinValue.Get(), ESlateAttributeBindAction::Initialize);
+	MaxSliderValue.Assign(*this, (InArgs._MaxSliderValue.Get().IsSet()) ? InArgs._MaxSliderValue : MaxValue.Get(), ESlateAttributeBindAction::Initialize);
 
 	MinFractionalDigits = (InArgs._MinFractionalDigits.Get().IsSet()) ? InArgs._MinFractionalDigits : DefaultMinFractionalDigits;
 	MaxFractionalDigits = (InArgs._MaxFractionalDigits.Get().IsSet()) ? InArgs._MaxFractionalDigits : DefaultMaxFractionalDigits;
@@ -117,14 +143,15 @@ void SSpinBox<NumericType>::Construct(const FArguments& InArgs)
 
 	UpdateIsSpinRangeUnlimited();
 
-	SliderExponent = InArgs._SliderExponent;
+	SliderExponent.Assign(*this, InArgs._SliderExponent);
 
-	SliderExponentNeutralValue = InArgs._SliderExponentNeutralValue;
+	bIsSliderExponentNeutralValueSet = InArgs._SliderExponentNeutralValue.IsSet();
+	SliderExponentNeutralValue.Assign(*this, InArgs._SliderExponentNeutralValue);
 
 	DistanceDragged = 0.0f;
 	PreDragValue = NumericType();
 
-	Delta = InArgs._Delta;
+	Delta.Assign(*this, InArgs._Delta);
 	ShiftMultiplier = InArgs._ShiftMultiplier;
 	CtrlMultiplier = InArgs._CtrlMultiplier;
 	ShiftAltMultiplier = InArgs._ShiftAltMultiplier;
@@ -141,7 +168,6 @@ void SSpinBox<NumericType>::Construct(const FArguments& InArgs)
 
 	const FMargin& TextMargin = InArgs._Style->TextPadding;
 
-	bDragging = false;
 	PointerDraggingSliderIndex = INDEX_NONE;
 
 	bIsTextChanging = false;
@@ -203,7 +229,7 @@ int32 SSpinBox<NumericType>::OnPaint(const FPaintArgs& Args, const FGeometry& Al
 	FSlateWindowElementList& OutDrawElements, int32 LayerId, const FWidgetStyle& InWidgetStyle, bool bParentEnabled) const
 {
 
-	const bool bActiveFeedback = bDragging || IsInTextMode();
+	const bool bActiveFeedback = bDragging.Get() || IsInTextMode();
 
 	const FSlateBrush* BackgroundImage = bActiveFeedback ? BackgroundActiveBrush : IsHovered() ?
 		                                     BackgroundHoveredBrush :
@@ -243,7 +269,7 @@ int32 SSpinBox<NumericType>::OnPaint(const FPaintArgs& Args, const FGeometry& Al
 		const float CachedSliderExponent = SliderExponent.Get();
 		if (!FMath::IsNearlyEqual(CachedSliderExponent, 1.f))
 		{
-			if (SliderExponentNeutralValue.IsSet() && SliderExponentNeutralValue.Get() > GetMinSliderValue() && SliderExponentNeutralValue.Get() < GetMaxSliderValue())
+			if (bIsSliderExponentNeutralValueSet && SliderExponentNeutralValue.Get() > GetMinSliderValue() && SliderExponentNeutralValue.Get() < GetMaxSliderValue())
 			{
 				//Compute a log curve on both side of the neutral value
 				float StartFractionFilled = Fraction((double)SliderExponentNeutralValue.Get(), (double)GetMinSliderValue(), (double)GetMaxSliderValue());
@@ -322,13 +348,13 @@ FReply SSpinBox<NumericType>::OnMouseButtonUp(const FGeometry& MyGeometry, const
 		if (!this->HasMouseCapture())
 		{
 			// Lost Capture - ensure reset
-			bDragging = false;
+			bDragging.Assign(*this, false);
 			PointerDraggingSliderIndex = INDEX_NONE;
 
 			return FReply::Unhandled();
 		}
 
-		if (bDragging)
+		if (bDragging.Get())
 		{
 			NumericType CurrentDelta = Delta.Get();
 			if (CurrentDelta != NumericType() && !CommitWithMultiplier(MouseEvent))
@@ -340,7 +366,7 @@ FReply SSpinBox<NumericType>::OnMouseButtonUp(const FGeometry& MyGeometry, const
 			NotifyValueCommitted(CurrentValue);
 		}
 
-		bDragging = false;
+		bDragging.Assign(*this, false);
 		PointerDraggingSliderIndex = INDEX_NONE;
 
 		FReply Reply = FReply::Handled().ReleaseMouseCapture();
@@ -369,7 +395,7 @@ void SSpinBox<NumericType>::ApplySliderMaxValueChanged(float SliderDeltaToAdd, b
 
 	NumericType NewMaxSliderValue = std::numeric_limits<NumericType>::min();
 
-	if (MaxSliderValue.IsSet() && MaxSliderValue.Get().IsSet())
+	if (MaxSliderValue.Get().IsSet())
 	{
 		NewMaxSliderValue = GetMaxSliderValue();
 
@@ -377,7 +403,7 @@ void SSpinBox<NumericType>::ApplySliderMaxValueChanged(float SliderDeltaToAdd, b
 		{
 			NewMaxSliderValue += (NumericType)SliderDeltaToAdd;
 
-			if (!MaxSliderValue.IsBound()) // simple value so we can update it without breaking the mechanic otherwise it must be handle by the callback implementer
+			if (!MaxSliderValue.IsBound(*this)) // simple value so we can update it without breaking the mechanic otherwise it must be handle by the callback implementer
 			{
 				SetMaxSliderValue(NewMaxSliderValue);
 			}
@@ -393,11 +419,11 @@ void SSpinBox<NumericType>::ApplySliderMaxValueChanged(float SliderDeltaToAdd, b
 template <typename NumericType>
 void SSpinBox<NumericType>::ApplySliderMinValueChanged(float SliderDeltaToAdd, bool UpdateOnlyIfLower)
 {
-	check(SupportDynamicSliderMaxValue.Get());
+	check(SupportDynamicSliderMinValue.Get());
 
 	NumericType NewMinSliderValue = std::numeric_limits<NumericType>::min();
 
-	if (MinSliderValue.IsSet() && MinSliderValue.Get().IsSet())
+	if (MinSliderValue.Get().IsSet())
 	{
 		NewMinSliderValue = GetMinSliderValue();
 
@@ -405,7 +431,7 @@ void SSpinBox<NumericType>::ApplySliderMinValueChanged(float SliderDeltaToAdd, b
 		{
 			NewMinSliderValue += (NumericType)SliderDeltaToAdd;
 
-			if (!MinSliderValue.IsBound()) // simple value so we can update it without breaking the mechanic otherwise it must be handle by the callback implementer
+			if (!MinSliderValue.IsBound(*this)) // simple value so we can update it without breaking the mechanic otherwise it must be handle by the callback implementer
 			{
 				SetMinSliderValue(NewMinSliderValue);
 			}
@@ -427,19 +453,19 @@ FReply SSpinBox<NumericType>::OnMouseMove(const FGeometry& MyGeometry, const FPo
 		if (!this->HasMouseCapture())
 		{
 			// Lost the mouse capture - ensure reset
-			bDragging = false;
+			bDragging.Assign(*this, false);
 			PointerDraggingSliderIndex = INDEX_NONE;
 
 			return FReply::Unhandled();
 		}
 
-		if (!bDragging)
+		if (!bDragging.Get())
 		{
 			DistanceDragged += (float)FMath::Abs(MouseEvent.GetCursorDelta().X);
 			if (DistanceDragged > FSlateApplication::Get().GetDragTriggerDistance())
 			{
 				ExitTextMode();
-				bDragging = true;
+				bDragging.Assign(*this, true);
 				OnBeginSliderMovement.ExecuteIfBound();
 			}
 
@@ -476,7 +502,7 @@ FReply SSpinBox<NumericType>::OnMouseMove(const FGeometry& MyGeometry, const FPo
 			//if we have a range to draw in
 			if (!bUnlimitedSpinRange)
 			{
-				bool HasValidExponentNeutralValue = SliderExponentNeutralValue.IsSet() && SliderExponentNeutralValue.Get() > GetMinSliderValue() && SliderExponentNeutralValue.Get() < GetMaxSliderValue();
+				bool HasValidExponentNeutralValue = bIsSliderExponentNeutralValueSet && SliderExponentNeutralValue.Get() > GetMinSliderValue() && SliderExponentNeutralValue.Get() < GetMaxSliderValue();
 
 				const float CachedSliderExponent = SliderExponent.Get();
 				// The amount currently filled in the spinbox, needs to be calculated to do deltas correctly.
@@ -529,7 +555,7 @@ FReply SSpinBox<NumericType>::OnMouseMove(const FGeometry& MyGeometry, const FPo
 				// If this control has a specified delta and sensitivity then we use that instead of the current value for determining how much to change.
 				const double Sign = (MouseEvent.GetCursorDelta().X > 0) ? 1.0 : -1.0;
 
-				if (LinearDeltaSensitivity.IsSet() && LinearDeltaSensitivity.Get() != 0 && Delta.IsSet() && Delta.Get() > 0)
+				if (LinearDeltaSensitivity.IsSet() && LinearDeltaSensitivity.Get() != 0 && Delta.Get() > 0)
 				{
 					const double MouseDelta = FMath::Abs(MouseEvent.GetCursorDelta().X / (float)LinearDeltaSensitivity.Get());
 					NewValue = InternalValue + (Sign * MouseDelta * FMath::Pow((double)Delta.Get(), (double)SliderExponent.Get())) * Step;
@@ -612,7 +638,7 @@ FCursorReply SSpinBox<NumericType>::OnCursorQuery(const FGeometry& MyGeometry, c
 		return FCursorReply::Cursor(EMouseCursor::Default);
 	}
 
-	return bDragging ?
+	return bDragging.Get() ?
 		       FCursorReply::Cursor(EMouseCursor::None) :
 		       FCursorReply::Cursor(EMouseCursor::ResizeLeftRight);
 }
@@ -627,7 +653,7 @@ bool SSpinBox<NumericType>::SupportsKeyboardFocus() const
 template <typename NumericType>
 FReply SSpinBox<NumericType>::OnFocusReceived(const FGeometry& MyGeometry, const FFocusEvent& InFocusEvent)
 {
-	if (!bDragging && (InFocusEvent.GetCause() == EFocusCause::Navigation || InFocusEvent.GetCause() == EFocusCause::SetDirectly))
+	if (!bDragging.Get() && (InFocusEvent.GetCause() == EFocusCause::Navigation || InFocusEvent.GetCause() == EFocusCause::SetDirectly))
 	{
 		EnterTextMode();
 		return FReply::Handled().SetUserFocus(EditableText.ToSharedRef(), InFocusEvent.GetCause());
@@ -686,16 +712,16 @@ bool SSpinBox<NumericType>::HasKeyboardFocus() const
 
 template <typename NumericType>
 TAttribute<NumericType> SSpinBox<NumericType>::GetValueAttribute() const
-{ return ValueAttribute; }
+{ return ValueAttribute.ToAttribute(*this); }
 
 template <typename NumericType>
 NumericType SSpinBox<NumericType>::GetValue() const
 { return ValueAttribute.Get(); }
 
 template <typename NumericType>
-void SSpinBox<NumericType>::SetValue(const TAttribute<NumericType>& InValueAttribute, const bool bShouldCommit)
+void SSpinBox<NumericType>::SetValue(TAttribute<NumericType> InValueAttribute, const bool bShouldCommit)
 {
-	ValueAttribute = InValueAttribute;
+	ValueAttribute.Assign(*this, MoveTemp(InValueAttribute), ESlateAttributeBindAction::Initialize);
 	const NumericType LocalValueAttribute = ValueAttribute.Get();
 	CommitValue(LocalValueAttribute, (double)LocalValueAttribute, ECommitMethod::CommittedViaCode, ETextCommit::Default, bShouldCommit);
 }
@@ -705,9 +731,9 @@ NumericType SSpinBox<NumericType>::GetMinValue() const
 { return MinValue.Get().Get(std::numeric_limits<NumericType>::lowest()); }
 
 template <typename NumericType>
-void SSpinBox<NumericType>::SetMinValue(const TAttribute<TOptional<NumericType>>& InMinValue)
+void SSpinBox<NumericType>::SetMinValue(TAttribute<TOptional<NumericType>> InMinValue)
 {
-	MinValue = InMinValue;
+	MinValue.Assign(*this, MoveTemp(InMinValue), ESlateAttributeBindAction::Initialize);
 	UpdateIsSpinRangeUnlimited();
 }
 
@@ -716,39 +742,39 @@ NumericType SSpinBox<NumericType>::GetMaxValue() const
 { return MaxValue.Get().Get(std::numeric_limits<NumericType>::max()); }
 
 template <typename NumericType>
-void SSpinBox<NumericType>::SetMaxValue(const TAttribute<TOptional<NumericType>>& InMaxValue)
+void SSpinBox<NumericType>::SetMaxValue(TAttribute<TOptional<NumericType>> InMaxValue)
 {
-	MaxValue = InMaxValue;
+	MaxValue.Assign(*this, MoveTemp(InMaxValue), ESlateAttributeBindAction::Initialize);
 	UpdateIsSpinRangeUnlimited();
 }
 
 template <typename NumericType>
 bool SSpinBox<NumericType>::IsMinSliderValueBound() const
-{ return MinSliderValue.IsBound(); }
+{ return MinSliderValue.IsBound(*this); }
 
 template <typename NumericType>
 NumericType SSpinBox<NumericType>::GetMinSliderValue() const
 { return MinSliderValue.Get().Get(std::numeric_limits<NumericType>::lowest()); }
 
 template <typename NumericType>
-void SSpinBox<NumericType>::SetMinSliderValue(const TAttribute<TOptional<NumericType>>& InMinSliderValue)
+void SSpinBox<NumericType>::SetMinSliderValue(TAttribute<TOptional<NumericType>> InMinSliderValue)
 {
-	MinSliderValue = (InMinSliderValue.Get().IsSet()) ? InMinSliderValue : MinValue;
+	MinSliderValue.Assign(*this, InMinSliderValue.Get().IsSet() ? MoveTemp(InMinSliderValue) : MinValue.Get(), ESlateAttributeBindAction::Initialize);
 	UpdateIsSpinRangeUnlimited();
 }
 
 template <typename NumericType>
 bool SSpinBox<NumericType>::IsMaxSliderValueBound() const
-{ return MaxSliderValue.IsBound(); }
+{ return MaxSliderValue.IsBound(*this); }
 
 template <typename NumericType>
 NumericType SSpinBox<NumericType>::GetMaxSliderValue() const
 { return MaxSliderValue.Get().Get(std::numeric_limits<NumericType>::max()); }
 
 template <typename NumericType>
-void SSpinBox<NumericType>::SetMaxSliderValue(const TAttribute<TOptional<NumericType>>& InMaxSliderValue)
+void SSpinBox<NumericType>::SetMaxSliderValue(TAttribute<TOptional<NumericType>> InMaxSliderValue)
 {
-	MaxSliderValue = (InMaxSliderValue.Get().IsSet()) ? InMaxSliderValue : MaxValue;;
+	MaxSliderValue.Assign(*this, InMaxSliderValue.Get().IsSet() ? MoveTemp(InMaxSliderValue) : MaxValue.Get(), ESlateAttributeBindAction::Initialize);
 	UpdateIsSpinRangeUnlimited();
 }
 
@@ -796,23 +822,23 @@ NumericType SSpinBox<NumericType>::GetDelta() const
 
 template <typename NumericType>
 void SSpinBox<NumericType>::SetDelta(NumericType InDelta)
-{ Delta.Set(InDelta); }
+{ Delta.Assign(*this, InDelta); }
 
 template <typename NumericType>
 float SSpinBox<NumericType>::GetSliderExponent() const
 { return SliderExponent.Get(); }
 
 template <typename NumericType>
-void SSpinBox<NumericType>::SetSliderExponent(const TAttribute<float>& InSliderExponent)
-{ SliderExponent = InSliderExponent; }
+void SSpinBox<NumericType>::SetSliderExponent(TAttribute<float> InSliderExponent)
+{ SliderExponent.Assign(*this, MoveTemp(InSliderExponent)); }
 
 template <typename NumericType>
 float SSpinBox<NumericType>::GetMinDesiredWidth() const
 { return MinDesiredWidth.Get(); }
 
 template <typename NumericType>
-void SSpinBox<NumericType>::SetMinDesiredWidth(const TAttribute<float>& InMinDesiredWidth)
-{ MinDesiredWidth = InMinDesiredWidth; }
+void SSpinBox<NumericType>::SetMinDesiredWidth(TAttribute<float> InMinDesiredWidth)
+{ MinDesiredWidth = MoveTemp(InMinDesiredWidth); }
 
 template <typename NumericType>
 const FSpinBoxStyle* SSpinBox<NumericType>::GetWidgetStyle() const
@@ -862,6 +888,7 @@ void SSpinBox<NumericType>::EnterTextMode()
 			EditableText->SetEditableText(OverrideEditableText);
 		}
 	}
+	Invalidate(EInvalidateWidgetReason::Paint);
 }
 
 template <typename NumericType>
@@ -869,6 +896,7 @@ void SSpinBox<NumericType>::ExitTextMode()
 {
 	TextBlock->SetVisibility(EVisibility::Visible);
 	EditableText->SetVisibility(EVisibility::Collapsed);
+	Invalidate(EInvalidateWidgetReason::Paint);
 }
 
 template <typename NumericType>
@@ -979,11 +1007,13 @@ void SSpinBox<NumericType>::CommitValue(NumericType NewValue, double NewSpinValu
 		NewSpinValue = FMath::Clamp<double>(NewSpinValue, (double)LocalMinValue, (double)LocalMaxValue);
 	}
 
-	if (!ValueAttribute.IsBound())
+	if (!ValueAttribute.IsBound(*this))
 	{
-		ValueAttribute.Set(NewValue);
+		ValueAttribute.Assign(*this, NewValue);
 	}
 
+	ValueAttribute.UpdateNow(*this);
+
 	// If not in spin mode, there is no need to jump to the value from the external source, continue to use the committed value.
 	if (CommitMethod == CommittedViaSpin)
 	{
@@ -1027,10 +1057,11 @@ void SSpinBox<NumericType>::CommitValue(NumericType NewValue, double NewSpinValu
 
 	OnValueChanged.ExecuteIfBound(NewValue);
 
-	if (!ValueAttribute.IsBound())
+	if (!ValueAttribute.IsBound(*this))
 	{
-		ValueAttribute.Set(NewValue);
+		ValueAttribute.Assign(*this, NewValue);
 	}
+	ValueAttribute.UpdateNow(*this);
 
 	// Update the cache of the external value to what the user believes the value is now.
 	const NumericType CurrentValue = ValueAttribute.Get();
@@ -1046,7 +1077,7 @@ void SSpinBox<NumericType>::CommitValue(NumericType NewValue, double NewSpinValu
 	// This ensures that dragging is cleared if focus has been removed from this widget in one of the delegate calls, such as when spawning a modal dialog.
 	if (!this->HasMouseCapture())
 	{
-		bDragging = false;
+		bDragging.Assign(*this, false);
 		PointerDraggingSliderIndex = INDEX_NONE;
 	}
 }
@@ -1147,7 +1178,7 @@ NumericType SSpinBox<NumericType>::RoundIfIntegerValue(double ValueToRound) cons
 template <typename NumericType>
 void SSpinBox<NumericType>::CancelMouseCapture()
 {
-	bDragging = false;
+	bDragging.Assign(*this, false);
 	PointerDraggingSliderIndex = INDEX_NONE;
 
 	InternalValue = (double)PreDragValue;
diff --git a/Engine/Source/Runtime/Slate/Public/Widgets/Input/SNumericEntryBox.h b/Engine/Source/Runtime/Slate/Public/Widgets/Input/SNumericEntryBox.h
--- a/Engine/Source/Runtime/Slate/Public/Widgets/Input/SNumericEntryBox.h
+++ b/Engine/Source/Runtime/Slate/Public/Widgets/Input/SNumericEntryBox.h
@@ -27,6 +27,16 @@
 #include "Widgets/Colors/SColorBlock.h"
 #include "Containers/UnrealString.h"
 
+
+template <typename NumericType>
+class SNumericEntryBox;
+
+template <typename NumericType>
+struct TWidgetTypeTraits<SNumericEntryBox<NumericType>>
+{
+	static constexpr bool SupportsInvalidation() { return true; }
+};
+
 /**
  * Implementation for a box that only accepts a numeric value or that can display an undetermined value via a string
  * Supports an optional spin box for manipulating a value by dragging with the mouse
diff --git a/Engine/Source/Runtime/Slate/Public/Widgets/Input/SSpinBox.h b/Engine/Source/Runtime/Slate/Public/Widgets/Input/SSpinBox.h
--- a/Engine/Source/Runtime/Slate/Public/Widgets/Input/SSpinBox.h
+++ b/Engine/Source/Runtime/Slate/Public/Widgets/Input/SSpinBox.h
@@ -52,6 +52,15 @@ namespace SpinBoxPrivate
 */
 SLATE_API float SpinBoxComputeExponentSliderFraction(float FractionFilled, float StartFractionFilled, float SliderExponent);
 
+template <typename NumericType>
+class SSpinBox;
+
+template <typename NumericType>
+struct TWidgetTypeTraits<SSpinBox<NumericType>>
+{
+	static constexpr bool SupportsInvalidation() { return true; }
+};
+
 /**
  * A Slate SpinBox resembles traditional spin boxes in that it is a widget that provides
  * keyboard-based and mouse-based manipulation of a numeric value.
@@ -62,6 +71,8 @@ template<typename NumericType>
 class SSpinBox
 	: public SCompoundWidget
 {
+	SLATE_DECLARE_WIDGET_API(SSpinBox, SCompoundWidget, SLATE_API)
+	static inline SLATE_IMPLEMENT_WIDGET(SSpinBox)
 public:
 
 	/** Notification for numeric value change */
@@ -269,31 +280,31 @@ public:
 	/** See the Value attribute */
 	SLATE_API NumericType GetValue() const;
 
-	SLATE_API void SetValue(const TAttribute<NumericType>& InValueAttribute, const bool bShouldCommit = false);
+	SLATE_API void SetValue(TAttribute<NumericType> InValueAttribute, const bool bShouldCommit = false);
 
 	/** See the MinValue attribute */
 	SLATE_API NumericType GetMinValue() const;
 
-	SLATE_API void SetMinValue(const TAttribute<TOptional<NumericType>>& InMinValue);
+	SLATE_API void SetMinValue(TAttribute<TOptional<NumericType>> InMinValue);
 
 	/** See the MaxValue attribute */
 	SLATE_API NumericType GetMaxValue() const;
 
-	SLATE_API void SetMaxValue(const TAttribute<TOptional<NumericType>>& InMaxValue);
+	SLATE_API void SetMaxValue(TAttribute<TOptional<NumericType>> InMaxValue);
 
 	/** See the MinSliderValue attribute */
 	SLATE_API bool IsMinSliderValueBound() const;
 
 	SLATE_API NumericType GetMinSliderValue() const;
 
-	SLATE_API void SetMinSliderValue(const TAttribute<TOptional<NumericType>>& InMinSliderValue);
+	SLATE_API void SetMinSliderValue(TAttribute<TOptional<NumericType>> InMinSliderValue);
 
 	/** See the MaxSliderValue attribute */
 	SLATE_API bool IsMaxSliderValueBound() const;
 
 	SLATE_API NumericType GetMaxSliderValue() const;
 
-	SLATE_API void SetMaxSliderValue(const TAttribute<TOptional<NumericType>>& InMaxSliderValue);
+	SLATE_API void SetMaxSliderValue(TAttribute<TOptional<NumericType>> InMaxSliderValue);
 
 	/** See the MinFractionalDigits attribute */
 	SLATE_API int32 GetMinFractionalDigits() const;
@@ -319,11 +330,11 @@ public:
 
 	/** See the SliderExponent attribute */
 	SLATE_API float GetSliderExponent() const;
-	SLATE_API void SetSliderExponent(const TAttribute<float>& InSliderExponent);
+	SLATE_API void SetSliderExponent(TAttribute<float> InSliderExponent);
 
 	/** See the MinDesiredWidth attribute */
 	SLATE_API float GetMinDesiredWidth() const;
-	SLATE_API void SetMinDesiredWidth(const TAttribute<float>& InMinDesiredWidth);
+	SLATE_API void SetMinDesiredWidth(TAttribute<float> InMinDesiredWidth);
 
 	SLATE_API const FSpinBoxStyle* GetWidgetStyle() const;
 	SLATE_API void SetWidgetStyle(const FSpinBoxStyle* InStyle);
@@ -417,7 +428,7 @@ private:
 	/** The default maximum fractional digits */
 	static constexpr int32 DefaultMaxFractionalDigits = 6;
 
-	TAttribute<NumericType> ValueAttribute;
+	TSlateAttribute<NumericType> ValueAttribute;
 	FOnValueChanged OnValueChanged;
 	FOnValueCommitted OnValueCommitted;
 	FSimpleDelegate OnBeginSliderMovement;
@@ -428,10 +439,6 @@ private:
 	/** Interface that defines conversion functionality for the templated type */
 	TAttribute<TSharedPtr< INumericTypeInterface<NumericType> >> InterfaceAttr;
 
-	/** True when no range is specified, spinner can be spun indefinitely */
-	bool bUnlimitedSpinRange;
-	void UpdateIsSpinRangeUnlimited();
-
 	const FSpinBoxStyle* Style;
 
 	const FSlateBrush* BackgroundHoveredBrush;
@@ -442,7 +449,7 @@ private:
 	const FSlateBrush* InactiveFillBrush;
 
 	float DistanceDragged;
-	TAttribute<NumericType> Delta;
+	TSlateAttribute<NumericType> Delta;
 
 	// Step multipliers
 	TAttribute<float> ShiftMultiplier;
@@ -451,12 +458,12 @@ private:
 	TAttribute<float> CtrlAltMultiplier;
 
 	TAttribute<int32> LinearDeltaSensitivity;
-	TAttribute<float> SliderExponent;
-	TAttribute<NumericType> SliderExponentNeutralValue;
-	TAttribute< TOptional<NumericType> > MinValue;
-	TAttribute< TOptional<NumericType> > MaxValue;
-	TAttribute< TOptional<NumericType> > MinSliderValue;
-	TAttribute< TOptional<NumericType> > MaxSliderValue;
+	TSlateAttribute<float> SliderExponent;
+	TSlateAttribute<NumericType> SliderExponentNeutralValue;
+	TSlateAttribute< TOptional<NumericType> > MinValue;
+	TSlateAttribute< TOptional<NumericType> > MaxValue;
+	TSlateAttribute< TOptional<NumericType> > MinSliderValue;
+	TSlateAttribute< TOptional<NumericType> > MaxSliderValue;
 	TAttribute< TOptional<int32> > MinFractionalDigits;
 	TAttribute< TOptional<int32> > MaxFractionalDigits;
 	TAttribute<bool> AlwaysUsesDeltaSnap;
@@ -510,12 +517,19 @@ private:
 	/** Used to prevent per-frame re-conversion of the cached numeric value to a string. */
 	FString CachedValueString;
 
+	/** Whether the user is dragging the slider */
+	TSlateAttribute<bool> bDragging;
+
+	/** True when the SliderExponentNeutralValue attribute is set */
+	bool bIsSliderExponentNeutralValueSet;
+
+	/** True when no range is specified, spinner can be spun indefinitely */
+	bool bUnlimitedSpinRange;
+	void UpdateIsSpinRangeUnlimited();
+
 	/** Whetever the interfaced setting changed and the CachedValueString needs to be recomputed. */
 	mutable bool bCachedValueStringDirty;
 
-	/** Whether the user is dragging the slider */
-	bool bDragging;
-
 	/** Re-entrant guard for the text changed handler */
 	bool bIsTextChanging;
 
