[Fix] Invalidation Panels being redrawn every frames

Due to changes made to handle Dynamic Invalidation switching to Slow Path, regular Invalidation panels are now redrawn every frame if `UE_SLATE_WITH_DYNAMIC_INVALIDATION` is enabled (in the editor only).
As the Invalidation Root code is unaware of Invalidation Panels. the proposed solution is to expose `FSlateInvalidationRoot::CanCacheThisFrame` which will return true if the Invalidation Panel can still cache. If not, the Invalidation Root will switch to Slow Path.

#Synced-CL 50145007

diff --git a/Engine/Source/Runtime/Slate/Private/Widgets/SInvalidationPanel.cpp b/Engine/Source/Runtime/Slate/Private/Widgets/SInvalidationPanel.cpp
--- a/Engine/Source/Runtime/Slate/Private/Widgets/SInvalidationPanel.cpp
+++ b/Engine/Source/Runtime/Slate/Private/Widgets/SInvalidationPanel.cpp
@@ -292,6 +292,22 @@ bool SInvalidationPanel::SupportsInvalidationRecursive(EInvalidationStrategy InI
 	}
 	return SWidget::SupportsInvalidationRecursive(InInvalidationStrategy);
 }
+
+bool SInvalidationPanel::CanCacheThisFrame() const
+{
+#if UE_SLATE_WITH_DYNAMIC_INVALIDATION
+
+	const bool bSupportsInvalidationRecursive = SupportsInvalidationRecursive(EInvalidationStrategy::UseCachedValue);
+	const bool bCanCacheValue = GetCanCache();
+	const bool bCanCacheThisFrame = bCanCacheValue & bSupportsInvalidationRecursive;
+
+#else
+	const bool bCanCacheThisFrame = GetCanCache();
+#endif
+
+	return bCanCacheThisFrame;
+}
+
 #endif
 
 UE_TRACE_EVENT_BEGIN(Cpu, SInvalidationPanel_SupportsInvalidationRecursive, NoSync)
@@ -309,25 +325,20 @@ int32 SInvalidationPanel::OnPaint( const FPaintArgs& Args, const FGeometry& Allo
 	bPaintedSinceLastPrepass = true;
 	SInvalidationPanel* MutableThis = const_cast<SInvalidationPanel*>(this);
 
-#if UE_SLATE_WITH_DYNAMIC_INVALIDATION
-	bool bSupportsInvalidationRecursive;
+	bool bCanCacheThisFrame;
 	{
-		#if UE_SLATE_TRACE_ENABLED
-			UE_TRACE_LOG_SCOPED_T(Cpu, SInvalidationPanel_SupportsInvalidationRecursive, SlateChannel)
-				<< SInvalidationPanel_SupportsInvalidationRecursive.DuringOnPaint(true)
-				<< SInvalidationPanel_SupportsInvalidationRecursive.CreatedInLocation(*GetCreatedInLocation().ToString());
-		#endif // UE_SLATE_TRACE_ENABLED
-
-		bSupportsInvalidationRecursive = SupportsInvalidationRecursive(EInvalidationStrategy::UseCachedValue);
-	}
-	const bool bCanCacheValue = GetCanCache();
-	const bool bCanCacheThisFrame = bCanCacheValue & bSupportsInvalidationRecursive;
-#if WITH_SLATE_DEBUGGING
-	FSlateDebugging::IncrementInvalidationPanelCount(bCanCacheValue, CanUseDynamicInvalidation(), bSupportsInvalidationRecursive);
+#if UE_SLATE_TRACE_ENABLED && UE_SLATE_WITH_DYNAMIC_INVALIDATION
+		UE_TRACE_LOG_SCOPED_T(Cpu, SInvalidationPanel_SupportsInvalidationRecursive, SlateChannel)
+			<< SInvalidationPanel_SupportsInvalidationRecursive.DuringOnPaint(true)
+			<< SInvalidationPanel_SupportsInvalidationRecursive.CreatedInLocation(*GetCreatedInLocation().ToString());
 #endif
-#else
-	const bool bCanCacheThisFrame = GetCanCache();
+		bCanCacheThisFrame = CanCacheThisFrame();
+	}
+
+#if WITH_SLATE_DEBUGGING && UE_SLATE_WITH_DYNAMIC_INVALIDATION
+	FSlateDebugging::IncrementInvalidationPanelCount(GetCanCache(), CanUseDynamicInvalidation(), SupportsInvalidationRecursive(EInvalidationStrategy::UseCachedValue));
 #endif
+
 	if (bCanCacheThisFrame != bWasCachable)
 	{
 		MutableThis->InvalidateRootChildOrder();
diff --git a/Engine/Source/Runtime/Slate/Public/Widgets/SInvalidationPanel.h b/Engine/Source/Runtime/Slate/Public/Widgets/SInvalidationPanel.h
--- a/Engine/Source/Runtime/Slate/Public/Widgets/SInvalidationPanel.h
+++ b/Engine/Source/Runtime/Slate/Public/Widgets/SInvalidationPanel.h
@@ -81,6 +81,8 @@ public:
 	SLATE_API bool CanUseDynamicInvalidation() const;
 
 	SLATE_API virtual bool SupportsInvalidationRecursive(EInvalidationStrategy InInvalidationStrategy) const override;
+
+	SLATE_API virtual bool CanCacheThisFrame() const override;
 #endif
 
 	//~ SWidget overrides
diff --git a/Engine/Source/Runtime/SlateCore/Private/FastUpdate/SlateInvalidationRoot.cpp b/Engine/Source/Runtime/SlateCore/Private/FastUpdate/SlateInvalidationRoot.cpp
--- a/Engine/Source/Runtime/SlateCore/Private/FastUpdate/SlateInvalidationRoot.cpp
+++ b/Engine/Source/Runtime/SlateCore/Private/FastUpdate/SlateInvalidationRoot.cpp
@@ -429,7 +429,7 @@ FSlateInvalidationResult FSlateInvalidationRoot::PaintInvalidationRoot(const FSl
 	if (!bNeedsSlowPath)
 	{
 		// In case widgets were added during Slow Path or PrePass, we need to recheck if we can call ProcessPreUpdate
-		bNeedsSlowPath = !GetRootWidget()->SupportsInvalidationRecursive(SWidget::EInvalidationStrategy::UseCachedValue);
+		bNeedsSlowPath = !CanCacheThisFrame();
 	}
 #endif
 
diff --git a/Engine/Source/Runtime/SlateCore/Public/FastUpdate/SlateInvalidationRoot.h b/Engine/Source/Runtime/SlateCore/Public/FastUpdate/SlateInvalidationRoot.h
--- a/Engine/Source/Runtime/SlateCore/Public/FastUpdate/SlateInvalidationRoot.h
+++ b/Engine/Source/Runtime/SlateCore/Public/FastUpdate/SlateInvalidationRoot.h
@@ -126,6 +126,12 @@ public:
 
 	SLATECORE_API void Advanced_ResetInvalidation(bool bClearResourcesImmediately);
 
+	/**
+	 * @return if the InvalidationRoot can cache this frame. This can return false for a Dynamic Invalidation Panel
+	 * not supporting invalidation, or a disabled Invalidation Panel.
+	 */
+	virtual bool CanCacheThisFrame() const { return true; };
+
 #if WITH_SLATE_DEBUGGING
 	/** @return the last paint type the invalidation root handle used. */
 	ESlateInvalidationPaintType GetLastPaintType() const { return LastPaintType; }
