SEditableTextBox/MultiLineEditableTextBox - Fix Error Reporting

With the Move to Slate Attributes, callbacks like `OnTextChanged` can be triggered when attributes are processed. During that time though, we cannot add or remove children or an ensure will trigger.
To avoid issues, we are creating the error reporting widget on Construct.

#Synced-CL 50145007

diff --git a/Engine/Source/Runtime/Slate/Private/Widgets/Input/SEditableTextBox.cpp b/Engine/Source/Runtime/Slate/Private/Widgets/Input/SEditableTextBox.cpp
--- a/Engine/Source/Runtime/Slate/Private/Widgets/Input/SEditableTextBox.cpp
+++ b/Engine/Source/Runtime/Slate/Private/Widgets/Input/SEditableTextBox.cpp
@@ -151,32 +151,39 @@ void SEditableTextBox::SetError( const FText& InError )
 void SEditableTextBox::SetError( const FString& InError )
 {
 	const bool bHaveError = !InError.IsEmpty();
-	if (!bHaveError)
+
+	if (!ErrorReporting)
 	{
-		if (ErrorReporting.IsValid())
+		// If we are clearing an error but none was set before, we can leave.
+		// If we have an active timer, we might have called this function previously, so we still need to register a timer
+		if (!bHaveError && !HasActiveTimers())
 		{
-			ErrorReporting->AsWidget()->SetVisibility(EVisibility::Collapsed);
-			ErrorReporting->SetError(FString());
+			return;
 		}
-
+		RegisterActiveTimer( 0.f,
+			FWidgetActiveTimerDelegate::CreateLambda([this, InError](double, float)
+			{
+				if (!ErrorReporting) // SetError might have been called a few times, so ErrorReporting might be set
+				{
+					// No error reporting was specified; make a default one
+					TSharedPtr<SPopupErrorText> ErrorTextWidget;
+					Box->AddSlot()
+					.AutoWidth()
+					.Padding(3,0)
+					[
+						SAssignNew( ErrorTextWidget, SPopupErrorText )
+					];
+					ErrorReporting = ErrorTextWidget;
+				}
+				SetError(InError);
+				return EActiveTimerReturnType::Stop;
+			})
+		);
 		return;
 	}
 
-	if ( !ErrorReporting.IsValid() )
-	{
-		// No error reporting was specified; make a default one
-		TSharedPtr<SPopupErrorText> ErrorTextWidget;
-		Box->AddSlot()
-		.AutoWidth()
-		.Padding(3,0)
-		[
-			SAssignNew( ErrorTextWidget, SPopupErrorText )
-		];
-		ErrorReporting = ErrorTextWidget;
-	}
-
-	ErrorReporting->AsWidget()->SetVisibility(EVisibility::Visible);
-	ErrorReporting->SetError( InError );
+	ErrorReporting->AsWidget()->SetVisibility(bHaveError ? EVisibility::Visible : EVisibility::Collapsed);
+	ErrorReporting->SetError(InError);
 }
 
 
@@ -509,16 +516,20 @@ void SEditableTextBox::OnEditableTextChanged(const FText& InText)
 {
 	OnTextChanged.ExecuteIfBound(InText);
 
-	FText OutErrorMessage;
-	if (!FTextEditHelper::VerifyTextLength(InText, OutErrorMessage, MaximumLength.Get(0)) ||
-		(OnVerifyTextChanged.IsBound() && !OnVerifyTextChanged.Execute(InText, OutErrorMessage)))
-	{
-		// Display as an error.
-		SetError(OutErrorMessage);
-	}
-	else
+	const int32 MaximumLengthValue = MaximumLength.Get();
+	if (OnVerifyTextChanged.IsBound() || MaximumLengthValue >= 0)
 	{
-		SetError(FText::GetEmpty());
+		FText OutErrorMessage;
+		if (!FTextEditHelper::VerifyTextLength(InText, OutErrorMessage, MaximumLengthValue) ||
+			(OnVerifyTextChanged.IsBound() && !OnVerifyTextChanged.Execute(InText, OutErrorMessage)))
+		{
+			// Display as an error.
+			SetError(OutErrorMessage);
+		}
+		else
+		{
+			SetError(FText::GetEmpty());
+		}
 	}
 }
 
diff --git a/Engine/Source/Runtime/Slate/Private/Widgets/Input/SMultiLineEditableTextBox.cpp b/Engine/Source/Runtime/Slate/Private/Widgets/Input/SMultiLineEditableTextBox.cpp
--- a/Engine/Source/Runtime/Slate/Private/Widgets/Input/SMultiLineEditableTextBox.cpp
+++ b/Engine/Source/Runtime/Slate/Private/Widgets/Input/SMultiLineEditableTextBox.cpp
@@ -437,20 +437,38 @@ void SMultiLineEditableTextBox::SetError(const FString& InError)
 {
 	const bool bHaveError = !InError.IsEmpty();
 
-	if ( !ErrorReporting.IsValid() )
+	if (!ErrorReporting)
 	{
-		// No error reporting was specified; make a default one
-		TSharedPtr<SPopupErrorText> ErrorTextWidget;
-		Box->AddSlot()
-		.AutoWidth()
-		.Padding(3,0)
-		[
-			SAssignNew( ErrorTextWidget, SPopupErrorText )
-		];
-		ErrorReporting = ErrorTextWidget;
+		// If we are clearing an error but none was set before, we can leave.
+		// If we have an active timer, we might have called this function previously, so we still need to register a timer
+		if (!bHaveError && !HasActiveTimers())
+		{
+			return;
+		}
+		RegisterActiveTimer( 0.f,
+			FWidgetActiveTimerDelegate::CreateLambda([this, InError](double, float)
+			{
+				if (!ErrorReporting) // SetError might have been called a few times, so ErrorReporting might be set
+				{
+					// No error reporting was specified; make a default one
+					TSharedPtr<SPopupErrorText> ErrorTextWidget;
+					Box->AddSlot()
+					.AutoWidth()
+					.Padding(3,0)
+					[
+						SAssignNew( ErrorTextWidget, SPopupErrorText )
+					];
+					ErrorReporting = ErrorTextWidget;
+				}
+				SetError(InError);
+				return EActiveTimerReturnType::Stop;
+			})
+		);
+		return;
 	}
 
-	ErrorReporting->SetError( InError );
+	ErrorReporting->AsWidget()->SetVisibility(bHaveError ? EVisibility::Visible : EVisibility::Collapsed);
+	ErrorReporting->SetError(InError);
 }
 
 bool SMultiLineEditableTextBox::HasError() const
