Slate: Add support for batch adding multiple columns to a SHeaderRow.
This avoids the N^2 behavior when adding one column at a time where we call RegenerateWidgets() for each Add() which internally iterates the current columns.
#rnx
#Synced-CL 50706270

diff --git a/Engine/Source/Runtime/Slate/Private/Widgets/Views/SHeaderRow.cpp b/Engine/Source/Runtime/Slate/Private/Widgets/Views/SHeaderRow.cpp
--- a/Engine/Source/Runtime/Slate/Private/Widgets/Views/SHeaderRow.cpp
+++ b/Engine/Source/Runtime/Slate/Private/Widgets/Views/SHeaderRow.cpp
@@ -616,6 +616,26 @@ void SHeaderRow::AddColumn( const FColumn::FArguments& NewColumnArgs )
 	AddColumn( *NewColumn );
 }
 
+void SHeaderRow::AddColumns( TConstArrayView<FColumn::FArguments> NewColumnArgsArray )
+{
+	Columns.Reserve(Columns.Num() + NewColumnArgsArray.Num());
+
+	const bool bHasFillerColumn = (Columns.Num() > 0 && Columns[Columns.Num() - 1].ColumnId == NAME_None);
+
+	for (FColumn::FArguments const& NewColumnArgs : NewColumnArgsArray)
+	{
+		SHeaderRow::FColumn* NewColumn = new SHeaderRow::FColumn(NewColumnArgs);
+		check(NewColumn->ColumnId != NAME_None);
+
+		const int32 InsertIndex = bHasFillerColumn ? Columns.Num() - 1 : Columns.Num();
+		Columns.Insert(NewColumn, InsertIndex);
+	}
+
+	ColumnsChanged.Broadcast(SharedThis(this));
+
+	RegenerateWidgets();
+}
+
 void SHeaderRow::AddColumn( SHeaderRow::FColumn& NewColumn )
 {
 	int32 InsertIdx = Columns.Num();
diff --git a/Engine/Source/Runtime/Slate/Public/Widgets/Views/SHeaderRow.h b/Engine/Source/Runtime/Slate/Public/Widgets/Views/SHeaderRow.h
--- a/Engine/Source/Runtime/Slate/Public/Widgets/Views/SHeaderRow.h
+++ b/Engine/Source/Runtime/Slate/Public/Widgets/Views/SHeaderRow.h
@@ -352,6 +352,7 @@ public:
 
 	/** Adds a column to the header */
 	SLATE_API void AddColumn( const FColumn::FArguments& NewColumnArgs );
+	SLATE_API void AddColumns( TConstArrayView<FColumn::FArguments> NewColumns );
 	SLATE_API void AddColumn( FColumn& NewColumn );
 
 	/** Inserts a column at the specified index in the header */
