Speculative fix for a crash when call to OnMainContentSlotResized when the PanelDrawerData is invalid.

Not sure how we can end up in this state so I added some defensive validity check when getting the coefficient of the content slot and of the panel drawer to avoid the crash.

#Synced-CL 50736586

diff --git a/Engine/Source/Runtime/Slate/Private/Framework/Docking/SPanelDrawerArea.cpp b/Engine/Source/Runtime/Slate/Private/Framework/Docking/SPanelDrawerArea.cpp
--- a/Engine/Source/Runtime/Slate/Private/Framework/Docking/SPanelDrawerArea.cpp
+++ b/Engine/Source/Runtime/Slate/Private/Framework/Docking/SPanelDrawerArea.cpp
@@ -309,12 +309,12 @@ void SPanelDrawerArea::OnPanelDrawerSlotResized(float InPanelDrawerCoefficient)
 
 float SPanelDrawerArea::GetMainContentCoefficient() const
 {
-	return PanelDrawerData->Size.MainContentCoefficient;
+	return PanelDrawerData ? PanelDrawerData->Size.MainContentCoefficient : 1.f;
 }
 
 float SPanelDrawerArea::GetPanelDrawerCoefficient() const
 {
-	return PanelDrawerData->Size.PanelDrawerCoefficient;
+	return PanelDrawerData ? PanelDrawerData->Size.PanelDrawerCoefficient : 0.f;
 }
 
 void SPanelDrawerArea::PlayAnimation(bool bIsOpening)
