SGUI: Initial Prefab Editor support
Enable via `SceneGraph.EnablePrefabEditorCanvas`

#Synced-CL 50736586

diff --git a/Engine/Source/Runtime/Slate/Private/Widgets/Text/STextBlock.cpp b/Engine/Source/Runtime/Slate/Private/Widgets/Text/STextBlock.cpp
--- a/Engine/Source/Runtime/Slate/Private/Widgets/Text/STextBlock.cpp
+++ b/Engine/Source/Runtime/Slate/Private/Widgets/Text/STextBlock.cpp
@@ -374,7 +374,7 @@ void STextBlock::SetColorAndOpacity(TAttribute<FSlateColor> InColorAndOpacity)
 	ColorAndOpacity.Assign(*this, MoveTemp(InColorAndOpacity));
 }
 
-void STextBlock::SetTextStyle(const FTextBlockStyle* InTextStyle)
+void STextBlock::SetTextStyle(const FTextBlockStyle* InTextStyle, bool bUpdateTextLayout)
 {
 	if (InTextStyle)
 	{
@@ -386,6 +386,11 @@ void STextBlock::SetTextStyle(const FTextBlockStyle* InTextStyle)
 		TextStyle = *Defaults._TextStyle;
 	}
 
+	if (bUpdateTextLayout)
+	{
+		bTextLayoutUpdateTextStyle = true;
+	}
+
 	InvalidateText(EInvalidateWidgetReason::Layout);
 }
 
diff --git a/Engine/Source/Runtime/Slate/Public/Widgets/Text/STextBlock.h b/Engine/Source/Runtime/Slate/Public/Widgets/Text/STextBlock.h
--- a/Engine/Source/Runtime/Slate/Public/Widgets/Text/STextBlock.h
+++ b/Engine/Source/Runtime/Slate/Public/Widgets/Text/STextBlock.h
@@ -212,7 +212,7 @@ public:
 	SLATE_API void SetColorAndOpacity(TAttribute<FSlateColor> InColorAndOpacity);
 
 	/** See TextStyle argument */
-	SLATE_API void SetTextStyle(const FTextBlockStyle* InTextStyle);
+	SLATE_API void SetTextStyle(const FTextBlockStyle* InTextStyle, bool bUpdateTextLayout = false);
 
 	/** See TextShapingMethod attribute */
 	SLATE_API void SetTextShapingMethod(const TOptional<ETextShapingMethod>& InTextShapingMethod);
