Fixed off-screen rendered Children of Children inside ScrollBoxes being unnavigable.
- ScrollPanel defines it's AllottedGeometry to be it's full content size if it uses a FNavigationMethod. SScrollBox overrides the HitTestBoundingBox with it's content size.

#Synced-CL 50826644

diff --git a/Engine/Source/Runtime/Slate/Private/Widgets/Layout/SScrollBox.cpp b/Engine/Source/Runtime/Slate/Private/Widgets/Layout/SScrollBox.cpp
--- a/Engine/Source/Runtime/Slate/Private/Widgets/Layout/SScrollBox.cpp
+++ b/Engine/Source/Runtime/Slate/Private/Widgets/Layout/SScrollBox.cpp
@@ -2,6 +2,7 @@
 
 #include "Widgets/Layout/SScrollBox.h"
 #include "Rendering/DrawElements.h"
+#include "Types/NavigationMetaData.h"
 #include "Types/SlateConstants.h"
 #include "Layout/LayoutUtils.h"
 #include "Widgets/SOverlay.h"
@@ -10,6 +11,130 @@
 
 #include UE_INLINE_GENERATED_CPP_BY_NAME(SScrollBox)
 
+namespace ScrollBoxFocus
+{
+static void EmplaceArrangedWidget(const FGeometry& MyGeometry, TSharedRef<SWidget> InParentWidget, TSharedRef<SWidget> InChildWidget, TMap<TSharedRef<SWidget>, FArrangedWidget>& OutWidgetsToArranged)
+{
+	FArrangedChildren ArrangedChildren(EVisibility::Visible);
+	InParentWidget->ArrangeChildren(MyGeometry, ArrangedChildren);
+	if (FArrangedWidget* FoundSelfArranged = ArrangedChildren.GetInternalArray().FindByPredicate([InChildWidget](const FArrangedWidget& WidgetIt) { return InChildWidget == WidgetIt.Widget; }))
+	{
+		OutWidgetsToArranged.Emplace(FoundSelfArranged->Widget, *FoundSelfArranged);
+	}
+}
+
+static void AppendFocusableChildren(
+	const FGeometry& MyGeometry,
+	TSharedRef<SWidget> InParentWidget,
+	TSharedRef<SWidget> InRuleWidget,
+	const uint32 InUserIndex,
+	const bool bIsRootCall,
+	TSharedPtr<SWidget>& OutUserFocusedWidget,
+	TMap<TSharedRef<SWidget>, FArrangedWidget>& OutWidgetsToArranged,
+	FHittestGrid& OutHitTestGrid)
+{
+	FChildren* InWidgetChildren = InParentWidget->GetChildren();
+	if (InWidgetChildren)
+	{
+		for (int32 SlotIndex = 0; SlotIndex < InWidgetChildren->Num(); ++SlotIndex)
+		{
+			TSharedRef<SWidget> Child = InWidgetChildren->GetChildAt(SlotIndex);
+
+			if (Child->GetVisibility().IsVisible())
+			{
+				const bool bIsRuleWidget = Child == InRuleWidget;
+				const bool bSupportsKeyboardFocus = Child->SupportsKeyboardFocus();
+
+				if (bIsRuleWidget)
+				{
+					EmplaceArrangedWidget(MyGeometry, InParentWidget, Child, OutWidgetsToArranged);
+				}
+				else if (bSupportsKeyboardFocus)
+				{
+					OutHitTestGrid.AddWidget(&Child.Get(), 0, 0, FSlateInvalidationWidgetSortOrder());
+					const bool bHasUserFocus = Child->HasUserFocus(InUserIndex).IsSet() || Child->HasUserFocusedDescendants(InUserIndex);
+					if (bHasUserFocus)
+					{
+						ensureMsgf(!OutUserFocusedWidget.IsValid(), TEXT("This should only be set once"));
+						OutUserFocusedWidget = Child.ToSharedPtr();
+
+						EmplaceArrangedWidget(MyGeometry, InParentWidget, Child, OutWidgetsToArranged);
+					}
+				}
+
+				// We only want to add Focusable Children that are children of the Rule Widget, i.e. the ScrollBox.
+				// We do need to start one widget higher in the hiearchy to arrange the ScrollBox/Rule widget itself
+				// If we added this child widget as keyboard focus already, we do NOT need to search it's children
+				const bool bIsChildThatDoesNotSupportKeyBoardFocus = !bIsRootCall && !bSupportsKeyboardFocus;
+				if (bIsRuleWidget || bIsChildThatDoesNotSupportKeyBoardFocus)
+				{
+					AppendFocusableChildren(MyGeometry, Child, InRuleWidget, InUserIndex, false, OutUserFocusedWidget, OutWidgetsToArranged, OutHitTestGrid);
+				}
+			}
+		}
+	}
+}
+
+static TSharedPtr<SWidget> GetNextWidget(
+	const FGeometry& MyGeometry,
+	const UE::Slate::FDeprecateVector2DResult& InHitTestGridSize,
+	TSharedRef<SWidget> InRuleWidget,
+	const uint32 InUserIndex,
+	const EUINavigation InNavigation,
+	FNavigationMethod& InNavigationMethod)
+{
+	TSharedPtr<SWidget> ReturnWidget;
+
+	TSharedPtr<SWidget> ParentWidget = InRuleWidget->GetParentWidget();
+	if (ParentWidget.IsValid())
+	{
+		TSharedPtr<SWidget> StartingWidget; // This is the current user focused widget
+		TMap<TSharedRef<SWidget>, FArrangedWidget> WidgetsToArranged;
+		FHittestGrid HitTestGrid;
+		{
+			UE::Slate::FDeprecateVector2DResult RectTopLeft = InRuleWidget->GetTickSpaceGeometry().GetRenderBoundingRect().GetTopLeft();		
+			HitTestGrid.SetHittestArea(RectTopLeft, InHitTestGridSize);
+		}
+
+		constexpr bool bIsRootCall = true;
+		AppendFocusableChildren(
+			MyGeometry,
+			ParentWidget.ToSharedRef(),
+			InRuleWidget,
+			InUserIndex,
+			bIsRootCall,
+			StartingWidget,
+			WidgetsToArranged,
+			HitTestGrid);
+
+		if (StartingWidget.IsValid())
+		{
+			const TSharedPtr<FNavigationMetaData> FocusWidgetNavMetaData = StartingWidget->GetMetaData<FNavigationMetaData>();
+			const TSharedPtr<SWidget> ExplicitNextWidget = FocusWidgetNavMetaData.IsValid() ? FocusWidgetNavMetaData->GetFocusRecipient(InNavigation).Pin() : nullptr;
+			if (ExplicitNextWidget.IsValid())
+			{
+				return ExplicitNextWidget;
+			}
+		}
+
+		const bool bRuleWidgetArranged = WidgetsToArranged.Contains(InRuleWidget);
+		const bool bUserFocusedWidgetArranged = StartingWidget.IsValid() && WidgetsToArranged.Contains(StartingWidget.ToSharedRef());
+		const bool bFoundArrangedWidgets = bRuleWidgetArranged && bUserFocusedWidgetArranged;
+		if (bFoundArrangedWidgets)
+		{
+			InNavigationMethod.Initialize(&HitTestGrid, nullptr);
+			
+			const FNavigationReply FocusedNavigationReply = StartingWidget->OnNavigation(MyGeometry, InNavigation);
+			const FArrangedWidget& ArrangedStartingWidget = WidgetsToArranged[StartingWidget.ToSharedRef()];
+			const FArrangedWidget& ArrangedRuleWidget = WidgetsToArranged[InRuleWidget];
+
+			ReturnWidget = InNavigationMethod.FindNextFocusableWidget(ArrangedStartingWidget, InNavigation, FocusedNavigationReply, ArrangedRuleWidget, InUserIndex);
+		}
+	}
+
+	return ReturnWidget;
+}
+} //namespace ScrollBoxFocus
 
 void SScrollBox::FSlot::Construct(const FChildren& SlotOwner, FSlotArguments&& InArgs)
 {
@@ -35,8 +160,9 @@ void SScrollPanel::PrivateRegisterAttributes(FSlateAttributeInitializer& Attribu
 	SScrollBox::FSlot::RegisterAttributes(Initializer);
 }
 
-void SScrollPanel::Construct(const FArguments& InArgs, TArray<SScrollBox::FSlot::FSlotArguments> InSlots)
+void SScrollPanel::Construct(const FArguments& InArgs, TArray<SScrollBox::FSlot::FSlotArguments> InSlots, TSharedRef<SScrollBox> InScrollBox)
 {
+	OwningScrollBox = InScrollBox;
 	PhysicalOffset = 0;
 	Children.AddSlots(MoveTemp(InSlots));
 	Orientation = InArgs._Orientation;
@@ -60,6 +186,24 @@ void SScrollPanel::OnArrangeChildren(const FGeometry& AllottedGeometry, FArrange
 	}
 }
 
+int32 SScrollPanel::OnPaint(const FPaintArgs& Args, const FGeometry& AllottedGeometry, const FSlateRect& MyCullingRect, FSlateWindowElementList& OutDrawElements, int32 LayerId, const FWidgetStyle& InWidgetStyle, bool bParentEnabled) const
+{
+	const bool bHasNavigationMethod = OwningScrollBox.IsValid() && OwningScrollBox.Pin()->HasNavigationMetaData();
+	if (bHasNavigationMethod)
+	{
+		// Use the content size of the scrollbox to navigate to offscreen rendered widgets
+		const FVector2D ScrollPanelSize = FVector2D::Max(AllottedGeometry.Size, GetArrangedSize());
+		const FGeometry ScrollBoxGeometry = AllottedGeometry.MakeChild(ScrollPanelSize, FSlateLayoutTransform());
+		const FSlateRect BoundingRect = OwningScrollBox.IsValid() ? OwningScrollBox.Pin()->GetHitTestBoundingRect() : GetHitTestBoundingRect();
+
+		return SPanel::OnPaint(Args, ScrollBoxGeometry, BoundingRect, OutDrawElements, LayerId, InWidgetStyle, bParentEnabled);
+	}
+	else
+	{
+		return SPanel::OnPaint(Args, AllottedGeometry, MyCullingRect, OutDrawElements, LayerId, InWidgetStyle, bParentEnabled);
+	}
+}
+
 FVector2D SScrollPanel::ComputeDesiredSize(float) const
 {
 	FVector2D ThisDesiredSize = FVector2D::ZeroVector;
@@ -161,7 +305,7 @@ void SScrollBox::Construct( const FArguments& InArgs )
 		bScrollBarIsExternal = false;
 	}
 
-	SAssignNew(ScrollPanel, SScrollPanel, MoveTemp(const_cast<TArray<FSlot::FSlotArguments>&>(InArgs._Slots)))
+	SAssignNew(ScrollPanel, SScrollPanel, MoveTemp(const_cast<TArray<FSlot::FSlotArguments>&>(InArgs._Slots)), StaticCastSharedRef<SScrollBox>(AsShared()))
 		.Clipping(InArgs._Clipping)
 		.Orientation(Orientation)
 		.BackPadScrolling(BackPadScrolling)
@@ -1132,6 +1276,27 @@ void SScrollBox::OnMouseCaptureLost(const FCaptureLostEvent& CaptureLostEvent)
 }
 
 FNavigationReply SScrollBox::OnNavigation(const FGeometry& MyGeometry, const FNavigationEvent& InNavigationEvent)
+{
+	TSharedPtr<FNavigationMetaData> NavMetaData = GetMetaData<FNavigationMetaData>();
+	if (NavMetaData.IsValid())
+	{
+		TInstancedStruct<FNavigationMethod>& NavigationMethodInstance = NavMetaData->GetMutableNavigationMethodStruct();
+		FNavigationMethod* NavigationMethod = NavigationMethodInstance.IsValid() ? NavigationMethodInstance.GetMutablePtr() : nullptr;
+		if (NavigationMethod)
+		{
+			TSharedPtr<SWidget> NextFocusableChild = ScrollBoxFocus::GetNextWidget(MyGeometry, GetScrollPanelContentSize(), AsShared(), InNavigationEvent.GetUserIndex(), InNavigationEvent.GetNavigationType(), *NavigationMethod);
+			if (NextFocusableChild.IsValid())
+			{
+				InternalScrollDescendantIntoView(MyGeometry, NextFocusableChild, false, NavigationDestination, NavigationScrollPadding);
+				return FNavigationReply::Explicit(NextFocusableChild);
+			}
+		}
+	}
+
+	return OnNavigation_Default(MyGeometry, InNavigationEvent);
+}
+
+FNavigationReply SScrollBox::OnNavigation_Default(const FGeometry& MyGeometry, const FNavigationEvent& InNavigationEvent)
 {
 	TSharedPtr<SWidget> FocusedChild;
 	int32 FocusedChildIndex = -1;
@@ -1214,6 +1379,45 @@ FNavigationReply SScrollBox::OnNavigation(const FGeometry& MyGeometry, const FNa
 	return SCompoundWidget::OnNavigation(MyGeometry, InNavigationEvent);
 }
 
+FSlateRect SScrollBox::GetHitTestBoundingRect() const
+{
+	const FGeometry WidgetGeo = GetPaintSpaceGeometry();
+	FSlateRect RenderRect = WidgetGeo.GetRenderBoundingRect();
+
+	const FVector2f ContentSize = GetScrollPanelContentSize();
+	const FVector2f TransformedContentSize = WidgetGeo.GetAccumulatedLayoutTransform().TransformVector(ContentSize);
+
+	RenderRect.Bottom = RenderRect.Top + TransformedContentSize.Y;
+	RenderRect.Right = RenderRect.Left + TransformedContentSize.X;
+
+	if (Orientation == EOrientation::Orient_Vertical)
+	{
+		RenderRect.Top -= DesiredScrollOffset;
+	}
+	else
+	{
+		RenderRect.Left -= DesiredScrollOffset;
+	}
+
+	return RenderRect;
+}
+
+void SScrollBox::OnAddedMetadata(const TSharedRef<ISlateMetaData>& AddedMetaData)
+{
+	if (AddedMetaData->IsOfType<FNavigationMetaData>())
+	{
+		bNavigationDataExists = true;
+	}
+}
+
+void SScrollBox::OnRemovedMetadata(const TSharedRef<ISlateMetaData>& RemovedMetaData)
+{
+	if (RemovedMetaData->IsOfType<FNavigationMetaData>())
+	{
+		bNavigationDataExists = false;
+	}
+}
+
 FReply SScrollBox::OnFocusReceived(const FGeometry& MyGeometry, const FFocusEvent& InFocusEvent)
 {
 	if (OnScrollBoxFocusReceived.IsBound())
diff --git a/Engine/Source/Runtime/Slate/Public/Widgets/Layout/SScrollBox.h b/Engine/Source/Runtime/Slate/Public/Widgets/Layout/SScrollBox.h
--- a/Engine/Source/Runtime/Slate/Public/Widgets/Layout/SScrollBox.h
+++ b/Engine/Source/Runtime/Slate/Public/Widgets/Layout/SScrollBox.h
@@ -24,6 +24,7 @@
 
 #include "SScrollBox.generated.h"
 
+class FNavigationMetaData;
 class FPaintArgs;
 class FSlateWindowElementList;
 class SScrollPanel;
@@ -394,13 +395,24 @@ public:
 	SLATE_API virtual FReply OnFocusReceived(const FGeometry& MyGeometry, const FFocusEvent& InFocusEvent) override;
 	SLATE_API virtual void OnFocusLost(const FFocusEvent& InFocusEvent) override;
 	virtual bool SupportsKeyboardFocus() const override { return bIsFocusable; }
+	SLATE_API virtual FSlateRect GetHitTestBoundingRect() const override;
+	SLATE_API virtual void OnAddedMetadata(const TSharedRef<ISlateMetaData>& AddedMetaData) override;
+	SLATE_API virtual void OnRemovedMetadata(const TSharedRef<ISlateMetaData>& RemovedMetaData) override;
 	// End of SWidget interface
 
+private:
+	FNavigationReply OnNavigation_Default(const FGeometry& MyGeometry, const FNavigationEvent& InNavigationEvent);
+
 protected:
 	SLATE_API void OnClippingChanged();
 
-private:
+public:
+	bool HasNavigationMetaData() const
+	{
+		return bNavigationDataExists;
+	}
 
+private:
 	/** Builds a default Scrollbar */
 	TSharedPtr<SScrollBar> ConstructScrollBar();
 
@@ -485,6 +497,9 @@ protected:
 	/** The scrollbar which controls scrolling for the scrollbox. */
 	TSharedPtr<SScrollBar> ScrollBar;
 
+	/** Flag if Navigation Data exists for the ScrollBox */
+	bool bNavigationDataExists = false;
+
 	/** The Top or Left Image of the shadow hint depending on orientation. */
 	TSharedPtr<SImage> TopLeftShadowHint;
 	/** The Bottom or Right Image of the shadow hint depending on orientation. */
@@ -645,7 +660,7 @@ public:
 	{
 	}
 
-	void Construct(const FArguments& InArgs, TArray<SScrollBox::FSlot::FSlotArguments> InSlots);
+	void Construct(const FArguments& InArgs, TArray<SScrollBox::FSlot::FSlotArguments> InSlots, TSharedRef<SScrollBox> InScrollBox);
 
 public:
 
@@ -661,6 +676,7 @@ public:
 	}
 
 	virtual void OnArrangeChildren(const FGeometry& AllottedGeometry, FArrangedChildren& ArrangedChildren) const override;
+	virtual int32 OnPaint(const FPaintArgs& Args, const FGeometry& AllottedGeometry, const FSlateRect& MyCullingRect, FSlateWindowElementList& OutDrawElements, int32 LayerId, const FWidgetStyle& InWidgetStyle, bool bParentEnabled) const override;
 
 	virtual FChildren* GetChildren() override
 	{
@@ -681,7 +697,7 @@ protected:
 	// End SWidget overrides.
 
 private:
-
+	TWeakPtr<SScrollBox> OwningScrollBox;
 	EOrientation Orientation;
 	bool BackPadScrolling;
 	bool FrontPadScrolling;
diff --git a/Engine/Source/Runtime/SlateCore/Private/Input/HittestGrid.cpp b/Engine/Source/Runtime/SlateCore/Private/Input/HittestGrid.cpp
--- a/Engine/Source/Runtime/SlateCore/Private/Input/HittestGrid.cpp
+++ b/Engine/Source/Runtime/SlateCore/Private/Input/HittestGrid.cpp
@@ -718,9 +718,8 @@ FSlateRect FHittestGrid::GetWidgetRenderBoundingRect(const TSharedPtr<SWidget>&
 	{
 		return FSlateRect();
 	}
-	FGeometry WidgetGeo = Widget->GetPaintSpaceGeometry();
-	WidgetGeo.AppendTransform(FSlateLayoutTransform(-GridWindowOrigin));
-	return WidgetGeo.GetRenderBoundingRect();
+
+	return Widget->GetHitTestBoundingRect().OffsetBy(-GridWindowOrigin);
 }
 
 FIntPoint FHittestGrid::GetCellCoordinate(UE::Slate::FDeprecateVector2DParameter Position) const
diff --git a/Engine/Source/Runtime/SlateCore/Private/Widgets/SWidget.cpp b/Engine/Source/Runtime/SlateCore/Private/Widgets/SWidget.cpp
--- a/Engine/Source/Runtime/SlateCore/Private/Widgets/SWidget.cpp
+++ b/Engine/Source/Runtime/SlateCore/Private/Widgets/SWidget.cpp
@@ -643,7 +643,11 @@ void SWidget::OnFinishedKeyInput()
 
 FNavigationReply SWidget::OnNavigation(const FGeometry& MyGeometry, const FNavigationEvent& InNavigationEvent)
 {
-	EUINavigation Type = InNavigationEvent.GetNavigationType();
+	return OnNavigation(MyGeometry, InNavigationEvent.GetNavigationType());
+}
+
+FNavigationReply SWidget::OnNavigation(const FGeometry& MyGeometry, const EUINavigation Type)
+{
 	TSharedPtr<FNavigationMetaData> NavigationMetaData = GetMetaData<FNavigationMetaData>();
 	if (NavigationMetaData.IsValid())
 	{
@@ -1427,6 +1431,11 @@ FSlateRect SWidget::CalculateCullingAndClippingRules(const FGeometry& AllottedGe
 	return IncomingCullingRect;
 }
 
+FSlateRect SWidget::GetHitTestBoundingRect() const
+{
+	return GetPaintSpaceGeometry().GetRenderBoundingRect();
+}
+
 int32 SWidget::Paint(const FPaintArgs& Args, const FGeometry& AllottedGeometry, const FSlateRect& MyCullingRect, FSlateWindowElementList& OutDrawElements, int32 LayerId, const FWidgetStyle& InWidgetStyle, bool bParentEnabled) const
 {
 	const EWidgetUpdateFlags PreviousUpdateFlag = UpdateFlags;
@@ -1983,6 +1992,8 @@ void SWidget::AddMetadataInternal(const TSharedRef<ISlateMetaData>& AddMe)
 #endif
 	}
 #endif
+
+	OnAddedMetadata(AddMe);
 }
 
 #if UE_SLATE_WITH_DYNAMIC_INVALIDATION
diff --git a/Engine/Source/Runtime/SlateCore/Public/Types/NavigationMetaData.h b/Engine/Source/Runtime/SlateCore/Public/Types/NavigationMetaData.h
--- a/Engine/Source/Runtime/SlateCore/Public/Types/NavigationMetaData.h
+++ b/Engine/Source/Runtime/SlateCore/Public/Types/NavigationMetaData.h
@@ -128,6 +128,11 @@ public:
 		NavigationMethodStruct = InNavigationMethodStruct;
 	}
 
+	TInstancedStruct<FNavigationMethod>& GetMutableNavigationMethodStruct()
+	{
+		return NavigationMethodStruct;
+	}
+
 	const TInstancedStruct<FNavigationMethod>& GetNavigationMethodStruct() const
 	{
 		return NavigationMethodStruct;
diff --git a/Engine/Source/Runtime/SlateCore/Public/Widgets/SWidget.h b/Engine/Source/Runtime/SlateCore/Public/Widgets/SWidget.h
--- a/Engine/Source/Runtime/SlateCore/Public/Widgets/SWidget.h
+++ b/Engine/Source/Runtime/SlateCore/Public/Widgets/SWidget.h
@@ -640,6 +640,11 @@ public:
 	 */
 	SLATECORE_API virtual FNavigationReply OnNavigation(const FGeometry& MyGeometry, const FNavigationEvent& InNavigationEvent);
 
+	/* 
+	 * Helper function to preview Navigation Replies without an explicit Navigation Event
+	 */
+	SLATECORE_API FNavigationReply OnNavigation(const FGeometry& MyGeometry, const EUINavigation Type);
+
 	/**
 	 * Called when the mouse is moved over the widget's window, to determine if we should report whether
 	 * OS-specific features should be active at this location (such as a title bar grip, system menu, etc.)
@@ -706,6 +711,11 @@ public:
 	 */
 	SLATECORE_API FSlateRect CalculateCullingAndClippingRules(const FGeometry& AllottedGeometry, const FSlateRect& MyCullingRect, bool& bClipToBounds, bool& bAlwaysClip, bool& bIntersectClipBounds) const;
 
+	/**
+	 * @return the SlateRect that FHittestGrid would use to test for hit bounds
+	 */
+	SLATECORE_API virtual FSlateRect GetHitTestBoundingRect() const;
+
 	bool HasAnyUpdateFlags(EWidgetUpdateFlags FlagsToCheck) const
 	{
 		return EnumHasAnyFlags(UpdateFlags, FlagsToCheck);
@@ -1407,9 +1417,20 @@ public:
 
 		checkf(Index != 0 || !HasRegisteredSlateAttribute(), TEXT("The first slot is reserved for SlateAttribute"));
 		MetaData.RemoveAtSwap(Index, 1);
+
+		OnRemovedMetadata(RemoveMe);
 		return 1;
 	}
 
+protected:
+	virtual void OnAddedMetadata(const TSharedRef<ISlateMetaData>& AddedMetaData)
+	{
+	}
+
+	virtual void OnRemovedMetadata(const TSharedRef<ISlateMetaData>& RemovedMetaData)
+	{
+	}
+
 private:
 
 	SLATECORE_API void AddMetadataInternal(const TSharedRef<ISlateMetaData>& AddMe);
@@ -1425,6 +1446,7 @@ private:
 			{
 				checkf(Index != 0 || !HasRegisteredSlateAttribute(), TEXT("The first slot is reserved for SlateAttribute"));
 				MetaData.RemoveAtSwap(Index);
+				OnRemovedMetadata(MetaDataEntry);
 			}
 		}
 		return NumBefore - MetaData.Num();
@@ -1440,6 +1462,7 @@ private:
 			{
 				checkf(Index != 0 || !HasRegisteredSlateAttribute(), TEXT("The first slot is reserved for SlateAttribute"));
 				MetaData.RemoveAtSwap(Index);
+				OnRemovedMetadata(MetaDataEntry);
 				return true;
 			}
 		}
