When restoring the editor main window layout we now account for the DPI scaling of the monitor.

#Synced-CL 50736586

diff --git a/Engine/Source/Runtime/Slate/Private/Framework/Docking/TabManager.cpp b/Engine/Source/Runtime/Slate/Private/Framework/Docking/TabManager.cpp
--- a/Engine/Source/Runtime/Slate/Private/Framework/Docking/TabManager.cpp
+++ b/Engine/Source/Runtime/Slate/Private/Framework/Docking/TabManager.cpp
@@ -1199,13 +1199,6 @@ TSharedPtr<SWidget> FTabManager::RestoreFrom(const TSharedRef<FLayout>& Layout,
 			if (bIsPrimaryArea && RestoredDockArea.IsValid() && ensure(!PrimaryDockArea.IsValid()))
 			{
 				PrimaryDockArea = RestoredDockArea;
-
-				if (ThisArea->WindowPlacement == FArea::EWindowPlacementInternal::Placement_Specified && this == &FGlobalTabmanager::Get().Get())
-				{
-					// The global manager will restore the main area window if is was specified
-					TSharedPtr<SWindow> MainWindow = RestoredDockArea->GetParentWindow();
-
-				}
 			}
 		}
 	}
@@ -2252,8 +2245,32 @@ TSharedPtr<SDockingNode> FTabManager::RestoreArea_Helper(const TSharedRef<FLayou
 			{
 				if (ParentWindow)
 				{
-					UE::Slate::FDeprecateVector2DResult WindowSize = ParentWindow->GetWindowSizeFromClientSize(NodeAsArea->UnscaledWindowSize, ParentWindow->GetDPIScaleFactor());
-					ParentWindow->ReshapeWindow(NodeAsArea->UnscaledWindowPosition, WindowSize);
+					/**
+					 * The SWindow api doesn't help here as some function use the DPI resolution and some don't, but this out of scope for this change.
+					 * This code should be refactored to be shared with the window logic also it doesn't support multi DPI scale factor across multiple monitor
+					 */
+					const float DPIScaleFactor = ParentWindow->GetDPIScaleFactor();
+					UE::Slate::FDeprecateVector2DResult WindowSize = ParentWindow->GetWindowSizeFromClientSize(NodeAsArea->UnscaledWindowSize * DPIScaleFactor, DPIScaleFactor);
+					FVector2D WindowPosition = NodeAsArea->UnscaledWindowPosition * DPIScaleFactor;
+
+					// Get desktop metrics
+					FDisplayMetrics DisplayMetrics;
+					FSlateApplicationBase::Get().GetCachedDisplayMetrics(DisplayMetrics);
+					const FPlatformRect& VirtualDisplayRect = DisplayMetrics.VirtualDisplayRect;
+					
+					if (WindowPosition.X < VirtualDisplayRect.Left ||
+						WindowPosition.X >= VirtualDisplayRect.Right ||
+						WindowPosition.Y < VirtualDisplayRect.Top ||
+						WindowPosition.Y >= VirtualDisplayRect.Bottom)
+					{
+						// The window is out of bound of the virtual display rect just use the preferred work area
+						FSlateRect PreferredWorkArea = FSlateApplicationBase::Get().GetPreferredWorkArea();
+						WindowSize = PreferredWorkArea.GetSize();
+						WindowPosition.X = PreferredWorkArea.Left;
+						WindowPosition.Y = PreferredWorkArea.Top;
+					}
+
+					ParentWindow->ReshapeWindow(WindowPosition, WindowSize);
 					if (NodeAsArea->bIsMaximized)
 					{
 						ParentWindow->Maximize();
