Implementing OnBeginTextEdit event for the SSearchBox that's present on the parent SEditableTextBox so that you could do some actions when you click the search button

#Synced-CL 50922933

diff --git a/Engine/Source/Runtime/Slate/Private/Widgets/Input/SSearchBox.cpp b/Engine/Source/Runtime/Slate/Private/Widgets/Input/SSearchBox.cpp
--- a/Engine/Source/Runtime/Slate/Private/Widgets/Input/SSearchBox.cpp
+++ b/Engine/Source/Runtime/Slate/Private/Widgets/Input/SSearchBox.cpp
@@ -33,6 +33,7 @@ void SSearchBox::Construct( const FArguments& InArgs )
 		.ClearKeyboardFocusOnCommit( false )
 		.OnTextChanged( this, &SSearchBox::HandleTextChanged )
 		.OnTextCommitted( this, &SSearchBox::HandleTextCommitted )
+		.OnBeginTextEdit(InArgs._OnBeginTextEdit)
 		.OnVerifyTextChanged( InArgs._OnVerifyTextChanged )
 		.MinDesiredWidth( InArgs._MinDesiredWidth )
 		.OnKeyDownHandler( InArgs._OnKeyDownHandler )
diff --git a/Engine/Source/Runtime/Slate/Public/Widgets/Input/SSearchBox.h b/Engine/Source/Runtime/Slate/Public/Widgets/Input/SSearchBox.h
--- a/Engine/Source/Runtime/Slate/Public/Widgets/Input/SSearchBox.h
+++ b/Engine/Source/Runtime/Slate/Public/Widgets/Input/SSearchBox.h
@@ -75,6 +75,9 @@ public:
 		/** Invoked whenever the text is committed (e.g. user presses enter) */
 		SLATE_EVENT( FOnTextCommitted, OnTextCommitted )
 
+		/** Callback when the text starts to be edited */
+		SLATE_EVENT( FOnBeginTextEdit, OnBeginTextEdit )
+
 		/** Called whenever the text is changed programmatically or interactively by the user */
 		SLATE_EVENT( FOnVerifyTextChanged, OnVerifyTextChanged )
 	
