Remove default keybinding for Close on Mac (was Ctrl+W, which conflicts with Close Major Tab)

[RN] Mac Editor: Cmd+W will no longer close the entire editor, but will now properly close the focused tab
#Synced-CL 50799194

diff --git a/Engine/Source/Runtime/Slate/Private/Framework/MultiBox/Mac/MacMenu.cpp b/Engine/Source/Runtime/Slate/Private/Framework/MultiBox/Mac/MacMenu.cpp
--- a/Engine/Source/Runtime/Slate/Private/Framework/MultiBox/Mac/MacMenu.cpp
+++ b/Engine/Source/Runtime/Slate/Private/Framework/MultiBox/Mac/MacMenu.cpp
@@ -139,7 +139,7 @@ public:
 		UI_COMMAND( HideOthers, "Hide Others", "Hide Others", EUserInterfaceActionType::None, FInputChord( EModifierKey::Control | EModifierKey::Alt, EKeys::H ) )
 		UI_COMMAND( ShowAll, "Show All", "Show All", EUserInterfaceActionType::None, FInputChord() )
 		UI_COMMAND( Minimize, "Minimize", "Minimize", EUserInterfaceActionType::None, FInputChord( EModifierKey::Control, EKeys::M ) )
-		UI_COMMAND( Close, "Close", "Close", EUserInterfaceActionType::None, FInputChord( EModifierKey::Control, EKeys::W ) )
+		UI_COMMAND( Close, "Close", "Close", EUserInterfaceActionType::None, FInputChord( FInputChord() ) )
 		UI_COMMAND( About, "About", "About", EUserInterfaceActionType::None, FInputChord( FInputChord() ) )
 		UI_COMMAND( Quit, "Quit", "Quit", EUserInterfaceActionType::None, FInputChord( EModifierKey::Control, EKeys::Q ) )		
 		UI_COMMAND( Zoom, "Zoom", "Zoom", EUserInterfaceActionType::None, FInputChord() )	
