SMultiLineEditableText : fix an issue where the horizontal scrollbar is visible but not functional when AutoWrapText is enabled.

In 44429508 a change was made to disallow horizontal scrollbars when AutoWrapText is enabled, but the change was only applied in the ::Construct method and therefore it could be made visible later in the ::UpdateAndClampHorizontalScrollBar method but its ::OnHScrollBarMoved handler would never get called since that is only set once in ::Construct.

This fix reads the current value of AutoWrapText in UpdateAndClampHorizontalScrollBar which is called each FSlateEditableTextLayout::Tick.

FTextCustomization : remove AutoWrapText(true) since the EditableTextProperty is only given a single line, it is preferable to have a functioning horizontal scrollbar here.

#Synced-CL 51025453

diff --git a/Engine/Source/Editor/DetailCustomizations/Private/TextCustomization.cpp b/Engine/Source/Editor/DetailCustomizations/Private/TextCustomization.cpp
--- a/Engine/Source/Editor/DetailCustomizations/Private/TextCustomization.cpp
+++ b/Engine/Source/Editor/DetailCustomizations/Private/TextCustomization.cpp
@@ -160,7 +160,6 @@ void FTextCustomization::CustomizeHeader( TSharedRef<class IPropertyHandle> InPr
 		[
 			SNew(STextPropertyEditableTextBox, EditableTextProperty)
 				.Font(FAppStyle::GetFontStyle("PropertyWindow.NormalFont"))
-				.AutoWrapText(true)
 		];
 }
 
diff --git a/Engine/Source/Runtime/Slate/Private/Widgets/Text/SMultiLineEditableText.cpp b/Engine/Source/Runtime/Slate/Private/Widgets/Text/SMultiLineEditableText.cpp
--- a/Engine/Source/Runtime/Slate/Private/Widgets/Text/SMultiLineEditableText.cpp
+++ b/Engine/Source/Runtime/Slate/Private/Widgets/Text/SMultiLineEditableText.cpp
@@ -66,11 +66,7 @@ void SMultiLineEditableText::Construct( const FArguments& InArgs )
 
 	HScrollBar = InArgs._HScrollBar;
 
-	// Horizontal scroll is unnecessary when auto wrap is enabled.
-	// Otherwise if allowed, having both enabled will cause a wrap+scroll relayout feedback loop which can result in jittering
-	const bool bAutoWrapEnabled = InArgs._AutoWrapText.IsSet() && InArgs._AutoWrapText.Get();
-
-    if (HScrollBar.IsValid() && !bAutoWrapEnabled)
+    if (HScrollBar.IsValid())
 	{
 		HScrollBar->SetUserVisibility(EVisibility::Collapsed);
 		HScrollBar->SetOnUserScrolled(FOnUserScrolled::CreateSP(this, &SMultiLineEditableText::OnHScrollBarMoved));
@@ -474,12 +470,23 @@ float SMultiLineEditableText::UpdateAndClampHorizontalScrollBar(const float InVi
 {
 	if (HScrollBar.IsValid())
 	{
-		HScrollBar->SetState(InViewOffset, InViewFraction);
-		HScrollBar->SetUserVisibility(InVisibilityOverride);
-		if (!HScrollBar->IsNeeded())
+		// Horizontal scroll is unnecessary when auto wrap is enabled.
+		// Otherwise if allowed, having both enabled will cause a wrap+scroll relayout feedback loop which can result in jittering
+		const bool bAutoWrapEnabled = SlateEditableTextLayout.GetAutoWrapText();
+
+		if (bAutoWrapEnabled)
 		{
-			// We cannot scroll, so ensure that there is no offset
-			return 0.0f;
+			HScrollBar->SetUserVisibility(EVisibility::Collapsed);
+		}
+		else
+		{
+			HScrollBar->SetState(InViewOffset, InViewFraction);
+			HScrollBar->SetUserVisibility(InVisibilityOverride);
+			if (!HScrollBar->IsNeeded())
+			{
+				// We cannot scroll, so ensure that there is no offset
+				return 0.0f;
+			}
 		}
 	}
 
