Slate - Fix incorrect alignment on header rows in the Fixed and FillSized cases
Manual size was already correctly taking SplitterHandleSize into account

#tests local testing MVG_all asset data misalignment is fixed

#Synced-CL 51073111

diff --git a/Engine/Source/Runtime/Slate/Private/Widgets/Views/SHeaderRow.cpp b/Engine/Source/Runtime/Slate/Private/Widgets/Views/SHeaderRow.cpp
--- a/Engine/Source/Runtime/Slate/Private/Widgets/Views/SHeaderRow.cpp
+++ b/Engine/Source/Runtime/Slate/Private/Widgets/Views/SHeaderRow.cpp
@@ -954,7 +954,7 @@ void SHeaderRow::RegenerateWidgets()
 						.SizeRule(SSplitter::SizeToContent)
 						[
 							SNew(SBox)
-							.WidthOverride(SomeColumn.GetWidth())
+							.WidthOverride(SomeColumn.GetWidth() - SplitterHandleSize)
 							[
 								NewHeader
 							]
@@ -1052,9 +1052,9 @@ void SHeaderRow::RegenerateWidgets()
 
 				case EColumnSizeMode::FillSized:
 				{
-					auto GetColumnWidthAsOptionalSize = [&SomeColumn]() -> FOptionalSize
+					auto GetColumnWidthAsOptionalSize = [&SomeColumn, SplitterHandleSizeCopy = SplitterHandleSize]() -> FOptionalSize
 					{
-						const float DesiredWidth = SomeColumn.GetWidth();
+						const float DesiredWidth = SomeColumn.GetWidth() - SplitterHandleSizeCopy;
 						return FOptionalSize(DesiredWidth);
 					};
 
