Slate: Prevent accidental infinite recursion with FSlateUser::ReleaseCapture

#rnx
#Synced-CL 51173083

diff --git a/Engine/Source/Runtime/Slate/Private/Framework/Application/SlateUser.cpp b/Engine/Source/Runtime/Slate/Private/Framework/Application/SlateUser.cpp
--- a/Engine/Source/Runtime/Slate/Private/Framework/Application/SlateUser.cpp
+++ b/Engine/Source/Runtime/Slate/Private/Framework/Application/SlateUser.cpp
@@ -301,8 +301,15 @@ void FSlateUser::ReleaseCapture(uint32 PointerIndex)
 	if (const FWeakWidgetPath* CaptorPath = PointerCaptorPathsByIndex.Find(PointerIndex))
 	{
 		WidgetsUnderPointerLastEventByIndex.Add(PointerIndex, *CaptorPath);
+		
+		TSharedPtr<SWidget> CaptorWidget = CaptorPath->GetLastWidget().Pin();
+		CaptorPath = nullptr;
 
-		if (TSharedPtr<SWidget> CaptorWidget = CaptorPath->GetLastWidget().Pin())
+		// Make sure we don't recursively come back here, in case OnMouseCaptureLost causes
+		// release capture again.
+		PointerCaptorPathsByIndex.Remove(PointerIndex);
+		
+		if (CaptorWidget.IsValid())
 		{
 			CaptorWidget->OnMouseCaptureLost(FCaptureLostEvent(UserIndex, PointerIndex));
 #if WITH_SLATE_DEBUGGING
@@ -310,9 +317,6 @@ void FSlateUser::ReleaseCapture(uint32 PointerIndex)
 #endif
 		}
 
-		CaptorPath = nullptr;
-		PointerCaptorPathsByIndex.Remove(PointerIndex);
-
 		if (PointerIndex == FSlateApplication::CursorPointerIndex)
 		{
 			// If cursor capture changes, we should refresh the cursor state.
