[ToolMenus] Added support for CustomWidget in SToolBarComboButtonBlock

#Synced-CL 51206406

#Synced-CL 51206406

diff --git a/Engine/Source/Developer/ToolMenus/Private/ToolMenus.cpp b/Engine/Source/Developer/ToolMenus/Private/ToolMenus.cpp
--- a/Engine/Source/Developer/ToolMenus/Private/ToolMenus.cpp
+++ b/Engine/Source/Developer/ToolMenus/Private/ToolMenus.cpp
@@ -10,6 +10,7 @@
 #include "Framework/MultiBox/MultiBoxBuilder.h"
 #include "Framework/Application/SlateApplication.h"
 #include "Framework/MultiBox/SToolBarButtonBlock.h"
+#include "Framework/MultiBox/SToolBarComboButtonBlock.h"
 #include "Internationalization/Internationalization.h"
 #include "UObject/WeakObjectPtrTemplates.h"
 
@@ -2104,7 +2105,7 @@ void UToolMenus::PopulateToolBarBuilderWithEntry(
 			bCouldHaveChildren = true;
 		}
 		
-		ToolBarBuilder.AddComboButton(
+		TSharedPtr<FToolBarComboButtonBlock> ComboButtonBlock = ToolBarBuilder.AddComboButton(
 			UIAction,
 			OnGetContent,
 			Block.Label,
@@ -2119,6 +2120,21 @@ void UToolMenus::PopulateToolBarBuilderWithEntry(
 			ActualResizeParams
 		);
 			
+		if (ComboButtonBlock.IsValid() && Block.MakeCustomWidget.IsBound())
+		{
+			TSharedRef<FMultiBox> MultiBox = ToolBarBuilder.GetMultiBox();
+
+			FToolMenuCustomWidgetContext EntryWidgetContext;
+			EntryWidgetContext.StyleSet = MultiBox->GetStyleSet();
+			EntryWidgetContext.StyleName = MultiBox->GetStyleName();
+
+			TSharedPtr<SWidget> CustomWidget = Block.MakeCustomWidget.Execute(MenuData->Context, EntryWidgetContext);
+			if (CustomWidget.IsValid())
+			{
+				ComboButtonBlock->SetCustomWidget(CustomWidget);
+			}
+		}
+
 		// Also add any top-level flagged children to the toolbar.
 		if (bCouldHaveChildren && !bIsRaisingToTopLevel)
 		{
diff --git a/Engine/Source/Runtime/Slate/Private/Framework/MultiBox/MultiBoxBuilder.cpp b/Engine/Source/Runtime/Slate/Private/Framework/MultiBox/MultiBoxBuilder.cpp
--- a/Engine/Source/Runtime/Slate/Private/Framework/MultiBox/MultiBoxBuilder.cpp
+++ b/Engine/Source/Runtime/Slate/Private/Framework/MultiBox/MultiBoxBuilder.cpp
@@ -750,7 +750,7 @@ TSharedPtr<FToolBarButtonBlock> FToolBarBuilder::AddToolBarButton(
 	return NewToolBarButtonBlock;
 }
 
-void FToolBarBuilder::AddComboButton(
+TSharedPtr<FToolBarComboButtonBlock> FToolBarBuilder::AddComboButton(
 	const FUIAction& InAction,
 	const FOnGetContent& InMenuContentGenerator,
 	const TAttribute<FText>& InLabelOverride,
@@ -785,6 +785,8 @@ void FToolBarBuilder::AddComboButton(
 	UE::Slate::Private::ApplySectionOverridesToBlock(NewToolBarComboButtonBlock, CurrentSectionResizeParams);
 
 	MultiBox->AddMultiBlock( NewToolBarComboButtonBlock );
+
+	return NewToolBarComboButtonBlock;
 }
 
 void FToolBarBuilder::AddToolbarStackButton(const TSharedPtr< const FUICommandInfo > InCommand, FName InTutorialHighlightName)
diff --git a/Engine/Source/Runtime/Slate/Private/Framework/MultiBox/SToolBarComboButtonBlock.cpp b/Engine/Source/Runtime/Slate/Private/Framework/MultiBox/SToolBarComboButtonBlock.cpp
--- a/Engine/Source/Runtime/Slate/Private/Framework/MultiBox/SToolBarComboButtonBlock.cpp
+++ b/Engine/Source/Runtime/Slate/Private/Framework/MultiBox/SToolBarComboButtonBlock.cpp
@@ -156,9 +156,20 @@ void SToolBarComboButtonBlock::BuildMultiBlockWidget(const ISlateStyle* StyleSet
 	// - The ComboButton that opens the menu
 	const bool bIsSplitButton = HasAction();
 
+	// A custom widget overrides all of the default content, including the icon and label, so we don't need to create those if we have a custom widget
+	const bool bHasCustomWidget = ToolBarComboButtonBlock->CustomWidget.IsValid();
+
+	// Some of the same layout considerations happen if either are true
+	const bool bIsSplitButtonOrHasCustomWidget = bIsSplitButton || bHasCustomWidget;
+
 	// Create the content for our button
 	TSharedRef<SWidget> ButtonContent = SNullWidget::NullWidget;
-	if (MultiBox->GetType() == EMultiBoxType::SlimHorizontalToolBar
+	if (!bIsSplitButton && bHasCustomWidget)
+	{
+		// The CustomWidget is the entire content of the button when it's not a split button
+		ButtonContent = ToolBarComboButtonBlock->CustomWidget.ToSharedRef();
+	}
+	else if (MultiBox->GetType() == EMultiBoxType::SlimHorizontalToolBar
 		|| MultiBox->GetType() == EMultiBoxType::SlimHorizontalUniformToolBar
 		|| MultiBox->GetType() == EMultiBoxType::SlimWrappingToolBar)
 	{
@@ -171,8 +182,8 @@ void SToolBarComboButtonBlock::BuildMultiBlockWidget(const ISlateStyle* StyleSet
 			.VAlign(VAlign_Center)
 			.HAlign(HAlign_Center)
 			[
-				// A split button will have the icon handled in its own widget
-				bIsSplitButton ? SNullWidget::NullWidget : IconWidget
+				// A split button or custom widget will have the icon or content handled in its own widget
+				bIsSplitButtonOrHasCustomWidget ? SNullWidget::NullWidget : IconWidget
 			]
 			// Label text
 			+ SHorizontalBox::Slot()
@@ -180,9 +191,9 @@ void SToolBarComboButtonBlock::BuildMultiBlockWidget(const ISlateStyle* StyleSet
 			.MinWidth(ToolBarStyle.ComboContentMinWidth)
 			.MaxWidth(ToolBarStyle.ComboContentMaxWidth)
 			.Padding(TAttribute<FMargin>::CreateLambda(
-				[bIsSplitButton, Label, ToolBarStyle]() -> FMargin
+				[bIsSplitButtonOrHasCustomWidget, Label, ToolBarStyle]() -> FMargin
 				{
-					return bIsSplitButton || Label.Get().IsEmpty() ? FMargin(0) : ToolBarStyle.LabelPadding;
+					return bIsSplitButtonOrHasCustomWidget || Label.Get().IsEmpty() ? FMargin(0) : ToolBarStyle.LabelPadding;
 				}
 			))
 			.HAlign(ToolBarStyle.ComboContentHorizontalAlignment)
@@ -265,6 +276,23 @@ void SToolBarComboButtonBlock::BuildMultiBlockWidget(const ISlateStyle* StyleSet
 
 	if (bIsSplitButton)
 	{
+		TSharedRef<SWidget> SplitButtonContent = SNullWidget::NullWidget;
+		if (bHasCustomWidget)
+		{
+			// If there's a custom widget, in the case of a split button, use it in place of the icon
+			SplitButtonContent = ToolBarComboButtonBlock->CustomWidget.ToSharedRef();
+		}
+		else
+		{
+			SplitButtonContent = 
+				SNew(SBox)
+				.HAlign(HAlign_Center)
+				.VAlign(VAlign_Center)
+				[
+					IconWidget
+				];
+		}
+		
 		if (UserInterfaceType == EUserInterfaceActionType::Button)
 		{
 			// When a button is specified, the combo menu is implied to be
@@ -288,12 +316,7 @@ void SToolBarComboButtonBlock::BuildMultiBlockWidget(const ISlateStyle* StyleSet
 			.IsChecked(this, &SToolBarComboButtonBlock::GetCheckState)
 			.IsEnabled(this, &SToolBarComboButtonBlock::IsEnabled)
 			[
-				SNew(SBox)
-				.HAlign(HAlign_Center)
-				.VAlign(VAlign_Center)
-				[
-					IconWidget
-				]
+				SplitButtonContent
 			];
 		// clang-format on
 	}
diff --git a/Engine/Source/Runtime/Slate/Public/Framework/MultiBox/MultiBoxBuilder.h b/Engine/Source/Runtime/Slate/Public/Framework/MultiBox/MultiBoxBuilder.h
--- a/Engine/Source/Runtime/Slate/Public/Framework/MultiBox/MultiBoxBuilder.h
+++ b/Engine/Source/Runtime/Slate/Public/Framework/MultiBox/MultiBoxBuilder.h
@@ -650,7 +650,7 @@ public:
 	 * @param	InUserInterfaceActionType	Optional user interface action type to use if a command isn't bound. If a command is bound, we use the action type associated with the command instead.
 	 * @param	InResizeParams			Optional resize parameters to control how this entry behaves during toolbar resizing.
 	 */
-	SLATE_API void AddComboButton(
+	SLATE_API TSharedPtr<class FToolBarComboButtonBlock> AddComboButton(
 		const FUIAction& InAction,
 		const FOnGetContent& InMenuContentGenerator,
 		const TAttribute<FText>& InLabelOverride = TAttribute<FText>(),
diff --git a/Engine/Source/Runtime/Slate/Public/Framework/MultiBox/SToolBarComboButtonBlock.h b/Engine/Source/Runtime/Slate/Public/Framework/MultiBox/SToolBarComboButtonBlock.h
--- a/Engine/Source/Runtime/Slate/Public/Framework/MultiBox/SToolBarComboButtonBlock.h
+++ b/Engine/Source/Runtime/Slate/Public/Framework/MultiBox/SToolBarComboButtonBlock.h
@@ -63,6 +63,9 @@ public:
 	void SetForceSmallIcons( const bool InForceSmallIcons ) { bForceSmallIcons = InForceSmallIcons; }
 
 	bool IsSimpleComboBox() const { return bSimpleComboBox; }
+
+	void SetCustomWidget(const TSharedPtr<SWidget>& InCustomWidget) { CustomWidget = InCustomWidget; }
+
 private:
 
 	/**
@@ -99,6 +102,9 @@ private:
 	/** Controls the Labels visibility, defaults to GetIconVisibility if no override is provided */
 	TOptional< EVisibility > LabelVisibility;
 
+	/** Optional custom widget to use instead of the standard button content */
+	TSharedPtr<SWidget> CustomWidget;
+
 	/** In the case where a command is not bound, the user interface action type to use. If a command is bound, we
 		simply use the action type associated with that command. */
 	EUserInterfaceActionType UserInterfaceActionType;
