[Slate] Fix cursor-locking and zone determination edge cases

Before: Cursor-locking reduces the clip rect by 1 pixel on every side if the widget's window is borderless and a game window.
- The original CL that introduced this behavior justifies it as "to prevent clicking outside of the window area"
- If a widget that has a non-integral bounding rect wants to lock the cursor, we should ensure that we round inward when converting to the integer result
- The previous implementation sometimes failed to do this because it used `RoundToInt` instead of `RoundToPositiveInfinity`; this should just be fixed
- There is no reason that the fix or workaround should be specific to widgets that happen to be in windows of any particular kind or borderedness
After: Cursor-locking ensures that the clip rect falls within the locking widget's rect, adjusting by only as much as necessary

Before: cursor-zone-determination will short-circuit to return ClientArea for fullscreen windows that are not borderless game windows
- Later, all fullscreen windows are exempted from resize-border detection, but not from override detection
- It appears that this logic was introduced to exempt borderless game windows from the short-circuit, so that they can have a title bar in full screen that when double-clicked toggles their full-screen status
- Even though the use-case for which the exemption was made happened to be for game windows, there is no reason an editor window might not behave the same way
After: cursor-zone-determination treats edges as client-area when fullscreen, but allows titlebar to override in fullscreen (for any window type)

#rn Cursor-locking rect is no longer offset excessively; fullscreen Normal type windows are now allowed to have titlebars

#Synced-CL 50977340

#Synced-CL 51251369

diff --git a/Engine/Source/Runtime/ApplicationCore/Public/GenericPlatform/GenericApplicationMessageHandler.h b/Engine/Source/Runtime/ApplicationCore/Public/GenericPlatform/GenericApplicationMessageHandler.h
--- a/Engine/Source/Runtime/ApplicationCore/Public/GenericPlatform/GenericApplicationMessageHandler.h
+++ b/Engine/Source/Runtime/ApplicationCore/Public/GenericPlatform/GenericApplicationMessageHandler.h
@@ -86,15 +86,18 @@ namespace EWindowZone
 	enum Type
 	{
 		NotInWindow			= 0,
-		TopLeftBorder		= 1,
-		TopBorder			= 2,
-		TopRightBorder		= 3,
-		LeftBorder			= 4,
-		ClientArea			= 5,
+		ClientArea			= 1,
+
+		// IsBorder below assumes these are contiguous in this order
+		TopLeftBorder		= 2,
+		TopBorder			= 3,
+		TopRightBorder		= 4,
+		LeftBorder			= 5,
 		RightBorder			= 6,
 		BottomLeftBorder	= 7,
 		BottomBorder		= 8,
 		BottomRightBorder	= 9,
+
 		TitleBar			= 10,
 		MinimizeButton		= 11,
 		MaximizeButton		= 12,
@@ -104,6 +107,11 @@ namespace EWindowZone
 		/** No zone specified */
 		Unspecified	= 0,
 	};
+
+	inline bool IsBorder(Type Zone)
+	{
+		return Zone >= TopLeftBorder && Zone <= BottomRightBorder;
+	}
 }
 
 
diff --git a/Engine/Source/Runtime/Slate/Private/Framework/Application/SlateUser.cpp b/Engine/Source/Runtime/Slate/Private/Framework/Application/SlateUser.cpp
--- a/Engine/Source/Runtime/Slate/Private/Framework/Application/SlateUser.cpp
+++ b/Engine/Source/Runtime/Slate/Private/Framework/Application/SlateUser.cpp
@@ -1017,13 +1017,6 @@ void FSlateUser::LockCursorInternal(const FWidgetPath& WidgetPath)
 		LastComputedLockBounds = SlateClipRect;
 		LockingWidgetPath = WidgetPath;
 
-		// Generate a screen space clip rect based on the widget's geometry
-#if PLATFORM_DESKTOP
-		const bool bIsBorderlessGameWindow = NativeWindow->IsDefinitionValid() && NativeWindow->GetDefinition().Type == EWindowType::GameWindow && !NativeWindow->GetDefinition().HasOSWindowBorder;
-		const int32 ClipRectAdjustment = bIsBorderlessGameWindow ? 0 : 1;
-#else
-		const int32 ClipRectAdjustment = 0;
-#endif
 		// Screen space mapping scales everything. When viewport resolution doesn't match platform resolution, 
 		// this causes offset cursor hit-tests in fullscreen. Correct when capturing mouse as viewport widget may be smaller than screen in pixels.
 		if (FSlateApplication::Get().GetTransformFullscreenMouseInput() && !GIsEditor && NativeWindow->GetWindowMode() == EWindowMode::Fullscreen)
@@ -1038,13 +1031,14 @@ void FSlateUser::LockCursorInternal(const FWidgetPath& WidgetPath)
 			SlateClipRect.Bottom /= DisplayDistortion.Y;
 		}
 
-		// Note: We round the upper left coordinate of the clip rect so we guarantee the rect is inside the geometry of the widget.  If we truncated when there is a half pixel we would cause the clip
-		// rect to be half a pixel larger than the geometry and cause the mouse to go outside of the geometry.
+		// Note: Slate uses floating point bounds, but the cursor needs integers. Round the clip rect inwardly, to guarantee it remains inside the
+		// widget.  Truncating the left and/or top would cause the clip rect to fall outside the geometry in cases where the widget's left/top
+		// coordinates are fractional.
 		RECT ClipRect;
-		ClipRect.left = FMath::RoundToInt(SlateClipRect.Left + ClipRectAdjustment);
-		ClipRect.top = FMath::RoundToInt(SlateClipRect.Top + ClipRectAdjustment);
-		ClipRect.right = FMath::TruncToInt(SlateClipRect.Right - ClipRectAdjustment);
-		ClipRect.bottom = FMath::TruncToInt(SlateClipRect.Bottom - ClipRectAdjustment);
+		ClipRect.left = FMath::RoundToPositiveInfinity(SlateClipRect.Left);
+		ClipRect.top = FMath::RoundToPositiveInfinity(SlateClipRect.Top);
+		ClipRect.right = FMath::RoundToNegativeInfinity(SlateClipRect.Right);
+		ClipRect.bottom = FMath::RoundToNegativeInfinity(SlateClipRect.Bottom);
 
 		// Lock the cursor to the widget
 		Cursor->Lock(&ClipRect);
diff --git a/Engine/Source/Runtime/Slate/Private/Test/SWindowZoneTest.cpp b/Engine/Source/Runtime/Slate/Private/Test/SWindowZoneTest.cpp
new file mode 100644
--- /dev/null
+++ b/Engine/Source/Runtime/Slate/Private/Test/SWindowZoneTest.cpp
@@ -0,0 +1,416 @@
+// Copyright Epic Games, Inc. All Rights Reserved.
+
+#include "CoreMinimal.h"
+#include "Misc/AutomationTest.h"
+
+#if WITH_AUTOMATION_WORKER && PLATFORM_DESKTOP
+
+#include "Widgets/SWindow.h"
+#include "Widgets/SLeafWidget.h"
+#include "GenericPlatform/GenericWindow.h"
+#include "GenericPlatform/GenericWindowDefinition.h"
+#include "Framework/Application/SlateApplication.h"
+#include "Framework/Application/SlateUser.h"
+
+namespace UE::Slate::Test
+{
+
+namespace
+{
+	constexpr float WindowSize = 256.0f;
+	constexpr float ResizeBorderSize = 5.0f;
+
+	struct FTestPositions
+	{
+		// Outside window positions
+		static FVector2D Outside() { return FVector2D(-10.0, -10.0); }
+		static FVector2D OutsideTopRight() { return FVector2D(WindowSize + 10.0, -10.0); }
+		static FVector2D OutsideBottomLeft() { return FVector2D(-10.0, WindowSize + 10.0); }
+		static FVector2D OutsideBottomRight() { return FVector2D(WindowSize + 10.0, WindowSize + 10.0); }
+		static FVector2D OutsideTop() { return FVector2D(WindowSize / 2.0, -10.0); }
+		static FVector2D OutsideBottom() { return FVector2D(WindowSize / 2.0, WindowSize + 10.0); }
+
+		// Border positions (within resize border region)
+		static FVector2D TopLeft() { return FVector2D(2.0, 2.0); }
+		static FVector2D Top() { return FVector2D(WindowSize / 2.0, 2.0); }
+		static FVector2D TopRight() { return FVector2D(WindowSize - 2.0, 2.0); }
+		static FVector2D Left() { return FVector2D(2.0, WindowSize / 2.0); }
+		static FVector2D Right() { return FVector2D(WindowSize - 2.0, WindowSize / 2.0); }
+		static FVector2D BottomLeft() { return FVector2D(2.0, WindowSize - 2.0); }
+		static FVector2D Bottom() { return FVector2D(WindowSize / 2.0, WindowSize - 2.0); }
+		static FVector2D BottomRight() { return FVector2D(WindowSize - 2.0, WindowSize - 2.0); }
+
+		// Center (client area)
+		static FVector2D Center() { return FVector2D(WindowSize / 2.0, WindowSize / 2.0); }
+	};
+
+	class FMockGenericWindow : public FGenericWindow
+	{
+	public:
+		FMockGenericWindow(EWindowMode::Type InMockWindowMode)
+			: MockWindowMode(InMockWindowMode)
+		{
+		}
+
+		virtual void SetWindowMode(EWindowMode::Type InNewWindowMode) override
+		{
+			MockWindowMode = InNewWindowMode;
+		}
+
+		virtual EWindowMode::Type GetWindowMode() const override
+		{
+			return MockWindowMode;
+		}
+
+		virtual bool IsMaximized() const override
+		{
+			return false;
+		}
+
+	private:
+		EWindowMode::Type MockWindowMode;
+	};
+
+	class SZoneOverrideWidget : public SLeafWidget
+	{
+	public:
+		SLATE_BEGIN_ARGS(SZoneOverrideWidget)
+			: _ZoneOverride(EWindowZone::Unspecified)
+		{
+		}
+		SLATE_ARGUMENT(EWindowZone::Type, ZoneOverride)
+		SLATE_END_ARGS()
+
+		void Construct(const FArguments& InArgs)
+		{
+			ZoneOverride = InArgs._ZoneOverride;
+		}
+
+		virtual FVector2D ComputeDesiredSize(float) const override
+		{
+			return FVector2D(WindowSize, WindowSize);
+		}
+
+		virtual int32 OnPaint(
+			const FPaintArgs& Args,
+			const FGeometry& AllottedGeometry,
+			const FSlateRect& MyCullingRect,
+			FSlateWindowElementList& OutDrawElements,
+			int32 LayerId,
+			const FWidgetStyle& InWidgetStyle,
+			bool bParentEnabled) const override
+		{
+			return LayerId;
+		}
+
+		virtual EWindowZone::Type GetWindowZoneOverride() const override
+		{
+			return ZoneOverride;
+		}
+
+	private:
+		EWindowZone::Type ZoneOverride;
+	};
+
+	TSharedRef<SWindow> CreateTestWindow(
+		EWindowType WindowType,
+		bool bUseOSWindowBorder,
+		ESizingRule SizingRule,
+		EWindowMode::Type WindowMode = EWindowMode::Windowed)
+	{
+		TSharedRef<SWindow> Result = SNew(SWindow)
+			.Type(WindowType)
+			.UseOSWindowBorder(bUseOSWindowBorder)
+			.ClientSize(FVector2D(WindowSize, WindowSize))
+			.ScreenPosition(FVector2f(256, 256))
+			.CreateTitleBar(false)
+			.SizingRule(SizingRule)
+			.UserResizeBorder(FMargin(ResizeBorderSize))
+			.Tag("Test");
+		Result->SetNativeWindow(MakeShared<FMockGenericWindow>(WindowMode));
+		return Result;
+	}
+
+	TSharedRef<SWindow> CreateTestWindow_ZoneOverride(
+		EWindowType WindowType,
+		bool bUseOSWindowBorder,
+		ESizingRule SizingRule,
+		EWindowZone::Type ZoneOverride,
+		TSharedPtr<SZoneOverrideWidget>& OutZoneWidget,
+		EWindowMode::Type WindowMode = EWindowMode::Windowed)
+	{
+		TSharedRef<SZoneOverrideWidget> ZoneWidget = SNew(SZoneOverrideWidget)
+			.ZoneOverride(ZoneOverride);
+		OutZoneWidget = ZoneWidget;
+
+		TSharedRef<SWindow> Result = SNew(SWindow)
+			.Type(WindowType)
+			.UseOSWindowBorder(bUseOSWindowBorder)
+			.ClientSize(FVector2D(WindowSize, WindowSize))
+			.ScreenPosition(FVector2f(100, 100))
+			.CreateTitleBar(false)
+			.SizingRule(SizingRule)
+			.UserResizeBorder(FMargin(ResizeBorderSize))
+			.Tag("Test")
+			[
+				ZoneWidget
+			];
+		Result->SetNativeWindow(MakeShared<FMockGenericWindow>(WindowMode));
+		return Result;
+	}
+
+	class FCursorModeScope
+	{
+	public:
+		FCursorModeScope(FAutomationTestBase& Test)
+		{
+			TSharedPtr<ICursor> Cursor = FSlateApplicationBase::Get().GetPlatformCursor();
+			OriginalType = Cursor->GetType();
+			Cursor->SetType(EMouseCursor::Default);
+		}
+
+		~FCursorModeScope()
+		{
+			FSlateApplicationBase::Get().GetPlatformCursor()->SetType(OriginalType);
+		}
+
+		EMouseCursor::Type OriginalType;
+	};
+
+	constexpr EAutomationTestFlags SlateWindowZoneTestFlags = EAutomationTestFlags::EditorContext | EAutomationTestFlags::ClientContext |
+		EAutomationTestFlags::ProgramContext | EAutomationTestFlags::LowPriority | EAutomationTestFlags::EngineFilter;
+}
+
+	/**
+	 * Test GetCurrentWindowZone returns NotInWindow for positions outside the window.
+	 * Tests all combinations of:
+	 * - Window types: GameWindow, Normal (2)
+	 * - Sizing rules: UserSized, FixedSize, Autosized (3)
+	 * Total: 6 combinations
+	 */
+	IMPLEMENT_SIMPLE_AUTOMATION_TEST(FSWindowZoneNotInWindowTest, "Slate.Window.GetCurrentWindowZone.NotInWindow", SlateWindowZoneTestFlags)
+
+	bool FSWindowZoneNotInWindowTest::RunTest(const FString& Parameters)
+	{
+		FCursorModeScope CursorModeScope(*this);
+		for (EWindowType WindowType : { EWindowType::GameWindow, EWindowType::Normal })
+		{
+			for (ESizingRule SizingRule : { ESizingRule::UserSized, ESizingRule::FixedSize, ESizingRule::Autosized })
+			{
+				TSharedRef<SWindow> Window = CreateTestWindow(WindowType, /*bUseOSWindowBorder*/ true, SizingRule);
+				for (const FVector2D& OutsidePos : {
+					FTestPositions::Outside(),
+					FTestPositions::OutsideTopRight(),
+					FTestPositions::OutsideBottomLeft(),
+					FTestPositions::OutsideBottomRight(),
+					FTestPositions::OutsideTop(),
+					FTestPositions::OutsideBottom(),
+				})
+				{
+					constexpr const ANSICHAR* SizingRuleNames[] = { "UserSized", "FixedSize", "Autosized" };
+
+					const EWindowZone::Type ActualZone = Window->GetCurrentWindowZone(OutsidePos);
+					TestEqual<int32>(
+						FString::Printf(TEXT("%S/%S at (%.0f,%.0f)"),
+							WindowType == EWindowType::Normal ? "Normal" : "GameWindow", SizingRuleNames[static_cast<int32>(SizingRule)],
+							OutsidePos.X, OutsidePos.Y),
+						ActualZone,
+						EWindowZone::NotInWindow);
+				}
+
+				FSlateApplication::Get().DestroyWindowImmediately(Window);
+			}
+		}
+
+		return true;
+	}
+
+	/**
+	 * Test GetCurrentWindowZone in WindowedFullscreen mode.
+	 * WindowedFullscreen should never return edge zones.
+	 */
+	IMPLEMENT_SIMPLE_AUTOMATION_TEST(FSWindowZoneWindowedFullscreenTest, "Slate.Window.GetCurrentWindowZone.WindowedFullscreen", SlateWindowZoneTestFlags)
+
+	bool FSWindowZoneWindowedFullscreenTest::RunTest(const FString& Parameters)
+	{
+		FCursorModeScope CursorModeScope(*this);
+		for (EWindowType WindowType : { EWindowType::GameWindow, EWindowType::Normal })
+		{
+			for (bool bOsWindowBorder : { false, true })
+			{
+				// Test without widget override
+				{
+					TSharedRef<SWindow> Window = CreateTestWindow(
+						WindowType, bOsWindowBorder, ESizingRule::UserSized, EWindowMode::WindowedFullscreen);
+
+					for (const FVector2D& Pos : {
+						FTestPositions::TopLeft(),
+						FTestPositions::Top(),
+						FTestPositions::TopRight(),
+						FTestPositions::Left(),
+						FTestPositions::Center(),
+						FTestPositions::Right(),
+						FTestPositions::BottomLeft(),
+						FTestPositions::Bottom(),
+						FTestPositions::BottomRight(),
+					})
+					{
+						const EWindowZone::Type ActualZone = Window->GetCurrentWindowZone(Pos);
+						TestEqual<int32>(
+							FString::Printf(TEXT("%S %S WindowedFullscreen at (%.0f,%.0f)"),
+								WindowType == EWindowType::Normal ? "Normal" : "GameWindow",
+								bOsWindowBorder ? "with OS window border" : "without OS window border",
+								Pos.X, Pos.Y),
+							ActualZone,
+							EWindowZone::ClientArea);
+					}
+
+					FSlateApplication::Get().DestroyWindowImmediately(Window);
+				}
+
+				// Test with widget zone override
+				{
+					TSharedPtr<SZoneOverrideWidget> ZoneWidget;
+					TSharedRef<SWindow> Window = CreateTestWindow_ZoneOverride(
+						WindowType, bOsWindowBorder, ESizingRule::UserSized, EWindowZone::TitleBar, ZoneWidget, EWindowMode::WindowedFullscreen);
+
+					FSlateApplication::Get().ForceRedrawWindow(Window);
+
+					const EWindowZone::Type ActualZone = Window->GetCurrentWindowZone(FTestPositions::Center());
+					TestEqual<int32>(
+						FString::Printf(TEXT("%S %S WindowedFullscreen respects widget override"),
+							WindowType == EWindowType::Normal ? "Normal" : "GameWindow",
+							bOsWindowBorder ? "with OS window border" : "without OS window border"),
+						ActualZone,
+						EWindowZone::TitleBar);
+
+					FSlateApplication::Get().DestroyWindowImmediately(Window);
+				}
+			}
+		}
+
+		return true;
+	}
+
+	/**
+	 * Test GetCurrentWindowZone for FixedSize and Autosized windows.
+	 * Non-resizable windows:
+	 * - Should return ClientArea for edge and middle cursor locations
+	 * - BUT should respect widget zone overrides
+	 */
+	IMPLEMENT_SIMPLE_AUTOMATION_TEST(FSWindowZoneNonResizableTest, "Slate.Window.GetCurrentWindowZone.NonResizable", SlateWindowZoneTestFlags)
+
+	bool FSWindowZoneNonResizableTest::RunTest(const FString& Parameters)
+	{
+		FCursorModeScope CursorModeScope(*this);
+		for (ESizingRule SizingRule : { ESizingRule::FixedSize, ESizingRule::Autosized })
+		{
+			// Test: Edge positions should return ClientArea (no resize handles)
+			{
+				TSharedRef<SWindow> Window = CreateTestWindow(EWindowType::Normal, /*bUseOSWindowBorder*/ true, SizingRule);
+
+				for (const FVector2D& EdgePos : {
+					FTestPositions::TopLeft(),
+					FTestPositions::Top(),
+					FTestPositions::TopRight(),
+					FTestPositions::Left(),
+					FTestPositions::Right(),
+					FTestPositions::BottomLeft(),
+					FTestPositions::Bottom(),
+					FTestPositions::BottomRight(),
+				})
+				{
+					const EWindowZone::Type ActualZone = Window->GetCurrentWindowZone(EdgePos);
+					TestEqual<int32>(
+						FString::Printf(TEXT("%S edge at (%.0f,%.0f)"),
+							SizingRule == ESizingRule::FixedSize ? "FixedSize" : "Autosized", EdgePos.X, EdgePos.Y),
+						ActualZone,
+						EWindowZone::ClientArea);
+				}
+
+				// Center should also be ClientArea
+				TestEqual<int32>(
+					FString::Printf(TEXT("%S center"), SizingRule == ESizingRule::FixedSize ? "FixedSize" : "Autosized"),
+					Window->GetCurrentWindowZone(FTestPositions::Center()),
+					EWindowZone::ClientArea);
+
+				FSlateApplication::Get().DestroyWindowImmediately(Window);
+			}
+
+			// Test: Widget zone override should be respected
+			{
+				TSharedPtr<SZoneOverrideWidget> ZoneWidget;
+				TSharedRef<SWindow> Window = CreateTestWindow_ZoneOverride(
+					EWindowType::Normal,
+					/*bUseOSWindowBorder*/ true,
+					SizingRule,
+					EWindowZone::TitleBar,
+					ZoneWidget);
+
+				FSlateApplication::Get().ForceRedrawWindow(Window);
+
+				const EWindowZone::Type ActualZone = Window->GetCurrentWindowZone(FTestPositions::Center());
+				TestEqual<int32>(
+					FString::Printf(TEXT("%S respects TitleBar widget override"),
+						SizingRule == ESizingRule::FixedSize ? "FixedSize" : "Autosized"),
+					ActualZone,
+					EWindowZone::TitleBar);
+
+				FSlateApplication::Get().DestroyWindowImmediately(Window);
+			}
+		}
+
+		return true;
+	}
+
+	/**
+	 * Test GetCurrentWindowZone for UserSized (resizable) windows.
+	 * UserSized windows should return appropriate border zones based on cursor position.
+	 */
+	IMPLEMENT_SIMPLE_AUTOMATION_TEST(FSWindowZoneUserSizedBordersTest, "Slate.Window.GetCurrentWindowZone.UserSizedBorders", SlateWindowZoneTestFlags)
+
+	bool FSWindowZoneUserSizedBordersTest::RunTest(const FString& Parameters)
+	{
+		FCursorModeScope CursorModeScope(*this);
+		struct FZoneTestCase
+		{
+			FVector2D Position;
+			EWindowZone::Type ExpectedZone;
+			const ANSICHAR* Name;
+		} TestCases[] = {
+			{ FTestPositions::TopLeft(),     EWindowZone::TopLeftBorder,     "TopLeftBorder" },
+			{ FTestPositions::Top(),         EWindowZone::TopBorder,         "TopBorder" },
+			{ FTestPositions::TopRight(),    EWindowZone::TopRightBorder,    "TopRightBorder" },
+			{ FTestPositions::Left(),        EWindowZone::LeftBorder,        "LeftBorder" },
+			{ FTestPositions::Center(),      EWindowZone::ClientArea,        "ClientArea" },
+			{ FTestPositions::Right(),       EWindowZone::RightBorder,       "RightBorder" },
+			{ FTestPositions::BottomLeft(),  EWindowZone::BottomLeftBorder,  "BottomLeftBorder" },
+			{ FTestPositions::Bottom(),      EWindowZone::BottomBorder,      "BottomBorder" },
+			{ FTestPositions::BottomRight(), EWindowZone::BottomRightBorder, "BottomRightBorder" },
+		};
+
+		for (EWindowType WindowType : { EWindowType::GameWindow, EWindowType::Normal })
+		{
+			TSharedRef<SWindow> Window = CreateTestWindow(
+				WindowType,
+				/*bUseOSWindowBorder*/ true,
+				ESizingRule::UserSized);
+
+			for (const FZoneTestCase& TestCase : TestCases)
+			{
+				const EWindowZone::Type ActualZone = Window->GetCurrentWindowZone(TestCase.Position);
+				TestEqual<int32>(
+					FString::Printf(TEXT("%S UserSized: %S"),
+						WindowType == EWindowType::GameWindow ? "GameWindow" : "Normal", TestCase.Name),
+					ActualZone,
+					TestCase.ExpectedZone);
+			}
+
+			FSlateApplication::Get().DestroyWindowImmediately(Window);
+		}
+
+		return true;
+	}
+} // namespace UE::Slate::Test
+
+#endif // WITH_AUTOMATION_WORKER && PLATFORM_DESKTOP
diff --git a/Engine/Source/Runtime/Slate/Private/Test/SlateUserCursorLockTest.cpp b/Engine/Source/Runtime/Slate/Private/Test/SlateUserCursorLockTest.cpp
new file mode 100644
--- /dev/null
+++ b/Engine/Source/Runtime/Slate/Private/Test/SlateUserCursorLockTest.cpp
@@ -0,0 +1,538 @@
+// Copyright Epic Games, Inc. All Rights Reserved.
+
+
+#include "CoreMinimal.h"
+#include "Misc/AutomationTest.h"
+
+#if WITH_AUTOMATION_WORKER && PLATFORM_DESKTOP
+
+#if PLATFORM_MICROSOFT
+// Needed to be able to use RECT
+#include "Microsoft/WindowsHWrapper.h"
+#elif PLATFORM_MAC || PLATFORM_IOS
+#include "Apple/AppleTagRect.h"
+#elif PLATFORM_LINUX
+#include "Unix/UnixSystemIncludes.h"
+#endif
+#include "SlateGlobals.h"
+#include "Framework/Application/SlateApplication.h"
+#include "Framework/Application/SlateUser.h"
+#include "Widgets/SWindow.h"
+#include "Widgets/SLeafWidget.h"
+#include "Widgets/SCompoundWidget.h"
+#include "Layout/Children.h"
+#include "GenericPlatform/GenericWindow.h"
+#include "GenericPlatform/GenericWindowDefinition.h"
+#include "GenericPlatform/ICursor.h"
+
+namespace
+{
+	using namespace UE::Slate;
+
+	/**
+	 * Mock window implementation that always reports itself as being in the foreground.
+	 * This allows cursor lock tests to run without depending on desktop window manager behavior,
+	 * which varies across platforms (e.g., macOS and Wayland may not allow forcing windows to foreground).
+	 */
+	class FMockForegroundWindow : public FGenericWindow
+	{
+	public:
+		FMockForegroundWindow()
+		{
+			// Initialize a minimal window definition
+			Definition = MakeShareable(new FGenericWindowDefinition());
+			Definition->Type = EWindowType::Normal;
+			Definition->XDesiredPositionOnScreen = 0.0f;
+			Definition->YDesiredPositionOnScreen = 0.0f;
+			Definition->WidthDesiredOnScreen = 100.0f;
+			Definition->HeightDesiredOnScreen = 100.0f;
+			Definition->TransparencySupport = EWindowTransparency::None;
+			Definition->HasOSWindowBorder = false;
+			Definition->AppearsInTaskbar = false;
+			Definition->IsTopmostWindow = false;
+			Definition->AcceptsInput = true;
+			Definition->ActivationPolicy = EWindowActivationPolicy::Always;
+			Definition->FocusWhenFirstShown = true;
+			Definition->HasCloseButton = false;
+			Definition->SupportsMinimize = false;
+			Definition->SupportsMaximize = false;
+			Definition->IsModalWindow = false;
+			Definition->IsRegularWindow = true;
+			Definition->HasSizingFrame = false;
+			Definition->SizeWillChangeOften = false;
+			Definition->ShouldPreserveAspectRatio = false;
+			Definition->ExpectedMaxWidth = 0;
+			Definition->ExpectedMaxHeight = 0;
+			Definition->Opacity = 1.0f;
+			Definition->CornerRadius = 0;
+		}
+
+		virtual bool IsForegroundWindow() const override
+		{
+			return true; // Always report as foreground for testing
+		}
+
+		virtual float GetDPIScaleFactor() const override
+		{
+			return 1.0f;
+		}
+
+		virtual bool IsDefinitionValid() const override
+		{
+			return Definition.IsValid();
+		}
+
+		virtual const FGenericWindowDefinition& GetDefinition() const override
+		{
+			return *Definition;
+		}
+	};
+
+	/**
+	 * Mock cursor implementation that captures the RECT passed to Lock() for verification.
+	 */
+	class FMockCursorForLockTest : public ICursor
+	{
+	public:
+		FMockCursorForLockTest() = default;
+
+		// ICursor interface
+		virtual FVector2D GetPosition() const override
+		{
+			return CursorPosition;
+		}
+		virtual void SetPosition(const int32 X, const int32 Y) override
+		{
+			CursorPosition = FVector2D(X, Y);
+		}
+		virtual void SetType(const EMouseCursor::Type InNewCursor) override
+		{
+			CurrentType = InNewCursor;
+		}
+		virtual EMouseCursor::Type GetType() const override
+		{
+			return CurrentType;
+		}
+		virtual void GetSize(int32& Width, int32& Height) const override
+		{
+			Width = 16; Height = 16;
+		}
+		virtual void Show(bool bShow) override
+		{
+			bIsVisible = bShow;
+		}
+
+		virtual void Lock(const RECT* const Bounds) override
+		{
+			bLockWasCalled = true;
+			if (Bounds)
+			{
+				bHasCapturedBounds = true;
+				CapturedBounds = *Bounds;
+			}
+			else
+			{
+				bHasCapturedBounds = false;
+			}
+		}
+
+		virtual void SetTypeShape(EMouseCursor::Type InCursorType, void* CursorHandle) override
+		{
+		}
+		// end ICursor interface
+
+		void Reset()
+		{
+			bLockWasCalled = false;
+			bHasCapturedBounds = false;
+			CapturedBounds = {};
+		}
+
+		bool bLockWasCalled = false;
+		bool bHasCapturedBounds = false;
+		RECT CapturedBounds = {};
+	private:
+		FVector2D CursorPosition = FVector2D::ZeroVector;
+		EMouseCursor::Type CurrentType = EMouseCursor::Default;
+		bool bIsVisible = true;
+	};
+
+	/**
+	 * Custom leaf widget for testing that reports a specific desired size.
+	 */
+	class STestWidgetWithKnownSize : public SLeafWidget
+	{
+	public:
+		SLATE_BEGIN_ARGS(STestWidgetWithKnownSize)
+			: _DesiredWidth(100.0f)
+			, _DesiredHeight(100.0f)
+		{}
+		SLATE_ARGUMENT(float, DesiredWidth)
+		SLATE_ARGUMENT(float, DesiredHeight)
+		SLATE_END_ARGS()
+
+		void Construct(const FArguments& InArgs)
+		{
+			DesiredWidth = InArgs._DesiredWidth;
+			DesiredHeight = InArgs._DesiredHeight;
+		}
+
+		virtual FVector2D ComputeDesiredSize(float) const override
+		{
+			return FVector2D(DesiredWidth, DesiredHeight);
+		}
+
+		virtual int32 OnPaint(const FPaintArgs& Args, const FGeometry& AllottedGeometry,
+			const FSlateRect& MyCullingRect, FSlateWindowElementList& OutDrawElements,
+			int32 LayerId, const FWidgetStyle& InWidgetStyle, bool bParentEnabled) const override
+		{
+			return LayerId;
+		}
+
+	private:
+		float DesiredWidth = 100.0f;
+		float DesiredHeight = 100.0f;
+	};
+
+	/**
+	 * Container widget that positions its child at a specific fractional offset.
+	 * Used to achieve fractional widget geometry bounds for testing cursor lock rounding.
+	 */
+	class SFractionalOffsetBox : public SCompoundWidget
+	{
+	public:
+		SLATE_BEGIN_ARGS(SFractionalOffsetBox)
+			: _OffsetX(0.0f)
+			, _OffsetY(0.0f)
+		{}
+		SLATE_ARGUMENT(float, OffsetX)
+		SLATE_ARGUMENT(float, OffsetY)
+		SLATE_DEFAULT_SLOT(FArguments, Content)
+		SLATE_END_ARGS()
+
+		SFractionalOffsetBox()
+			: OffsetX(0.0f)
+			, OffsetY(0.0f)
+			, ChildSlot(this)
+		{
+		}
+
+		void Construct(const FArguments& InArgs)
+		{
+			OffsetX = InArgs._OffsetX;
+			OffsetY = InArgs._OffsetY;
+
+			ChildSlot
+			[
+				InArgs._Content.Widget
+			];
+		}
+
+		virtual void OnArrangeChildren(const FGeometry& AllottedGeometry, FArrangedChildren& ArrangedChildren) const override
+		{
+			if (ChildSlot.GetWidget()->GetVisibility() != EVisibility::Collapsed)
+			{
+				const FVector2D ChildDesiredSize = ChildSlot.GetWidget()->GetDesiredSize();
+				ArrangedChildren.AddWidget(
+					AllottedGeometry.MakeChild(
+						ChildSlot.GetWidget(),
+						FVector2D(OffsetX, OffsetY),
+						ChildDesiredSize
+					)
+				);
+			}
+		}
+
+		virtual FChildren* GetChildren() override
+		{
+			return &ChildSlot;
+		}
+
+		virtual FVector2D ComputeDesiredSize(float LayoutScaleMultiplier) const override
+		{
+			const FVector2D ChildDesiredSize = ChildSlot.GetWidget()->GetDesiredSize();
+			return FVector2D(OffsetX + ChildDesiredSize.X, OffsetY + ChildDesiredSize.Y);
+		}
+
+	private:
+		float OffsetX;
+		float OffsetY;
+		FSingleWidgetChildrenWithBasicLayoutSlot ChildSlot;
+	};
+
+	/**
+	 * Test case parameters for cursor lock rounding tests.
+	 * Uses fractional offset within window to achieve desired geometry bounds.
+	 */
+	struct FCursorLockTestCase
+	{
+		// Fractional offset of widget within window
+		float OffsetX;
+		float OffsetY;
+		// Widget dimensions (can be fractional to achieve fractional right/bottom bounds)
+		float WidgetWidth;
+		float WidgetHeight;
+		// Expected clip rect after rounding (offset by window origin)
+		int32 ExpectedLeft;
+		int32 ExpectedTop;
+		int32 ExpectedRight;
+		int32 ExpectedBottom;
+		const ANSICHAR* Description;
+	};
+
+	// Fixed window position for all tests (integral coordinates)
+	constexpr float WindowOriginX = 100.0f;
+	constexpr float WindowOriginY = 100.0f;
+
+	/**
+	 * Function to test cursor lock behavior for a specific geometry and window configuration.
+	 * Returns true if the test passed, false otherwise.
+	 */
+	bool TestCursorLockRounding(
+		FAutomationTestBase* Test,
+		EWindowType WindowType,
+		bool bUseOSWindowBorder,
+		const FCursorLockTestCase& TestCase,
+		const ANSICHAR* WindowDescription)
+	{
+		if (!FSlateApplication::IsInitialized())
+		{
+			Test->AddError(TEXT("FSlateApplication is not initialized"));
+			return false;
+		}
+
+		// Window client size needs to accommodate the offset plus widget size
+		const float ClientWidth = TestCase.OffsetX + TestCase.WidgetWidth;
+		const float ClientHeight = TestCase.OffsetY + TestCase.WidgetHeight;
+
+		// Create the test widget with the specified size
+		TSharedRef<STestWidgetWithKnownSize> TestWidget = SNew(STestWidgetWithKnownSize)
+			.DesiredWidth(TestCase.WidgetWidth)
+			.DesiredHeight(TestCase.WidgetHeight);
+
+		// Wrap in offset container to achieve fractional positioning
+		TSharedRef<SFractionalOffsetBox> OffsetBox = SNew(SFractionalOffsetBox)
+			.OffsetX(TestCase.OffsetX)
+			.OffsetY(TestCase.OffsetY)
+			[
+				TestWidget
+			];
+
+		// Create the test window at integral coordinates
+		TSharedRef<SWindow> TestWindow = SNew(SWindow)
+			.Type(WindowType)
+			.UseOSWindowBorder(bUseOSWindowBorder)
+			.ClientSize(FVector2f(FMath::CeilToFloat(ClientWidth), FMath::CeilToFloat(ClientHeight)))
+			.ScreenPosition(FVector2f(WindowOriginX, WindowOriginY))
+			.AutoCenter(EAutoCenter::None)
+			.CreateTitleBar(false)
+			.SizingRule(ESizingRule::FixedSize)
+			.SupportsMaximize(false)
+			.SupportsMinimize(false)
+			.HasCloseButton(false)
+			[
+				OffsetBox
+			];
+		TSharedRef<FMockForegroundWindow> MockWindow = MakeShared<FMockForegroundWindow>();
+		TestWindow->SetNativeWindow(MockWindow);
+
+		FSlateApplication::Get().ForceRedrawWindow(TestWindow);
+
+		TSharedRef<FMockCursorForLockTest> MockCursor = MakeShared<FMockCursorForLockTest>();
+		TSharedPtr<FSlateUser> CursorUser = FSlateApplication::Get().GetCursorUser();
+		if (!CursorUser.IsValid())
+		{
+			Test->AddError(FString::Printf(TEXT("%S %S: Could not get cursor user"), WindowDescription, TestCase.Description));
+			FSlateApplication::Get().DestroyWindowImmediately(TestWindow);
+			return false;
+		}
+
+		// Until we have a way of constructing SWindow from a custom native window AND registering it to Slate Application's window list
+		bool bWasFastWidgetPathEnabled = GSlateFastWidgetPath;
+		GSlateFastWidgetPath = true;
+
+		// Store to restore later
+		TSharedPtr<ICursor> OriginalCursor = CursorUser->GetCursor();
+
+		CursorUser->OverrideCursor(MockCursor);
+		CursorUser->LockCursor(TestWidget);
+
+		bool bTestPassed = true;
+
+		if (!MockCursor->bLockWasCalled)
+		{
+			Test->AddError(FString::Printf(TEXT("%S %S: Cursor Lock() was not called"), WindowDescription, TestCase.Description));
+			bTestPassed = false;
+		}
+		else if (!MockCursor->bHasCapturedBounds)
+		{
+			Test->AddError(FString::Printf(TEXT("%S %S: Cursor Lock() was called with null bounds"), WindowDescription, TestCase.Description));
+			bTestPassed = false;
+		}
+		else
+		{
+			const RECT& CapturedRect = MockCursor->CapturedBounds;
+
+			const FDeprecateVector2DResult ClientLocation = OffsetBox->GetCachedGeometry().GetAbsolutePosition();
+			const float InputLeft = ClientLocation.X + TestCase.OffsetX;
+			const float InputTop = ClientLocation.Y + TestCase.OffsetY;
+			const float InputRight = InputLeft + TestCase.WidgetWidth;
+			const float InputBottom = InputTop + TestCase.WidgetHeight;
+			int32 OffsetExpectedLeft = ClientLocation.X + TestCase.ExpectedLeft;
+			int32 OffsetExpectedTop = ClientLocation.Y + TestCase.ExpectedTop;
+			int32 OffsetExpectedRight = ClientLocation.X + TestCase.ExpectedRight;
+			int32 OffsetExpectedBottom = ClientLocation.Y + TestCase.ExpectedBottom;
+
+			// Verify the captured rect matches expectations
+			if (CapturedRect.left != OffsetExpectedLeft)
+			{
+				Test->AddError(FString::Printf(TEXT("%S %S: Clip rect left mismatch. Expected %i, got %i (input: %.1f)"),
+					WindowDescription, TestCase.Description, OffsetExpectedLeft, CapturedRect.left, InputLeft));
+				bTestPassed = false;
+			}
+			if (CapturedRect.top != OffsetExpectedTop)
+			{
+				Test->AddError(FString::Printf(TEXT("%S %S: Clip rect top mismatch. Expected %i, got %i (input: %.1f)"),
+					WindowDescription, TestCase.Description, OffsetExpectedTop, CapturedRect.top, InputTop));
+				bTestPassed = false;
+			}
+			if (CapturedRect.right != OffsetExpectedRight)
+			{
+				Test->AddError(FString::Printf(TEXT("%S %S: Clip rect right mismatch. Expected %i, got %i (input: %.1f)"),
+					WindowDescription, TestCase.Description, OffsetExpectedRight, CapturedRect.right, InputRight));
+				bTestPassed = false;
+			}
+			if (CapturedRect.bottom != OffsetExpectedBottom)
+			{
+				Test->AddError(FString::Printf(TEXT("%S %S: Clip rect bottom mismatch. Expected %i, got %i (input: %.1f)"),
+					WindowDescription, TestCase.Description, OffsetExpectedBottom, CapturedRect.bottom, InputBottom));
+				bTestPassed = false;
+			}
+		}
+
+		// Cleanup
+		CursorUser->UnlockCursor();
+		CursorUser->OverrideCursor(OriginalCursor);
+		GSlateFastWidgetPath = bWasFastWidgetPathEnabled;
+		FSlateApplication::Get().DestroyWindowImmediately(TestWindow);
+
+		return bTestPassed;
+	}
+
+	/**
+	 * Run all cursor lock rounding tests for a specific window configuration.
+	 */
+	bool RunCursorLockRoundingTestsForWindowConfig(
+		FAutomationTestBase* Test,
+		EWindowType WindowType,
+		bool bUseOSWindowBorder,
+		const ANSICHAR* WindowDescription)
+	{
+		// Test cases enumerated:
+		// - All 16 combinations of .4/.6 for left, top, right, bottom
+		// - 1 case with fully integral coordinates
+		//
+		// Expected rounding behavior (to stay inside the widget):
+		// - Left/Top (min coords): Round UP (ceiling) - 100.4 -> 101, 100.6 -> 101
+		// - Right/Bottom (max coords): Round DOWN (floor/truncate) - 200.4 -> 200, 200.6 -> 200
+		//
+		// Test case format: { OffsetX, OffsetY, Width, Height, ExpectedLeft, ExpectedTop, ExpectedRight, ExpectedBottom, Description }
+		// Window is at (100, 100), so geometry bounds are:
+		//   Left = 100 + OffsetX, Top = 100 + OffsetY
+		//   Right = Left + Width, Bottom = Top + Height
+
+		const TArray<FCursorLockTestCase> TestCases = {
+			// Fully integral case - should remain unchanged
+			// Bounds: (100, 100, 200, 200)
+			{ 0.0f, 0.0f, 100.0f, 100.0f, 0, 0, 100, 100, "Integral bounds" },
+
+			// Left=.4, Top=.4 (OffsetX=0.4, OffsetY=0.4)
+			{ 0.4f, 0.4f, 100.0f, 100.0f, 1, 1, 100, 100, "L.4 T.4 R.4 B.4" },
+			{ 0.4f, 0.4f, 100.0f, 100.2f, 1, 1, 100, 100, "L.4 T.4 R.4 B.6" },
+			{ 0.4f, 0.4f, 100.2f, 100.0f, 1, 1, 100, 100, "L.4 T.4 R.6 B.4" },
+			{ 0.4f, 0.4f, 100.2f, 100.2f, 1, 1, 100, 100, "L.4 T.4 R.6 B.6" },
+
+			// Left=.4, Top=.6 (OffsetX=0.4, OffsetY=0.6)
+			{ 0.4f, 0.6f,  99.8f,  99.8f, 1, 1, 100, 100, "L.4 T.6 R.4 B.4" },
+			{ 0.4f, 0.6f,  99.8f, 100.0f, 1, 1, 100, 100, "L.4 T.6 R.4 B.6" },
+			{ 0.4f, 0.6f, 100.0f,  99.8f, 1, 1, 100, 100, "L.4 T.6 R.6 B.4" },
+			{ 0.4f, 0.6f, 100.0f, 100.0f, 1, 1, 100, 100, "L.4 T.6 R.6 B.6" },
+
+			// Left=.6, Top=.4 (OffsetX=0.6, OffsetY=0.4)
+			{ 0.6f, 0.4f,  99.8f,  99.8f, 1, 1, 100, 100, "L.6 T.4 R.4 B.4" },
+			{ 0.6f, 0.4f,  99.8f, 100.0f, 1, 1, 100, 100, "L.6 T.4 R.4 B.6" },
+			{ 0.6f, 0.4f, 100.0f,  99.8f, 1, 1, 100, 100, "L.6 T.4 R.6 B.4" },
+			{ 0.6f, 0.4f, 100.0f, 100.0f, 1, 1, 100, 100, "L.6 T.4 R.6 B.6" },
+
+			// Left=.6, Top=.6 (OffsetX=0.6, OffsetY=0.6)
+			{ 0.6f, 0.6f,  99.8f,  99.8f, 1, 1, 100, 100, "L.6 T.6 R.4 B.4" },
+			{ 0.6f, 0.6f,  99.8f, 100.0f, 1, 1, 100, 100, "L.6 T.6 R.4 B.6" },
+			{ 0.6f, 0.6f, 100.0f,  99.8f, 1, 1, 100, 100, "L.6 T.6 R.6 B.4" },
+			{ 0.6f, 0.6f, 100.0f, 100.0f, 1, 1, 100, 100, "L.6 T.6 R.6 B.6" },
+		};
+
+		bool bAllPassed = true;
+		for (const FCursorLockTestCase& TestCase : TestCases)
+		{
+			if (!TestCursorLockRounding(Test, WindowType, bUseOSWindowBorder, TestCase, WindowDescription))
+			{
+				bAllPassed = false;
+			}
+		}
+
+		return bAllPassed;
+	}
+	
+	constexpr EAutomationTestFlags SlateUserCursorLockTestFlags = EAutomationTestFlags::EditorContext | EAutomationTestFlags::ClientContext |
+		EAutomationTestFlags::ProgramContext | EAutomationTestFlags::LowPriority | EAutomationTestFlags::EngineFilter;
+
+}
+
+namespace UE::Slate::Test
+{
+	IMPLEMENT_SIMPLE_AUTOMATION_TEST(FSlateUserCursorLockBorderlessGameWindowTest, "Slate.User.CursorLock.BorderlessGameWindow", SlateUserCursorLockTestFlags)
+
+	bool FSlateUserCursorLockBorderlessGameWindowTest::RunTest(const FString& Parameters)
+	{
+		return RunCursorLockRoundingTestsForWindowConfig(
+			this,
+			EWindowType::GameWindow,
+			/*bUseOSWindowBorder*/ false,
+			"Borderless GameWindow");
+	}
+
+	IMPLEMENT_SIMPLE_AUTOMATION_TEST(FSlateUserCursorLockBorderlessNormalWindowTest, "Slate.User.CursorLock.BorderlessNormalWindow", SlateUserCursorLockTestFlags)
+
+	bool FSlateUserCursorLockBorderlessNormalWindowTest::RunTest(const FString& Parameters)
+	{
+		return RunCursorLockRoundingTestsForWindowConfig(
+			this,
+			EWindowType::Normal,
+			/*bUseOSWindowBorder*/ false,
+			"Borderless Normal");
+	}
+
+	IMPLEMENT_SIMPLE_AUTOMATION_TEST(FSlateUserCursorLockBorderedGameWindowTest, "Slate.User.CursorLock.BorderedGameWindow", SlateUserCursorLockTestFlags)
+
+	bool FSlateUserCursorLockBorderedGameWindowTest::RunTest(const FString& Parameters)
+	{
+		return RunCursorLockRoundingTestsForWindowConfig(
+			this,
+			EWindowType::GameWindow,
+			/*bUseOSWindowBorder*/ true,
+			"Bordered GameWindow");
+	}
+
+	IMPLEMENT_SIMPLE_AUTOMATION_TEST(FSlateUserCursorLockBorderedNormalWindowTest, "Slate.User.CursorLock.BorderedNormalWindow", SlateUserCursorLockTestFlags)
+
+	bool FSlateUserCursorLockBorderedNormalWindowTest::RunTest(const FString& Parameters)
+	{
+		return RunCursorLockRoundingTestsForWindowConfig(
+			this,
+			EWindowType::Normal,
+			/*bUseOSWindowBorder*/ true,
+			"Bordered Normal");
+	}
+} // namespace UE::Slate::Test
+
+#endif // WITH_AUTOMATION_WORKER && PLATFORM_DESKTOP
diff --git a/Engine/Source/Runtime/SlateCore/Private/Widgets/SWindow.cpp b/Engine/Source/Runtime/SlateCore/Private/Widgets/SWindow.cpp
--- a/Engine/Source/Runtime/SlateCore/Private/Widgets/SWindow.cpp
+++ b/Engine/Source/Runtime/SlateCore/Private/Widgets/SWindow.cpp
@@ -1925,26 +1925,22 @@ bool PointWithinSlateRect(const FVector2f& Point, const FSlateRect& Rect)
 
 EWindowZone::Type SWindow::GetCurrentWindowZone(UE::Slate::FDeprecateVector2DParameter LocalMousePosition)
 {
-	const bool bIsFullscreenMode = GetWindowMode() == EWindowMode::WindowedFullscreen || GetWindowMode() == EWindowMode::Fullscreen;
-	const bool bIsBorderlessGameWindow = Type == EWindowType::GameWindow && !bHasOSWindowBorder;
-
 	const float WindowDPIScale = FSlateApplicationBase::Get().GetApplicationScale() * (NativeWindow.IsValid() ? GetDPIScaleFactor() : 1.0f);
 
 	const FMargin DPIScaledResizeBorder = UserResizeBorder * WindowDPIScale;
 
-	const bool bIsCursorVisible = FSlateApplicationBase::Get().GetPlatformCursor()->GetType() != EMouseCursor::None;
-
-	// Don't allow position/resizing of window while in fullscreen mode by ignoring Title Bar/Border Zones
-	if ( (bIsFullscreenMode && !bIsBorderlessGameWindow) || !bIsCursorVisible )
+	if (FSlateApplicationBase::Get().GetPlatformCursor()->GetType() == EMouseCursor::None)
 	{
 		return EWindowZone::ClientArea;
 	}
-	else if(LocalMousePosition.X >= 0 && LocalMousePosition.X < Size.X &&
-			LocalMousePosition.Y >= 0 && LocalMousePosition.Y < Size.Y)
+	if (LocalMousePosition.X >= 0 && LocalMousePosition.X < Size.X && LocalMousePosition.Y >= 0 && LocalMousePosition.Y < Size.Y)
 	{
 		int32 Row = 1;
 		int32 Col = 1;
-		if (SizingRule == ESizingRule::UserSized && !bIsFullscreenMode && !NativeWindow->IsMaximized())
+		EWindowMode::Type Mode = GetWindowMode();
+		const bool bFullscreen = Mode == EWindowMode::Fullscreen || Mode == EWindowMode::WindowedFullscreen;
+		const bool bResizable = SizingRule == ESizingRule::UserSized && !bFullscreen && !NativeWindow->IsMaximized();
+		if (bResizable)
 		{
 			if (LocalMousePosition.X < (DPIScaledResizeBorder.Left + 5))
 			{
@@ -1989,11 +1985,13 @@ EWindowZone::Type SWindow::GetCurrentWindowZone(UE::Slate::FDeprecateVector2DPar
 		if (InZone == EWindowZone::ClientArea)
 		{
 			// Hittest to see if the widget under the mouse should be treated as a title bar (i.e. should move the window)
-			FWidgetPath HitTestResults = FSlateApplicationBase::Get().GetHitTesting().LocateWidgetInWindow(FSlateApplicationBase::Get().GetCursorPos(), SharedThis(this), false, INDEX_NONE);
+			constexpr bool bIgnoreEnabledStatus = false;
+			FWidgetPath HitTestResults = FSlateApplicationBase::Get().GetHitTesting()
+				.LocateWidgetInWindow(GetPositionInScreen() + LocalMousePosition, SharedThis(this), bIgnoreEnabledStatus, INDEX_NONE);
 			if( HitTestResults.Widgets.Num() > 0 )
 			{
 				const EWindowZone::Type ZoneOverride = HitTestResults.Widgets.Last().Widget->GetWindowZoneOverride();
-				if( ZoneOverride != EWindowZone::Unspecified )
+				if (ZoneOverride != EWindowZone::Unspecified && (bResizable || !EWindowZone::IsBorder(ZoneOverride)))
 				{
 					// The widget overrode the window zone
 					InZone = ZoneOverride;
