Implemented `Slate.NavigationDebugging` to display where Slate thinks the next `EUINavigation` would end up at

#Synced-CL 51318692
#Synced-CL 51318692

diff --git a/Engine/Source/Runtime/Engine/Private/Slate/SGameLayerManager.cpp b/Engine/Source/Runtime/Engine/Private/Slate/SGameLayerManager.cpp
--- a/Engine/Source/Runtime/Engine/Private/Slate/SGameLayerManager.cpp
+++ b/Engine/Source/Runtime/Engine/Private/Slate/SGameLayerManager.cpp
@@ -21,6 +21,10 @@
 #include "Types/InvisibleToWidgetReflectorMetaData.h"
 #include "Framework/Application/SlateApplication.h"
 
+#if WITH_SLATE_DEBUGGING
+#include "Navigation/SlateNavigationDebugger.h"
+#endif //WITH_SLATE_DEBUGGING
+
 #include UE_INLINE_GENERATED_CPP_BY_NAME(SGameLayerManager)
 
 #ifndef UE_SLATE_WITH_GAMELAYER_CANVAS_VISIBILITY_COMMANDS
@@ -433,7 +437,12 @@ int32 SGameLayerManager::OnPaint(const FPaintArgs& Args, const FGeometry& Allott
 #endif
 
 	OutDrawElements.SetIsInGameLayer(true);
-	const int32 ResultLayer = SCompoundWidget::OnPaint(Args, AllottedGeometry, MyCullingRect, OutDrawElements, LayerId, InWidgetStyle, bParentEnabled);
+	int32 ResultLayer = SCompoundWidget::OnPaint(Args, AllottedGeometry, MyCullingRect, OutDrawElements, LayerId, InWidgetStyle, bParentEnabled);
+
+#if WITH_SLATE_DEBUGGING
+	ResultLayer = FSlateNavigationDebugger::Get().Paint(OutDrawElements, ResultLayer);
+#endif
+
 	OutDrawElements.SetIsInGameLayer(false);
 	return ResultLayer;
 }
diff --git a/Engine/Source/Runtime/Slate/Private/Framework/Application/SlateApplication.cpp b/Engine/Source/Runtime/Slate/Private/Framework/Application/SlateApplication.cpp
--- a/Engine/Source/Runtime/Slate/Private/Framework/Application/SlateApplication.cpp
+++ b/Engine/Source/Runtime/Slate/Private/Framework/Application/SlateApplication.cpp
@@ -3347,7 +3347,7 @@ void FSlateApplication::ProcessReply( const FWidgetPath& CurrentEventPath, const
 	uint32 PointerIndex = InMouseEvent != nullptr ? InMouseEvent->GetPointerIndex() : CursorPointerIndex;
 
 	TSharedRef<FSlateUser> SlateUser = GetOrCreateUser(UserIndex);
-	if (SlateUser->HasCapture(PointerIndex) && (TheReply.ShouldReleaseMouse() || bStartingDragDrop))
+ 	if (SlateUser->HasCapture(PointerIndex) && (TheReply.ShouldReleaseMouse() || bStartingDragDrop))
 	{
 		SlateUser->ReleaseCapture(PointerIndex);
 	}
@@ -6410,31 +6410,25 @@ EUINavigation FSlateApplication::NavigateFromWidget(const uint32 InUserIndex, TS
 	return EUINavigation::Invalid;
 }
 
-bool FSlateApplication::AttemptNavigation(const FWidgetPath& NavigationSource, const FNavigationEvent& NavigationEvent, const FNavigationReply& NavigationReply, const FArrangedWidget& BoundaryWidget)
+FSlateApplication::FNavigationResult FSlateApplication::CalculateDestinationWidget(const FWidgetPath& NavigationSource, const FNavigationReply& NavigationReply, EUINavigation NavigationType, uint32 UserIndex, const FArrangedWidget& BoundaryWidget) const
 {
-	if ( !NavigationSource.IsValid() )
+	if (!NavigationSource.IsValid())
 	{
-		return false;
+		return FNavigationResult();
 	}
 
-	TSharedPtr<SWidget> DestinationWidget = TSharedPtr<SWidget>();
-	bool bAlwaysHandleNavigationAttempt = false;
+	FNavigationResult Result;
 
-#if WITH_SLATE_DEBUGGING
-	ESlateDebuggingNavigationMethod NavigationMethod = ESlateDebuggingNavigationMethod::Unknown;
-#endif
-
-	EUINavigation NavigationType = NavigationEvent.GetNavigationType();
 	if ( NavigationReply.GetBoundaryRule() == EUINavigationRule::Explicit )
 	{
 		const SWidget* FocusRecipient = NavigationReply.GetFocusRecipient().Get();
 		if ( FocusRecipient && FocusRecipient->IsEnabled() && FocusRecipient->SupportsKeyboardFocus() )
 		{
-			DestinationWidget = NavigationReply.GetFocusRecipient();
-			bAlwaysHandleNavigationAttempt = true;
+			Result.DestinationWidget = NavigationReply.GetFocusRecipient();
+			Result.bAlwaysHandleNavigationAttempt = true;
 
 #if WITH_SLATE_DEBUGGING
-			NavigationMethod = ESlateDebuggingNavigationMethod::Explicit;
+			Result.NavigationMethod = ESlateDebuggingNavigationMethod::Explicit;
 #endif
 		}
 #if WITH_SLATE_DEBUGGING
@@ -6467,17 +6461,17 @@ bool FSlateApplication::AttemptNavigation(const FWidgetPath& NavigationSource, c
 			// Switch worlds for widgets in the current path 
 			FScopedSwitchWorldHack SwitchWorld(NavigationSource);
 
-			DestinationWidget = FocusDelegate.Execute(NavigationType);
-			bAlwaysHandleNavigationAttempt = true;
+			Result.DestinationWidget = FocusDelegate.Execute(NavigationType);
+			Result.bAlwaysHandleNavigationAttempt = true;
 
 #if WITH_SLATE_DEBUGGING
-			NavigationMethod = ESlateDebuggingNavigationMethod::CustomDelegateBound;
+			Result.NavigationMethod = ESlateDebuggingNavigationMethod::CustomDelegateBound;
 #endif
 		}
 		else
 		{
 #if WITH_SLATE_DEBUGGING
-			NavigationMethod = ESlateDebuggingNavigationMethod::CustomDelegateUnbound;
+			Result.NavigationMethod = ESlateDebuggingNavigationMethod::CustomDelegateUnbound;
 #endif
 		}
 	}
@@ -6490,12 +6484,14 @@ bool FSlateApplication::AttemptNavigation(const FWidgetPath& NavigationSource, c
 			FWeakWidgetPath WeakNavigationSource(NavigationSource);
 			FWidgetPath NewFocusedWidgetPath = WeakNavigationSource.ToNextFocusedPath(NavigationType, NavigationReply, BoundaryWidget);
 
+			checkf(NewFocusedWidgetPath.IsValid(), TEXT("This was assuming the `NewFocusedWidgetPath` is valid already"));
+
 			// Resolve the Widget Path
 			FArrangedWidget& NewFocusedArrangedWidget = NewFocusedWidgetPath.Widgets.Last();
-			DestinationWidget = NewFocusedArrangedWidget.Widget;
+			Result.DestinationWidget = NewFocusedArrangedWidget.Widget;
 
 #if WITH_SLATE_DEBUGGING
-			NavigationMethod = ESlateDebuggingNavigationMethod::NextOrPrevious;
+			Result.NavigationMethod = ESlateDebuggingNavigationMethod::NextOrPrevious;
 #endif
 		}
 		else
@@ -6506,21 +6502,35 @@ bool FSlateApplication::AttemptNavigation(const FWidgetPath& NavigationSource, c
 			// Switch worlds for widgets in the current path 
 			FScopedSwitchWorldHack SwitchWorld(NavigationSource);
 
-			DestinationWidget = NavigationSource.GetDeepestWindow()->GetHittestGrid().FindNextFocusableWidget(FocusedArrangedWidget, NavigationType, NavigationReply, BoundaryWidget, NavigationEvent.GetUserIndex());
+			Result.DestinationWidget = NavigationSource.GetDeepestWindow()->GetHittestGrid().FindNextFocusableWidget(FocusedArrangedWidget, NavigationType, NavigationReply, BoundaryWidget, UserIndex);
 
 #if WITH_SLATE_DEBUGGING
-			NavigationMethod = ESlateDebuggingNavigationMethod::HitTestGrid;
+			Result.NavigationMethod = ESlateDebuggingNavigationMethod::HitTestGrid;
 #endif
 		}
 	}
 
+	return Result;
+}
+
+bool FSlateApplication::AttemptNavigation(const FWidgetPath& NavigationSource, const FNavigationEvent& NavigationEvent, const FNavigationReply& NavigationReply, const FArrangedWidget& BoundaryWidget)
+{
+	if ( !NavigationSource.IsValid() )
+	{
+		return false;
+	}
+
+	const EUINavigation NavigationType = NavigationEvent.GetNavigationType();
+	const uint32 UserIndex = NavigationEvent.GetUserIndex();
+	const FNavigationResult DestinationResult = CalculateDestinationWidget(NavigationSource, NavigationReply, NavigationType, UserIndex, BoundaryWidget);
+
 #if WITH_SLATE_DEBUGGING
-	FSlateDebugging::BroadcastAttemptNavigation(NavigationEvent, NavigationReply, NavigationSource, DestinationWidget, NavigationMethod);
+	FSlateDebugging::BroadcastAttemptNavigation(NavigationEvent, NavigationReply, NavigationSource, DestinationResult.DestinationWidget, DestinationResult.NavigationMethod);
 #endif
 
-	FScopedNavigationTransition ScopedNavigationTransition(NavigationType, GetUser(NavigationEvent.GetUserIndex()));
+	FScopedNavigationTransition ScopedNavigationTransition(NavigationType, GetUser(UserIndex));
 
-	return ExecuteNavigation(NavigationSource, DestinationWidget, NavigationEvent.GetUserIndex(), bAlwaysHandleNavigationAttempt);
+	return ExecuteNavigation(NavigationSource, DestinationResult.DestinationWidget, NavigationEvent.GetUserIndex(), DestinationResult.bAlwaysHandleNavigationAttempt);
 }
 
 bool FSlateApplication::ExecuteNavigation(const FWidgetPath& NavigationSource, TSharedPtr<SWidget> DestinationWidget, const uint32 UserIndex, bool bAlwaysHandleNavigationAttempt)
diff --git a/Engine/Source/Runtime/Slate/Private/Navigation/SlateNavigationDebugger.cpp b/Engine/Source/Runtime/Slate/Private/Navigation/SlateNavigationDebugger.cpp
new file mode 100644
--- /dev/null
+++ b/Engine/Source/Runtime/Slate/Private/Navigation/SlateNavigationDebugger.cpp
@@ -0,0 +1,614 @@
+// Copyright Epic Games, Inc. All Rights Reserved.
+
+#include "Navigation/SlateNavigationDebugger.h"
+
+#if WITH_SLATE_DEBUGGING
+
+#include "Misc/Paths.h"
+#include "Styling/SlateStyleRegistry.h"
+#include "Styling/SlateStyleMacros.h"
+#include "Framework/Application/SlateApplication.h"
+#include "Framework/Application/SlateUser.h"
+
+static TAutoConsoleVariable<bool> CVarDisplayCandidateWidgets(
+	TEXT("Slate.NavigationDebugging.DisplayCandidateWidgets"),
+	false,
+	TEXT("Enable to see candidate widgets"),
+	ECVF_Default);
+
+namespace ArrowGeometryHelpers
+{
+	constexpr bool bAntiAlias = true;
+	constexpr float LineThickness = 2.0f;
+	constexpr float DashLineLength = 4.0f;
+	constexpr float CandidateColorAlpha = 0.25f;
+	static const FVector2f ArrowBrushSize = FVector2f(30.0f, 30.0f);
+
+	static EUINavigation GetOpposingDirection(const EUINavigation InNavDirection)
+	{
+		if (InNavDirection == EUINavigation::Up)
+		{
+			return EUINavigation::Down;
+		}
+		else if (InNavDirection == EUINavigation::Down)
+		{
+			return EUINavigation::Up;
+		}
+		else if (InNavDirection == EUINavigation::Left)
+		{
+			return EUINavigation::Right;
+		}
+		else if (InNavDirection == EUINavigation::Right)
+		{
+			return EUINavigation::Left;
+		}
+
+		return InNavDirection;
+	}
+
+	static FVector2f GetArrowCoordinatesForLine(const EUINavigation InNavDirection)
+	{
+		FVector2f Coordinates = FVector2f::ZeroVector;
+
+		if (InNavDirection == EUINavigation::Up)
+		{
+			Coordinates = FVector2f(0.5f, 0.0f);
+		}
+		else if (InNavDirection == EUINavigation::Down)
+		{
+			Coordinates = FVector2f(0.5f, 1.0f);
+		}
+		else if (InNavDirection == EUINavigation::Left)
+		{
+			Coordinates = FVector2f(0.0f, 0.5f);
+		}
+		else if (InNavDirection == EUINavigation::Right)
+		{
+			Coordinates = FVector2f(1.0f, 0.5f);
+		}
+
+		return Coordinates;
+	}
+
+	static bool IsWidgetOutOfViewport(const FSlateUserNavigationData& InNavigationData, TSharedRef<const SWidget> InWidget)
+	{
+		check(InNavigationData.CanPaint());
+
+		TSharedRef<SWidget> ViewportWidget = InNavigationData.ViewportWidget.ToSharedRef();
+
+		const FGeometry ViewportGeometry = ViewportWidget->GetCachedGeometry().MakeChild(InNavigationData.WindowWidget->GetPositionInScreen() * -1);
+		const FSlateRect ViewportRect = FSlateRect::FromPointAndExtent(ViewportGeometry.GetAbsolutePosition(), ViewportGeometry.GetAbsoluteSize());
+		
+		const FGeometry WidgetGeometry = InWidget->GetCachedGeometry().MakeChild(InNavigationData.WindowWidget->GetPositionInScreen() * -1);
+		const FSlateRect WidgetRect = FSlateRect::FromPointAndExtent(WidgetGeometry.GetAbsolutePosition(), WidgetGeometry.GetAbsoluteSize());
+
+		return !FSlateRect::DoRectanglesIntersect(ViewportRect, WidgetRect);
+	}
+
+	static FGeometry GetArrowGeometryForWidget_OnScreen(const FSlateUserNavigationData& InNavigationData, const EUINavigation InNavDirection, const FGeometry& InWidgetGeometry)
+	{
+		check(InNavigationData.CanPaint());
+
+		const FVector2f WidgetSize = InWidgetGeometry.GetLocalSize();
+		const FVector2f DefaultArrowTextureSize = ArrowGeometryHelpers::ArrowBrushSize;
+
+		const bool bArrowsWiderThanWidget = WidgetSize.X - (4.0f * DefaultArrowTextureSize.X) <= 0.0f;
+		const bool bArrowsTallerThanWidget = WidgetSize.Y - (4.0f * DefaultArrowTextureSize.Y) <= 0.0f;
+		const bool bScaleDownArrows = bArrowsWiderThanWidget || bArrowsTallerThanWidget;
+		const float ScaledDownArrowSize = FMath::Min(WidgetSize.X, WidgetSize.Y) * 0.25f;
+
+		const FVector2f ArrowTextureSize = bScaleDownArrows ? FVector2f(ScaledDownArrowSize) : DefaultArrowTextureSize;
+
+		FVector2f ArrowPosition;
+		FVector2f AdjustmentForOutline;
+
+		if (InNavDirection == EUINavigation::Up)
+		{
+			ArrowPosition = FVector2f((WidgetSize.X - ArrowTextureSize.X) * 0.5, 0.0f);
+			AdjustmentForOutline = FVector2f(0.0f, ArrowGeometryHelpers::LineThickness);
+		}
+		else if (InNavDirection == EUINavigation::Down)
+		{
+			ArrowPosition = FVector2f((WidgetSize.X - ArrowTextureSize.X) * 0.5, WidgetSize.Y - ArrowTextureSize.Y);
+			AdjustmentForOutline = FVector2f(0.0f, -2.0f * ArrowGeometryHelpers::LineThickness);
+		}
+		else if (InNavDirection == EUINavigation::Left)
+		{
+			ArrowPosition = FVector2f(0.0f, (WidgetSize.Y - ArrowTextureSize.Y) * 0.5f);
+			AdjustmentForOutline = FVector2f(ArrowGeometryHelpers::LineThickness, 0.0f);
+		}
+		else if (InNavDirection == EUINavigation::Right)
+		{
+			ArrowPosition = FVector2f(WidgetSize.X - ArrowTextureSize.X, (WidgetSize.Y - ArrowTextureSize.Y) * 0.5f);
+			AdjustmentForOutline = FVector2f(-2.0f * ArrowGeometryHelpers::LineThickness, 0.0f);
+		}
+
+		const FGeometry ViewportWidgetGeometry = InNavigationData.ViewportWidget->GetCachedGeometry().MakeChild(InNavigationData.WindowWidget->GetPositionInScreen() * -1);
+		const FVector2f ArrowPositionInViewport = ViewportWidgetGeometry.AbsoluteToLocal(InWidgetGeometry.LocalToAbsolute(ArrowPosition)) + AdjustmentForOutline;
+
+		return ViewportWidgetGeometry.MakeChild(ArrowTextureSize, FSlateLayoutTransform(1.0f, ArrowPositionInViewport));
+	}
+
+	static FGeometry GetArrowGeometryForWidget_OffScreen(const FSlateUserNavigationData& InNavigationData, const EUINavigation InNavDirection, const FNavigationDebuggingStyle::EArrowType InArrowType, TSharedRef<const SWidget> InWidget)
+	{
+		check(InNavigationData.CanPaint());
+
+		TSharedRef<SWindow> WindowWidget = InNavigationData.WindowWidget.ToSharedRef();
+		TSharedPtr<SWidget> ViewportWidget = InNavigationData.ViewportWidget.ToSharedRef();
+
+		const FGeometry ViewportWidgetGeometry = ViewportWidget->GetCachedGeometry().MakeChild(InNavigationData.WindowWidget->GetPositionInScreen() * -1);
+		const FGeometry WidgetGeometry = InWidget->GetCachedGeometry().MakeChild(InNavigationData.WindowWidget->GetPositionInScreen() * -1);
+		const FVector2f WidgetPosition = WidgetGeometry.GetAbsolutePosition();
+
+		const bool bIsFocusArrow = InArrowType == FNavigationDebuggingStyle::EArrowType::Focus;
+		const EUINavigation ArrowDirection = bIsFocusArrow ? InNavDirection : GetOpposingDirection(InNavDirection);
+		// Clamping On Screen Widget Geometry Positions
+		const FGeometry OnScreenArrowGeometry = GetArrowGeometryForWidget_OnScreen(InNavigationData, ArrowDirection, WidgetGeometry);
+		const FVector2f ArrowCenter = OnScreenArrowGeometry.GetAbsolutePositionAtCoordinates(FVector2f(0.5f, 0.5f));
+
+		const FVector2f ViewportTopLeft = ViewportWidgetGeometry.GetAbsolutePositionAtCoordinates(FVector2f(0.0f, 0.0f));
+		const FVector2f ViewportExtents = ViewportWidgetGeometry.GetAbsolutePositionAtCoordinates(FVector2f(1.0f, 1.0f));
+		const FVector2f ClampedWidgetPosition = FVector2f::Clamp(ArrowCenter, FVector2f(0.0f, 0.0f), ViewportExtents);
+
+		const FVector2f ArrowTextureSize = ArrowGeometryHelpers::ArrowBrushSize;
+		const FVector2f HalfArrowSize = (ArrowTextureSize * 0.5f);
+
+		FVector2f ArrowPositionForWidget = ClampedWidgetPosition - HalfArrowSize;
+
+		const bool bIsBelowViewport = WidgetPosition.Y > ViewportExtents.Y;
+		if (bIsBelowViewport)
+		{
+			ArrowPositionForWidget.Y -= HalfArrowSize.Y;
+		}
+
+		const bool bIsAboveViewport = (WidgetPosition.Y + HalfArrowSize.Y) < ViewportTopLeft.Y;
+		if (bIsAboveViewport)
+		{
+			ArrowPositionForWidget.Y += HalfArrowSize.Y;
+		}
+
+		const bool bIsRightViewport = (WidgetPosition.X + HalfArrowSize.X) > ViewportExtents.X;
+		if (bIsRightViewport)
+		{
+			ArrowPositionForWidget.X -= HalfArrowSize.X;
+		}
+
+		const bool bIsLeftViewport = WidgetPosition.X < ViewportTopLeft.X;
+		if (bIsLeftViewport)
+		{
+			ArrowPositionForWidget.X += HalfArrowSize.X;
+		}
+
+		return ViewportWidgetGeometry.MakeChild(ArrowTextureSize, FSlateLayoutTransform(1.0f, ViewportWidgetGeometry.AbsoluteToLocal(ArrowPositionForWidget)));
+	}
+
+	static void DrawArrow(const FGeometry& InArrowGeometry, const FSlateBrush* InArrowBrush, float ColorAlpha, FSlateWindowElementList& OutDrawElements, int32 LayerId)
+	{
+		FSlateDrawElement::MakeBox(
+			OutDrawElements,
+			LayerId,
+			InArrowGeometry.ToPaintGeometry(),
+			InArrowBrush,
+			ESlateDrawEffect::None,
+			FLinearColor::White.CopyWithNewOpacity(ColorAlpha)
+		);
+	}
+}
+
+FNavigationDebuggingStyle::FNavigationDebuggingStyle()
+	: FSlateStyleSet(FName(TEXT("NavigationDebuggingStyle")))
+{
+	SetContentRoot(FPaths::EngineContentDir() / TEXT("Slate") / TEXT("Icons")/ TEXT("Navigation"));
+	SetCoreContentRoot(FPaths::EngineContentDir() / TEXT("Slate") / TEXT("Icons") / TEXT("Navigation"));
+
+	// Navigation Debugging Arrows
+	const FVector2f ArrowTextureSize = FVector2f(64.0f, 64.0f);
+	FSlateVectorImageBrush* Navigation_Focus_Left = new CORE_IMAGE_BRUSH_SVG("focus-arrow-left", ArrowTextureSize);
+	FSlateVectorImageBrush* Navigation_Focus_Right = new CORE_IMAGE_BRUSH_SVG("focus-arrow-right", ArrowTextureSize);
+	FSlateVectorImageBrush* Navigation_Focus_Up = new CORE_IMAGE_BRUSH_SVG("focus-arrow-up", ArrowTextureSize);
+	FSlateVectorImageBrush* Navigation_Focus_Down = new CORE_IMAGE_BRUSH_SVG("focus-arrow-down", ArrowTextureSize);
+
+	FSlateVectorImageBrush* Navigation_Candidate_Left = new CORE_IMAGE_BRUSH_SVG("candidate-arrow-left", ArrowTextureSize);
+	FSlateVectorImageBrush* Navigation_Candidate_Right = new CORE_IMAGE_BRUSH_SVG("candidate-arrow-right", ArrowTextureSize);
+	FSlateVectorImageBrush* Navigation_Candidate_Up = new CORE_IMAGE_BRUSH_SVG("candidate-arrow-up", ArrowTextureSize);
+	FSlateVectorImageBrush* Navigation_Candidate_Down = new CORE_IMAGE_BRUSH_SVG("candidate-arrow-down", ArrowTextureSize);
+
+	FSlateVectorImageBrush* Navigation_Target_Left = new CORE_IMAGE_BRUSH_SVG("target-arrow-left", ArrowTextureSize);
+	FSlateVectorImageBrush* Navigation_Target_Right = new CORE_IMAGE_BRUSH_SVG("target-arrow-right", ArrowTextureSize);
+	FSlateVectorImageBrush* Navigation_Target_Up = new CORE_IMAGE_BRUSH_SVG("target-arrow-up", ArrowTextureSize);
+	FSlateVectorImageBrush* Navigation_Target_Down = new CORE_IMAGE_BRUSH_SVG("target-arrow-down", ArrowTextureSize);
+
+	NavigationArrowBrushes = TMap<EUINavigation, TMap<EArrowType, const FSlateBrush*>>(
+		{
+			{ EUINavigation::Up, TMap<EArrowType, const FSlateBrush*>(
+				{
+					{ EArrowType::Focus, { Navigation_Focus_Up } },
+					{ EArrowType::Candidate, { Navigation_Candidate_Up } },
+					{ EArrowType::Target, { Navigation_Target_Up } },
+				})
+			},
+			{ EUINavigation::Down,  TMap<EArrowType, const FSlateBrush*>(
+				{
+					{ EArrowType::Focus, { Navigation_Focus_Down } },
+					{ EArrowType::Candidate, { Navigation_Candidate_Down } },
+					{ EArrowType::Target, { Navigation_Target_Down } },
+				})
+			},
+			{ EUINavigation::Left,  TMap<EArrowType, const FSlateBrush*>(
+				{
+					{ EArrowType::Focus, { Navigation_Focus_Left } },
+					{ EArrowType::Candidate, { Navigation_Candidate_Left } },
+					{ EArrowType::Target, { Navigation_Target_Left } },
+				})
+			},
+			{ EUINavigation::Right,  TMap<EArrowType, const FSlateBrush*>(
+				{
+					{ EArrowType::Focus, { Navigation_Focus_Right } },
+					{ EArrowType::Candidate, { Navigation_Candidate_Right } },
+					{ EArrowType::Target, { Navigation_Target_Right } },
+				})
+			}
+		});
+
+	Set("Debug.Navigation.Focus.Left", Navigation_Focus_Left);
+	Set("Debug.Navigation.Focus.Right", Navigation_Focus_Right);
+	Set("Debug.Navigation.Focus.Up", Navigation_Focus_Up);
+	Set("Debug.Navigation.Focus.Down", Navigation_Focus_Down);
+
+	Set("Debug.Navigation.Candidate.Left", Navigation_Candidate_Left);
+	Set("Debug.Navigation.Candidate.Right", Navigation_Candidate_Right);
+	Set("Debug.Navigation.Candidate.Up", Navigation_Candidate_Up);
+	Set("Debug.Navigation.Candidate.Down", Navigation_Candidate_Down);
+
+	Set("Debug.Navigation.Target.Left", Navigation_Target_Left);
+	Set("Debug.Navigation.Target.Right", Navigation_Target_Right);
+	Set("Debug.Navigation.Target.Up", Navigation_Target_Up);
+	Set("Debug.Navigation.Target.Down", Navigation_Target_Down);
+
+	FSlateStyleRegistry::RegisterSlateStyle(*this);
+}
+
+FNavigationDebuggingStyle::~FNavigationDebuggingStyle()
+{
+	FSlateStyleRegistry::UnRegisterSlateStyle(*this);
+}
+
+FLinearColor FNavigationDebuggingStyle::GetFocusedWidgetColor() const
+{
+	return FLinearColor::Green;
+}
+
+FLinearColor FNavigationDebuggingStyle::GetTargetWidgetColor() const
+{
+	return FLinearColor::Red;
+}
+
+const FSlateBrush* FNavigationDebuggingStyle::GetArrowBrush(const EUINavigation InNavigation, const EArrowType InArrow) const
+{
+	const TMap<EArrowType, const FSlateBrush*>* FoundKeyOne = NavigationArrowBrushes.Find(InNavigation);
+	const FSlateBrush* const * FoundKeyTwo = FoundKeyOne ? FoundKeyOne->Find(InArrow) : nullptr;
+	return FoundKeyTwo ? *FoundKeyTwo : nullptr;
+}
+
+/*static*/ FNavigationDebuggingStyle& FNavigationDebuggingStyle::Get()
+{
+	static FNavigationDebuggingStyle Instance;
+	return Instance;
+}
+
+/*static*/ FSlateNavigationDebugger& FSlateNavigationDebugger::Get()
+{
+	static FSlateNavigationDebugger Instance;
+	return Instance;
+}
+
+bool FSlateUserNavigationData::CanPaint() const
+{ 
+	return FocusedWidgetPath->IsValid() && ViewportWidget.IsValid() && FocusedWidget.IsValid() && WindowWidget.IsValid();
+}
+
+FSlateUserNavigationData::FSlateUserNavigationData(FSlateUser& InSlateUser)
+	: FocusedWidgetPath(InSlateUser.GetFocusPath())
+{
+	if (FocusedWidgetPath->IsValid())
+	{
+		WindowWidget = FocusedWidgetPath->GetWindow();
+		FocusedWidget = FocusedWidgetPath->GetLastWidget();
+
+		TSharedPtr<ISlateViewport> ViewportInterface = WindowWidget.IsValid() ? WindowWidget->GetViewport() : nullptr;
+		TWeakPtr<SWidget> WeakViewportWidget = ViewportInterface.IsValid() ? ViewportInterface->GetWidget() : nullptr;
+		ViewportWidget = WeakViewportWidget.IsValid() ? WeakViewportWidget.Pin() : nullptr;
+	}
+
+	if (CanPaint())
+	{
+		WindowGeometry = WindowWidget->GetPaintSpaceGeometry();
+	}
+}
+
+void FSlateNavigationDebugger::SetListenToHitTestGrid(const bool bBind, TArray<FHittestGrid::FDebuggingFindNextFocusableWidgetArgs>& InContainer, FDelegateHandle& InHandle) const
+{
+	if (bBind)
+	{
+		InHandle = FHittestGrid::OnFindNextFocusableWidgetExecuted.AddLambda(
+			[&InContainer](const FHittestGrid* /*HittestGrid*/, const FHittestGrid::FDebuggingFindNextFocusableWidgetArgs& Info)
+			{
+				InContainer.Add(Info);
+			});
+	}
+	else
+	{
+		FHittestGrid::OnFindNextFocusableWidgetExecuted.Remove(InHandle);
+	}
+}
+
+TMap<EUINavigation, TSharedRef<SWidget>> FSlateNavigationDebugger::SimulateNavigationDirection(const FSlateUserNavigationData& InNavigationData, const int32 SlateUserIndex) const
+{
+	check(InNavigationData.CanPaint());
+
+	TMap<EUINavigation, TSharedRef<SWidget>> ReturnValue;
+
+	const TSet<EUINavigation> DirectionsToDebug = TSet<EUINavigation>
+		({
+			EUINavigation::Left,
+			EUINavigation::Right,
+			EUINavigation::Up,
+			EUINavigation::Down,
+		});
+
+	for (const EUINavigation& NavigationDirection : DirectionsToDebug)
+	{
+		for (int32 WidgetIndex = InNavigationData.FocusedWidgetPath->Widgets.Num() - 1; WidgetIndex >= 0; --WidgetIndex)
+		{
+			FArrangedWidget& SomeWidgetGettingEvent = InNavigationData.FocusedWidgetPath->Widgets[WidgetIndex];
+			if (SomeWidgetGettingEvent.Widget->IsEnabled())
+			{
+				const FNavigationReply FocusedNavigationReply = InNavigationData.FocusedWidget->OnNavigation(SomeWidgetGettingEvent.Geometry, NavigationDirection);
+
+				if (FocusedNavigationReply.GetBoundaryRule() != EUINavigationRule::Escape || SomeWidgetGettingEvent.Widget == InNavigationData.WindowWidget || WidgetIndex == 0)
+				{
+					FSlateApplication::FNavigationResult NavigationResult = FSlateApplication::Get().CalculateDestinationWidget(
+						InNavigationData.FocusedWidgetPath.Get(),
+						FocusedNavigationReply,
+						NavigationDirection,
+						SlateUserIndex,
+						SomeWidgetGettingEvent);
+
+					const bool bFoundKeyboardFocusableWidget = NavigationResult.DestinationWidget.IsValid() && NavigationResult.DestinationWidget->SupportsKeyboardFocus();
+					if (bFoundKeyboardFocusableWidget)
+					{
+						bool bIsViewportWidgetChild = false;
+						TSharedRef<SWidget> ViewportWidgetRef = InNavigationData.ViewportWidget.ToSharedRef();
+						TSharedPtr<SWidget> WidgetTreeTraverse = NavigationResult.DestinationWidget.ToSharedRef();
+
+						while (WidgetTreeTraverse.IsValid() && !bIsViewportWidgetChild)
+						{
+							bIsViewportWidgetChild = WidgetTreeTraverse.IsValid() && WidgetTreeTraverse.ToSharedRef() == ViewportWidgetRef;
+							TSharedPtr<SWidget> Parent = WidgetTreeTraverse.IsValid() ? WidgetTreeTraverse->GetParentWidget() : nullptr;
+							WidgetTreeTraverse = Parent;
+						}
+
+						if (bIsViewportWidgetChild)
+						{
+							ReturnValue.Emplace(NavigationDirection, NavigationResult.DestinationWidget.ToSharedRef());
+						}
+					}
+				}
+			}
+		}
+	}
+
+	return ReturnValue;
+}
+
+FGeometry FSlateNavigationDebugger::PaintOnScreenWidgetArrow(const FSlateUserNavigationData& InNavigationData, const EUINavigation InNavDirection, const FGeometry& InWidgetGeometry, const FNavigationDebuggingStyle::EArrowType InArrowType, FSlateWindowElementList& OutDrawElements, int32 LayerId) const
+{
+	const bool bIsFocusedWidget = InArrowType == FNavigationDebuggingStyle::EArrowType::Focus;
+	const EUINavigation DirectionForArrow = bIsFocusedWidget ? InNavDirection : ArrowGeometryHelpers::GetOpposingDirection(InNavDirection);
+	const FGeometry OnScreenWidgetArrowGeometry = ArrowGeometryHelpers::GetArrowGeometryForWidget_OnScreen(InNavigationData, DirectionForArrow, InWidgetGeometry);
+
+	const FSlateBrush* BrushToUse = FNavigationDebuggingStyle::Get().GetArrowBrush(InNavDirection, InArrowType);
+	if (ensure(BrushToUse))
+	{
+		const float ColorAlpha = InArrowType == FNavigationDebuggingStyle::EArrowType::Candidate ? ArrowGeometryHelpers::CandidateColorAlpha : 1.0f;
+		ArrowGeometryHelpers::DrawArrow(OnScreenWidgetArrowGeometry, BrushToUse, ColorAlpha, OutDrawElements, LayerId);
+	}
+
+	return OnScreenWidgetArrowGeometry;
+}
+
+FGeometry FSlateNavigationDebugger::PaintOffScreenWidgetArrow(const FSlateUserNavigationData& InNavigationData, const EUINavigation InNavDirection, TSharedRef<const SWidget> InWidget, const FNavigationDebuggingStyle::EArrowType InArrowType, FSlateWindowElementList& OutDrawElements, int32 LayerId) const
+{
+	const FGeometry OffscreenArrowIndicator = ArrowGeometryHelpers::GetArrowGeometryForWidget_OffScreen(InNavigationData, InNavDirection, InArrowType, InWidget);
+
+	const FSlateBrush* BrushToUse = FNavigationDebuggingStyle::Get().GetArrowBrush(ArrowGeometryHelpers::GetOpposingDirection(InNavDirection), InArrowType);
+	if (ensure(BrushToUse))
+	{
+		const float ColorAlpha = InArrowType == FNavigationDebuggingStyle::EArrowType::Candidate ? ArrowGeometryHelpers::CandidateColorAlpha : 1.0f;
+		ArrowGeometryHelpers::DrawArrow(OffscreenArrowIndicator, BrushToUse, ColorAlpha, OutDrawElements, LayerId);
+	}
+
+	return OffscreenArrowIndicator;
+}
+
+void FSlateNavigationDebugger::PaintGeometryOutline(const FGeometry& InGeometry, const FLinearColor& InColor, FSlateWindowElementList& OutDrawElements, int32 LayerId) const
+{
+	FSlateDrawElement::MakeGeometryOutline(
+		OutDrawElements,
+		LayerId,
+		InGeometry.ToPaintGeometry(),
+		InColor,
+		ESlateDrawEffect::None,
+		ArrowGeometryHelpers::bAntiAlias,
+		ArrowGeometryHelpers::LineThickness
+	);
+}
+
+void FSlateNavigationDebugger::PaintDashedLine(const FSlateUserNavigationData& InNavigationData, TArray<FVector2f>&& PointsToConnect, const FLinearColor& InColor, FSlateWindowElementList& OutDrawElements, int32 LayerId) const
+{
+	check(InNavigationData.CanPaint());
+
+	FSlateDrawElement::MakeDashedLines(
+		OutDrawElements,
+		LayerId,
+		InNavigationData.WindowGeometry.ToPaintGeometry(),
+		MoveTemp(PointsToConnect),
+		ESlateDrawEffect::None,
+		InColor,
+		ArrowGeometryHelpers::LineThickness,
+		ArrowGeometryHelpers::DashLineLength,
+		0.0f
+	);
+}
+
+int32 FSlateNavigationDebugger::PaintCandidateWidgets(const FSlateUserNavigationData& InNavigationData, const TArray<FHittestGrid::FDebuggingFindNextFocusableWidgetArgs>& InCandidateWidgets, FSlateWindowElementList& OutDrawElements, int32 LayerId) const
+{
+	if (CVarDisplayCandidateWidgets.GetValueOnAnyThread())
+	{
+		check(InNavigationData.CanPaint());
+
+		TMap<TSharedRef<const SWidget>, TSet<EUINavigation>> HighlightedCandidateWidgets;
+
+		for (const FHittestGrid::FDebuggingFindNextFocusableWidgetArgs& Info : InCandidateWidgets)
+		{
+			for (const FHittestGrid::FDebuggingFindNextFocusableWidgetArgs::FWidgetResult& Result : Info.IntermediateResults)
+			{
+				if (Result.Widget.IsValid())
+				{
+					TSharedRef<const SWidget> CandidateWidgetRef = Result.Widget.ToSharedRef();
+
+					const bool bNotCandidate = Result.Result.EqualTo(HittestGridDebuggingText::Valid);
+
+					const bool bAlreadyOutlined = HighlightedCandidateWidgets.Contains(CandidateWidgetRef);
+					const bool bAlreadyPaintedArrow = bAlreadyOutlined && HighlightedCandidateWidgets[CandidateWidgetRef].Contains(Info.Direction);
+
+					if (!bNotCandidate && !bAlreadyPaintedArrow)
+					{
+						HighlightedCandidateWidgets.FindOrAdd(CandidateWidgetRef).Add(Info.Direction);
+
+						const FGeometry CandidateWidgetGeometry = CandidateWidgetRef->GetCachedGeometry().MakeChild(InNavigationData.WindowWidget->GetPositionInScreen() * -1);						
+						if (!bAlreadyOutlined)
+						{
+							const FLinearColor OutlineColor = FNavigationDebuggingStyle::Get().GetTargetWidgetColor().CopyWithNewOpacity(ArrowGeometryHelpers::CandidateColorAlpha);
+							PaintGeometryOutline(CandidateWidgetGeometry, OutlineColor, OutDrawElements, ++LayerId);
+						}
+
+						const bool bOnScreen = !Result.Result.EqualTo(HittestGridDebuggingText::DoesNotIntersect);
+						if (bOnScreen)
+						{
+							PaintOnScreenWidgetArrow(InNavigationData, Info.Direction, CandidateWidgetGeometry, FNavigationDebuggingStyle::EArrowType::Candidate, OutDrawElements, ++LayerId);
+						}
+						else
+						{
+							PaintOffScreenWidgetArrow(InNavigationData, Info.Direction, CandidateWidgetRef, FNavigationDebuggingStyle::EArrowType::Candidate, OutDrawElements, ++LayerId);
+						}
+					}
+				}
+			}
+		}
+	}
+
+	return LayerId;
+}
+
+int32 FSlateNavigationDebugger::PaintFocusTargetArrowsAndLines(const TPair<EUINavigation, TSharedRef<SWidget>>& FocusedToDestination, const FSlateUserNavigationData& InNavigationData, FSlateWindowElementList& OutDrawElements, int32 LayerId) const
+{
+	check(InNavigationData.CanPaint());
+
+	// Paint Target Widget Arrows
+	FVector2f TargetWidgetArrowPosition;
+	{
+		const bool bTargetWidgetInViewport = !ArrowGeometryHelpers::IsWidgetOutOfViewport(InNavigationData, FocusedToDestination.Value);
+		if (bTargetWidgetInViewport)
+		{
+			const FGeometry TargetWidgetGeometry = FocusedToDestination.Value->GetCachedGeometry().MakeChild(InNavigationData.WindowWidget->GetPositionInScreen() * -1);
+			const FGeometry TargetArrowGeometry = PaintOnScreenWidgetArrow(InNavigationData, FocusedToDestination.Key, TargetWidgetGeometry, FNavigationDebuggingStyle::EArrowType::Target, OutDrawElements, ++LayerId);
+			TargetWidgetArrowPosition = TargetArrowGeometry.GetAbsolutePositionAtCoordinates(ArrowGeometryHelpers::GetArrowCoordinatesForLine(ArrowGeometryHelpers::GetOpposingDirection(FocusedToDestination.Key)));
+		}
+		else
+		{
+			const FGeometry TargetArrowGeometry = PaintOffScreenWidgetArrow(InNavigationData, FocusedToDestination.Key, FocusedToDestination.Value, FNavigationDebuggingStyle::EArrowType::Target, OutDrawElements, ++LayerId);
+			TargetWidgetArrowPosition = TargetArrowGeometry.GetAbsolutePositionAtCoordinates(ArrowGeometryHelpers::GetArrowCoordinatesForLine(ArrowGeometryHelpers::GetOpposingDirection(FocusedToDestination.Key)));
+		}
+	}
+
+	// Paint Focus Widget Arrows
+	FVector2f FocusedWidgetArrowPosition;
+	{
+		const bool bFocusWidgetInViewport = !ArrowGeometryHelpers::IsWidgetOutOfViewport(InNavigationData, InNavigationData.FocusedWidget.ToSharedRef());
+		if (bFocusWidgetInViewport)
+		{
+			const FGeometry FocusedWidgetGeometry = InNavigationData.FocusedWidget->GetCachedGeometry().MakeChild(InNavigationData.WindowWidget->GetPositionInScreen() * -1);
+			const FGeometry FocusArrowGeometry = PaintOnScreenWidgetArrow(InNavigationData, FocusedToDestination.Key, FocusedWidgetGeometry, FNavigationDebuggingStyle::EArrowType::Focus, OutDrawElements, ++LayerId);
+			FocusedWidgetArrowPosition = FocusArrowGeometry.GetAbsolutePositionAtCoordinates(ArrowGeometryHelpers::GetArrowCoordinatesForLine(FocusedToDestination.Key));
+		}
+		else
+		{
+			const FGeometry FocusArrowGeometry = PaintOffScreenWidgetArrow(InNavigationData, FocusedToDestination.Key, InNavigationData.FocusedWidget.ToSharedRef(), FNavigationDebuggingStyle::EArrowType::Focus, OutDrawElements, ++LayerId);
+			FocusedWidgetArrowPosition = FocusArrowGeometry.GetAbsolutePositionAtCoordinates(ArrowGeometryHelpers::GetArrowCoordinatesForLine(FocusedToDestination.Key));
+		}
+	}
+
+	// Paint Dashed Lines
+	TArray<FVector2f> FocusToTargetPoints;
+	FocusToTargetPoints.Add(FocusedWidgetArrowPosition);
+	FocusToTargetPoints.Add(TargetWidgetArrowPosition);
+
+	PaintDashedLine(InNavigationData, MoveTemp(FocusToTargetPoints), FNavigationDebuggingStyle::Get().GetFocusedWidgetColor(), OutDrawElements, ++LayerId);
+
+	return LayerId;
+}
+
+int32 FSlateNavigationDebugger::PaintForUser(FSlateUser& InSlateUser, FSlateWindowElementList& OutDrawElements, int32 LayerId) const
+{
+	const FSlateUserNavigationData NavigationData = FSlateUserNavigationData(InSlateUser);
+
+	if (NavigationData.CanPaint())
+	{
+		FDelegateHandle HittestGridHandle;
+		TArray<FHittestGrid::FDebuggingFindNextFocusableWidgetArgs> HitTestGridInfo;
+
+		// Bind to listen to candidate widgets
+		SetListenToHitTestGrid(true, HitTestGridInfo, HittestGridHandle);
+
+		// Gather Destination Widgets
+		TMap<EUINavigation, TSharedRef<SWidget>> DirectionToDestinationWidget = SimulateNavigationDirection(NavigationData, InSlateUser.GetUserIndex());
+
+		// Paint Candidate Widgets
+		LayerId = PaintCandidateWidgets(NavigationData, HitTestGridInfo, OutDrawElements, LayerId);
+
+		// Highlight Focused Widget
+		PaintGeometryOutline(NavigationData.FocusedWidget->GetPaintSpaceGeometry(), FNavigationDebuggingStyle::Get().GetFocusedWidgetColor(), OutDrawElements, ++LayerId);
+
+		// Paint Focused Widget Arrows and Target Arrows; with Dashed Lines
+		for (const TPair<EUINavigation, TSharedRef<SWidget>>& DirectionToWidget : DirectionToDestinationWidget)
+		{
+			// Paint Target Widget Outline
+			PaintGeometryOutline(DirectionToWidget.Value->GetPaintSpaceGeometry(), FNavigationDebuggingStyle::Get().GetTargetWidgetColor(), OutDrawElements, ++LayerId);
+
+			// Paint Arrows and Lines
+			LayerId = PaintFocusTargetArrowsAndLines(DirectionToWidget, NavigationData, OutDrawElements, LayerId);
+		}
+		
+		// Unbind to listen to candidate widgets
+		SetListenToHitTestGrid(false, HitTestGridInfo, HittestGridHandle);
+	}
+
+	return LayerId;
+}
+
+int32 FSlateNavigationDebugger::Paint(FSlateWindowElementList& OutDrawElements, int32 LayerId) const
+{
+	if (GSlateNavigationDebugging)
+	{
+		FSlateApplication::Get().ForEachUser([&OutDrawElements, &LayerId](FSlateUser& SlateUser) -> void
+											 {
+												LayerId = FSlateNavigationDebugger::Get().PaintForUser(SlateUser, OutDrawElements, LayerId);
+											 });
+	}
+
+	return LayerId;
+}
+
+
+#endif // WITH_SLATE_DEBUGGING
\ No newline at end of file
diff --git a/Engine/Source/Runtime/Slate/Public/Framework/Application/SlateApplication.h b/Engine/Source/Runtime/Slate/Public/Framework/Application/SlateApplication.h
--- a/Engine/Source/Runtime/Slate/Public/Framework/Application/SlateApplication.h
+++ b/Engine/Source/Runtime/Slate/Public/Framework/Application/SlateApplication.h
@@ -1904,6 +1904,18 @@ private:
 	/* Event receiver for when a menu in the menu stack is destroyed */
 	SLATE_API void OnMenuDestroyed(const TSharedRef<IMenu>& Menu);
 
+public:	
+	struct FNavigationResult
+	{
+		TSharedPtr<SWidget> DestinationWidget;
+		bool bAlwaysHandleNavigationAttempt = false;
+#if WITH_SLATE_DEBUGGING
+		ESlateDebuggingNavigationMethod NavigationMethod = ESlateDebuggingNavigationMethod::Unknown;
+#endif
+	};
+
+	SLATE_API FNavigationResult CalculateDestinationWidget(const FWidgetPath& NavigationSource, const FNavigationReply& NavigationReply, EUINavigation NavigationType, uint32 UserIndex, const FArrangedWidget& BoundaryWidget) const;
+
 private:
 	SLATE_API FSlateApplication();
 	SLATE_API void SetLastUserInteractionTime(const double InCurrentTime);	
diff --git a/Engine/Source/Runtime/Slate/Public/Navigation/SlateNavigationDebugger.h b/Engine/Source/Runtime/Slate/Public/Navigation/SlateNavigationDebugger.h
new file mode 100644
--- /dev/null
+++ b/Engine/Source/Runtime/Slate/Public/Navigation/SlateNavigationDebugger.h
@@ -0,0 +1,97 @@
+// Copyright Epic Games, Inc. All Rights Reserved.
+
+#pragma once
+
+#include "Debugging/SlateDebugging.h"
+
+#if WITH_SLATE_DEBUGGING
+#include "Input/HittestGrid.h"
+#include "Styling/SlateStyle.h"
+
+class FSlateUser;
+class FSlateWindowElementList;
+class FWidgetPath;
+class SWidget;
+class SWindow;
+
+class FDelegateHandle;
+struct FGeometry;
+
+class FNavigationDebuggingStyle : public FSlateStyleSet
+{
+public:
+	enum class EArrowType
+	{
+		Focus,
+		Candidate,
+		Target
+	};
+
+	FLinearColor GetFocusedWidgetColor() const;
+	FLinearColor GetTargetWidgetColor() const;
+
+	const FSlateBrush* GetArrowBrush(const EUINavigation InNavigation, const EArrowType InArrow) const;
+
+	static FNavigationDebuggingStyle& Get();
+
+private:
+	FNavigationDebuggingStyle();
+	FNavigationDebuggingStyle(const FNavigationDebuggingStyle&) = delete;
+	FNavigationDebuggingStyle& operator=(const FNavigationDebuggingStyle&) = delete;
+
+	virtual ~FNavigationDebuggingStyle() override;
+
+private:
+	TMap<EUINavigation, TMap<EArrowType, const FSlateBrush*>> NavigationArrowBrushes;
+};
+
+struct FSlateUserNavigationData
+{
+	FSlateUserNavigationData() = delete;
+	FSlateUserNavigationData(const FSlateUserNavigationData&) = delete;
+	FSlateUserNavigationData& operator=(const FSlateUserNavigationData&) = delete;
+
+	FSlateUserNavigationData(FSlateUser& InSlateUser);
+
+public:
+	bool CanPaint() const;
+
+public:
+	TSharedRef<FWidgetPath> FocusedWidgetPath;
+	TSharedPtr<SWindow> WindowWidget;
+	TSharedPtr<SWidget> ViewportWidget;
+	TSharedPtr<SWidget> FocusedWidget;
+	FGeometry WindowGeometry;
+};
+
+struct FSlateNavigationDebugger
+{
+	SLATE_API static FSlateNavigationDebugger& Get();
+
+	SLATE_API int32 Paint(FSlateWindowElementList& OutDrawElements, int32 LayerId) const;
+
+private:
+	FSlateNavigationDebugger() = default;
+	FSlateNavigationDebugger(const FSlateNavigationDebugger&) = delete;
+	FSlateNavigationDebugger& operator=(const FSlateNavigationDebugger&) = delete;
+
+private:
+	int32 PaintForUser(FSlateUser& InSlateUser, FSlateWindowElementList& OutDrawElements, int32 LayerId) const;
+
+	int32 PaintCandidateWidgets(const FSlateUserNavigationData& InNavigationData, const TArray<FHittestGrid::FDebuggingFindNextFocusableWidgetArgs>& InCandidateWidgets, FSlateWindowElementList& OutDrawElements, int32 LayerId) const;
+	int32 PaintFocusTargetArrowsAndLines(const TPair<EUINavigation, TSharedRef<SWidget>>& FocusedToDestination, const FSlateUserNavigationData& InNavigationData, FSlateWindowElementList& OutDrawElements, int32 LayerId) const;
+	void PaintGeometryOutline(const FGeometry& InGeometry, const FLinearColor& InColor, FSlateWindowElementList& OutDrawElements, int32 LayerId) const;
+	void PaintDashedLine(const FSlateUserNavigationData& InNavigationData, TArray<FVector2f>&& PointsToConnect, const FLinearColor& InColor, FSlateWindowElementList& OutDrawElements, int32 LayerId) const;
+
+	// Returns Arrow Geometry
+	FGeometry PaintOnScreenWidgetArrow(const FSlateUserNavigationData& InNavigationData, const EUINavigation InNavDirection, const FGeometry& InWidgetGeometry, const FNavigationDebuggingStyle::EArrowType InArrowType, FSlateWindowElementList& OutDrawElements, int32 LayerId) const;
+
+	// Returns Arrow Geometry
+	FGeometry PaintOffScreenWidgetArrow(const FSlateUserNavigationData& InNavigationData, const EUINavigation InNavDirection, TSharedRef<const SWidget> InWidget, const FNavigationDebuggingStyle::EArrowType InArrowType, FSlateWindowElementList& OutDrawElements, int32 LayerId) const;
+
+private:
+	TMap<EUINavigation, TSharedRef<SWidget>> SimulateNavigationDirection(const FSlateUserNavigationData& InNavigationData, const int32 SlateUserIndex) const;
+	void SetListenToHitTestGrid(const bool bBind, TArray<FHittestGrid::FDebuggingFindNextFocusableWidgetArgs>& InContainer, FDelegateHandle& InHandle) const;
+};
+
+#endif // WITH_SLATE_DEBUGGING
\ No newline at end of file
diff --git a/Engine/Source/Runtime/SlateCore/Private/Input/HittestGrid.cpp b/Engine/Source/Runtime/SlateCore/Private/Input/HittestGrid.cpp
--- a/Engine/Source/Runtime/SlateCore/Private/Input/HittestGrid.cpp
+++ b/Engine/Source/Runtime/SlateCore/Private/Input/HittestGrid.cpp
@@ -348,14 +348,14 @@ namespace UE::Slate::Private
 #if WITH_SLATE_DEBUGGING
 namespace HittestGridDebuggingText
 {
-	static FText Valid = LOCTEXT("StateValid", "Valid"); //~ The widget is valid will be consider as the result
-	static FText NotCompatibleWithUserIndex = LOCTEXT("StateNotCompatibleWithUserIndex", "User Index not compatible"); //~ The widget is not compatible with the requested user index
-	static FText DoesNotIntersect = LOCTEXT("StateDoesNotIntersect", "Does not intersect"); //~ The widget rect is not in the correct direction or is not intersecting with the "swept" rectangle
-	static FText PreviousWidgetIsBetter = LOCTEXT("StatePreviousWidgetIsBetter", "Previous Widget was better"); //~ The widget would be valid but the previous valid is closer
-	static FText NotADescendant = LOCTEXT("StateNotADescendant", "Not a descendant"); //~ We have a non escape boundary condition and the widget isn't a descendant of our boundary
-	static FText Disabled = LOCTEXT("StateNotEnabled", "Disabled"); //~ The widget is not enabled
-	static FText ParentDisabled = LOCTEXT("StateParentNotEnabled", "ParentDisabled"); //~ A parent of the widget is disabled
-	static FText DoesNotSuportKeyboardFocus = LOCTEXT("StateDoesNotSuportKeyboardFocus", "Keyboard focus unsupported"); //~ THe widget does not support keyboard focus
+	/*extern*/ const FText Valid = LOCTEXT("StateValid", "Valid"); //~ The widget is valid will be consider as the result
+	/*extern*/ const FText NotCompatibleWithUserIndex = LOCTEXT("StateNotCompatibleWithUserIndex", "User Index not compatible"); //~ The widget is not compatible with the requested user index
+	/*extern*/ const FText DoesNotIntersect = LOCTEXT("StateDoesNotIntersect", "Does not intersect"); //~ The widget rect is not in the correct direction or is not intersecting with the "swept" rectangle
+	/*extern*/ const FText PreviousWidgetIsBetter = LOCTEXT("StatePreviousWidgetIsBetter", "Previous Widget was better"); //~ The widget would be valid but the previous valid is closer
+	/*extern*/ const FText NotADescendant = LOCTEXT("StateNotADescendant", "Not a descendant"); //~ We have a non escape boundary condition and the widget isn't a descendant of our boundary
+	/*extern*/ const FText Disabled = LOCTEXT("StateNotEnabled", "Disabled"); //~ The widget is not enabled
+	/*extern*/ const FText ParentDisabled = LOCTEXT("StateParentNotEnabled", "ParentDisabled"); //~ A parent of the widget is disabled
+	/*extern*/ const FText DoesNotSuportKeyboardFocus = LOCTEXT("StateDoesNotSuportKeyboardFocus", "Keyboard focus unsupported"); //~ THe widget does not support keyboard focus
 }
 #define AddToNextFocusableWidgetCondidateDebugResults(Candidate, Result) { if (IntermediateResultsPtr) { IntermediateResultsPtr->Emplace((Candidate), (Result)); } }
 #else
diff --git a/Engine/Source/Runtime/SlateCore/Public/Input/HittestGrid.h b/Engine/Source/Runtime/SlateCore/Public/Input/HittestGrid.h
--- a/Engine/Source/Runtime/SlateCore/Public/Input/HittestGrid.h
+++ b/Engine/Source/Runtime/SlateCore/Public/Input/HittestGrid.h
@@ -14,6 +14,20 @@
 class FArrangedChildren;
 struct FNavigationMethod;
 
+#if WITH_SLATE_DEBUGGING
+namespace HittestGridDebuggingText
+{
+	extern SLATECORE_API const FText Valid;
+	extern SLATECORE_API const FText PreviousWidgetIsBetter;
+	extern SLATECORE_API const FText NotCompatibleWithUserIndex;
+	extern SLATECORE_API const FText DoesNotIntersect;
+	extern SLATECORE_API const FText NotADescendant;
+	extern SLATECORE_API const FText Disabled;
+	extern SLATECORE_API const FText ParentDisabled;
+	extern SLATECORE_API const FText DoesNotSuportKeyboardFocus;
+};
+#endif // WITH_SLATE_DEBUGGING
+
 class ICustomHitTestPath
 {
 public:
