Fix define collision
#rnx

diff --git a/Engine/Source/Runtime/Slate/Private/Widgets/Layout/SWindowTitleBarArea.cpp b/Engine/Source/Runtime/Slate/Private/Widgets/Layout/SWindowTitleBarArea.cpp
--- a/Engine/Source/Runtime/Slate/Private/Widgets/Layout/SWindowTitleBarArea.cpp
+++ b/Engine/Source/Runtime/Slate/Private/Widgets/Layout/SWindowTitleBarArea.cpp
@@ -10,8 +10,8 @@
 #include "Brushes/SlateBoxBrush.h"
 #include "Brushes/SlateImageBrush.h"
 
-#define IMAGE_BRUSH( RelativePath, ... ) FSlateImageBrush(FPaths::EngineContentDir() / TEXT("Slate") / RelativePath + TEXT(".png"), __VA_ARGS__)
-#define BOX_BRUSH( RelativePath, ... ) FSlateBoxBrush(FPaths::EngineContentDir() / TEXT("Slate") / RelativePath + TEXT(".png"), __VA_ARGS__)
+#define WINDOW_IMAGE_BRUSH( RelativePath, ... ) FSlateImageBrush(FPaths::EngineContentDir() / TEXT("Slate") / RelativePath + TEXT(".png"), __VA_ARGS__)
+#define WINDOW_BOX_BRUSH( RelativePath, ... ) FSlateBoxBrush(FPaths::EngineContentDir() / TEXT("Slate") / RelativePath + TEXT(".png"), __VA_ARGS__)
 
 FSimpleDelegate SWindowTitleBarArea::OnCloseButtonClicked;
 bool SWindowTitleBarArea::bIsCloseButtonActive = true;
@@ -34,34 +34,34 @@ SWindowTitleBarArea::SWindowTitleBarArea()
 	if (!bButtonsStyleInitialized)
 	{
 		const FButtonStyle ButtonStyle = FButtonStyle()
-			.SetNormal(BOX_BRUSH("Common/Button", FVector2D(32, 32), 8.0f / 32.0f))
-			.SetHovered(BOX_BRUSH("Common/Button_Hovered", FVector2D(32, 32), 8.0f / 32.0f))
-			.SetPressed(BOX_BRUSH("Common/Button_Pressed", FVector2D(32, 32), 8.0f / 32.0f))
+			.SetNormal(WINDOW_BOX_BRUSH("Common/Button", FVector2D(32, 32), 8.0f / 32.0f))
+			.SetHovered(WINDOW_BOX_BRUSH("Common/Button_Hovered", FVector2D(32, 32), 8.0f / 32.0f))
+			.SetPressed(WINDOW_BOX_BRUSH("Common/Button_Pressed", FVector2D(32, 32), 8.0f / 32.0f))
 			.SetNormalPadding(FMargin(2, 2, 2, 2))
 			.SetPressedPadding(FMargin(2, 3, 2, 1));
 
 		MinimizeButtonStyle = FButtonStyle(ButtonStyle)
-			.SetNormal(IMAGE_BRUSH("Common/Window/WindowButton_Minimize_Normal", FVector2D(27.0f, 18.0f)))
-			.SetHovered(IMAGE_BRUSH("Common/Window/WindowButton_Minimize_Hovered", FVector2D(27.0f, 18.0f)))
-			.SetPressed(IMAGE_BRUSH("Common/Window/WindowButton_Minimize_Pressed", FVector2D(27.0f, 18.0f)))
-			.SetDisabled(IMAGE_BRUSH("Common/Window/WindowButton_Minimize_Disabled", FVector2D(27.0f, 18.0f)));
+			.SetNormal(WINDOW_IMAGE_BRUSH("Common/Window/WindowButton_Minimize_Normal", FVector2D(27.0f, 18.0f)))
+			.SetHovered(WINDOW_IMAGE_BRUSH("Common/Window/WindowButton_Minimize_Hovered", FVector2D(27.0f, 18.0f)))
+			.SetPressed(WINDOW_IMAGE_BRUSH("Common/Window/WindowButton_Minimize_Pressed", FVector2D(27.0f, 18.0f)))
+			.SetDisabled(WINDOW_IMAGE_BRUSH("Common/Window/WindowButton_Minimize_Disabled", FVector2D(27.0f, 18.0f)));
 
 		MaximizeButtonStyle = FButtonStyle(ButtonStyle)
-			.SetNormal(IMAGE_BRUSH("Common/Window/WindowButton_Maximize_Normal", FVector2D(23.0f, 18.0f)))
-			.SetHovered(IMAGE_BRUSH("Common/Window/WindowButton_Maximize_Hovered", FVector2D(23.0f, 18.0f)))
-			.SetPressed(IMAGE_BRUSH("Common/Window/WindowButton_Maximize_Pressed", FVector2D(23.0f, 18.0f)))
-			.SetDisabled(IMAGE_BRUSH("Common/Window/WindowButton_Maximize_Disabled", FVector2D(27.0f, 18.0f)));
+			.SetNormal(WINDOW_IMAGE_BRUSH("Common/Window/WindowButton_Maximize_Normal", FVector2D(23.0f, 18.0f)))
+			.SetHovered(WINDOW_IMAGE_BRUSH("Common/Window/WindowButton_Maximize_Hovered", FVector2D(23.0f, 18.0f)))
+			.SetPressed(WINDOW_IMAGE_BRUSH("Common/Window/WindowButton_Maximize_Pressed", FVector2D(23.0f, 18.0f)))
+			.SetDisabled(WINDOW_IMAGE_BRUSH("Common/Window/WindowButton_Maximize_Disabled", FVector2D(27.0f, 18.0f)));
 
 		RestoreButtonStyle = FButtonStyle(ButtonStyle)
-			.SetNormal(IMAGE_BRUSH("Common/Window/WindowButton_Restore_Normal", FVector2D(23.0f, 18)))
-			.SetHovered(IMAGE_BRUSH("Common/Window/WindowButton_Restore_Hovered", FVector2D(23.0f, 18)))
-			.SetPressed(IMAGE_BRUSH("Common/Window/WindowButton_Restore_Pressed", FVector2D(23.0f, 18)))
-			.SetDisabled(IMAGE_BRUSH("Common/Window/WindowButton_Maximize_Disabled", FVector2D(27.0f, 18.0f)));
+			.SetNormal(WINDOW_IMAGE_BRUSH("Common/Window/WindowButton_Restore_Normal", FVector2D(23.0f, 18)))
+			.SetHovered(WINDOW_IMAGE_BRUSH("Common/Window/WindowButton_Restore_Hovered", FVector2D(23.0f, 18)))
+			.SetPressed(WINDOW_IMAGE_BRUSH("Common/Window/WindowButton_Restore_Pressed", FVector2D(23.0f, 18)))
+			.SetDisabled(WINDOW_IMAGE_BRUSH("Common/Window/WindowButton_Maximize_Disabled", FVector2D(27.0f, 18.0f)));
 
 		CloseButtonStyle = FButtonStyle(ButtonStyle)
-			.SetNormal(IMAGE_BRUSH("Common/Window/WindowButton_Close_Normal", FVector2D(44.0f, 18.0f)))
-			.SetHovered(IMAGE_BRUSH("Common/Window/WindowButton_Close_Hovered", FVector2D(44.0f, 18.0f)))
-			.SetPressed(IMAGE_BRUSH("Common/Window/WindowButton_Close_Pressed", FVector2D(44.0f, 18.0f)));
+			.SetNormal(WINDOW_IMAGE_BRUSH("Common/Window/WindowButton_Close_Normal", FVector2D(44.0f, 18.0f)))
+			.SetHovered(WINDOW_IMAGE_BRUSH("Common/Window/WindowButton_Close_Hovered", FVector2D(44.0f, 18.0f)))
+			.SetPressed(WINDOW_IMAGE_BRUSH("Common/Window/WindowButton_Close_Pressed", FVector2D(44.0f, 18.0f)));
 
 		bButtonsStyleInitialized = true;
 	}
