Fix for the minimized window not saving their restored size properly in the layout system.

This would result in the editor opening the windows with a bad size.

#Synced-CL 51355948

diff --git a/Engine/Source/Runtime/Slate/Private/Framework/Docking/SDockingArea.cpp b/Engine/Source/Runtime/Slate/Private/Framework/Docking/SDockingArea.cpp
--- a/Engine/Source/Runtime/Slate/Private/Framework/Docking/SDockingArea.cpp
+++ b/Engine/Source/Runtime/Slate/Private/Framework/Docking/SDockingArea.cpp
@@ -20,8 +20,9 @@ namespace UE::Slate::Private
 	{
 		bool bIsWindowMaximized = InWindow.IsWindowMaximized();
 		FSlateRect WindowRect;
-		if (bIsWindowMaximized)
+		if (bIsWindowMaximized || InWindow.IsWindowMinimized() || !InWindow.IsVisible())
 		{
+			// Confusing name but does use the restore size when there is one before the actual size of the window.
 			WindowRect = InWindow.GetNonMaximizedRectInScreen();
 		}
 		else
