Fixed editor key bindings essentially resetting to default whenever opening the Editor Preferences, closing it and the editor crashing since EditorKeyBindings.ini was written to disk only on editor shutdown.

#rnx
#Synced-CL 51362910

diff --git a/Engine/Source/Runtime/Slate/Private/Framework/Commands/InputBindingManager.cpp b/Engine/Source/Runtime/Slate/Private/Framework/Commands/InputBindingManager.cpp
--- a/Engine/Source/Runtime/Slate/Private/Framework/Commands/InputBindingManager.cpp
+++ b/Engine/Source/Runtime/Slate/Private/Framework/Commands/InputBindingManager.cpp
@@ -245,6 +245,13 @@ void FUserDefinedChords::SaveChords() const
 			ChordJsonArray.Add(EscapedContent);
 		}
 
+		// Disable file operations to prevent the .ini from being deleted when all the sections are emptied below
+		const bool bFileOperationsDisabled = GConfig->AreFileOperationsDisabled();
+		if (!bFileOperationsDisabled)
+		{
+			GConfig->DisableFileOperations();
+		}
+
 		// Remove all sections not related to keybinds...
 		TArray<FString> SectionNames;
 		if (GConfig->GetSectionNames(GEditorKeyBindingsIni, SectionNames))
@@ -255,12 +262,21 @@ void FUserDefinedChords::SaveChords() const
 			}
 		}
 
+		if (!bFileOperationsDisabled)
+		{
+			GConfig->EnableFileOperations();
+		}
+
 		GConfig->SetArray(TEXT("UserDefinedChords"), TEXT("UserDefinedChords"), ChordJsonArray, GEditorKeyBindingsIni);
 
 		// Clean up old keys (if they still exist)
 		GConfig->RemoveKey(TEXT("UserDefinedGestures"), TEXT("UserDefinedGestures"), GEditorKeyBindingsIni);
 		GConfig->RemoveKey(TEXT("UserDefinedChords"), TEXT("Content"), GEditorKeyBindingsIni);
 		GConfig->RemoveKey(TEXT("UserDefinedChords"), TEXT("Content"), GEditorKeyBindingsIni);
+
+		// Flush the .ini to ensure all new key bindings are saved
+		const bool bRemoveFromCache = false;
+		GConfig->Flush(bRemoveFromCache, GEditorKeyBindingsIni);
 	}
 }
 
