Currently, if the menu anchor is showing content and ends up being hidden, the popup will continue being shown. This can happen when the menu itself is made hidden, but also when one of its ancestor is hidden from view. To avoid this and to properly hide the content when the MenuAnchor is, we need to adjust the menu content visibility and check that all of ancestors are visible.

#Synced-CL 51454850

diff --git a/Engine/Source/Runtime/Slate/Private/Widgets/Input/SMenuAnchor.cpp b/Engine/Source/Runtime/Slate/Private/Widgets/Input/SMenuAnchor.cpp
--- a/Engine/Source/Runtime/Slate/Private/Widgets/Input/SMenuAnchor.cpp
+++ b/Engine/Source/Runtime/Slate/Private/Widgets/Input/SMenuAnchor.cpp
@@ -121,6 +121,33 @@ FGeometry SMenuAnchor::ComputeNewWindowMenuPlacement(const FGeometry& AllottedGe
 		FSlateLayoutTransform(TransformPoint(Inverse(AllottedGeometry.GetAccumulatedLayoutTransform()), NewPositionDesktopSpace)));
 }
 
+EVisibility SMenuAnchor::GetMenuContentVisibility(TAttribute<EVisibility> OriginalVisibilityAttribute) const
+{
+	FWidgetProxyHandle ProxyHandle = GetProxyHandle();
+	if (ProxyHandle.IsValid(this))
+	{
+		FSlateInvalidationWidgetVisibility Visibility = ProxyHandle.GetWidgetVisibility(this);
+		if (!Visibility.IsVisible())
+		{
+			return EVisibility::Collapsed;
+		}
+	}
+	else // if we don't have access to the proxy handle, we need to walk the hierarchy
+	{
+		TSharedPtr<const SWidget> Widget = AsShared();
+		while (Widget)
+		{
+			if (!Widget->GetVisibility().IsVisible())
+			{
+				return EVisibility::Collapsed;
+			}
+			Widget = Widget->GetParentWidget();
+		}
+	}
+
+	return OriginalVisibilityAttribute.Get(EVisibility::Collapsed);
+}
+
 void SMenuAnchor::Tick( const FGeometry& AllottedGeometry, const double InCurrentTime, const float InDeltaTime )
 {
 	TSharedPtr<SWindow> PopupWindow = PopupWindowPtr.Pin();
@@ -413,6 +440,12 @@ void SMenuAnchor::SetIsOpen( bool InIsOpen, const bool bFocusMenu, const int32 F
 									UE_LOG(LogSlate, Error, TEXT(" MenuContentRef: %s"), *MenuContentRef->ToString());
 								}
 #endif
+								// We want to be sure the menu content ends up hidden if this SMenuAnchor is also hidden
+								NewMenu->GetContent()->SetVisibility(TAttribute<EVisibility>::CreateSP(this,
+									&SMenuAnchor::GetMenuContentVisibility,
+									AccessWidgetVisibilityAttribute(NewMenu->GetContent().ToSharedRef())
+								));
+
 								if (NewMenu->GetOwnedWindow().IsValid())
 								{
 									PopupMenuPtr = NewMenu;
diff --git a/Engine/Source/Runtime/Slate/Public/Widgets/Input/SMenuAnchor.h b/Engine/Source/Runtime/Slate/Public/Widgets/Input/SMenuAnchor.h
--- a/Engine/Source/Runtime/Slate/Public/Widgets/Input/SMenuAnchor.h
+++ b/Engine/Source/Runtime/Slate/Public/Widgets/Input/SMenuAnchor.h
@@ -157,6 +157,12 @@ protected:
 	/** Computes the placement geometry for menus displayed in a separately created window */
 	SLATE_API FGeometry ComputeNewWindowMenuPlacement(const FGeometry& AllottedGeometry, const FVector2D& PopupDesiredSize, EMenuPlacement PlacementMode) const;
 
+	/**
+	 * Checks if any parent of the current SMenuAnchor is hidden (in which case this returns EVisibility::Collapsed).
+	 * If all the ancestors are visible, this returns the value of the OriginalVisibilityAttribute.
+	 */
+	SLATE_API EVisibility GetMenuContentVisibility(TAttribute<EVisibility> OriginalVisibilityAttribute) const;
+
 	static SLATE_API TArray<TWeakPtr<IMenu>> OpenApplicationMenus;
 
 	TSlateAttributeRef<EMenuPlacement> GetPlacementAttribute() const
