Setting a minimum arrow size for displaying an arrow in the navigation debug tool
Fixing Geometry to use PaintSpaceGeometry (Allotted) over CachedGeometry (TickSpace)
Using AbsoluteSize and AbsolutePosition to calculate ArrowGeometry

#Synced-CL 51546663

#Synced-CL 51546663

#Synced-CL 51546663

diff --git a/Engine/Source/Runtime/Slate/Private/Navigation/SlateNavigationDebugger.cpp b/Engine/Source/Runtime/Slate/Private/Navigation/SlateNavigationDebugger.cpp
--- a/Engine/Source/Runtime/Slate/Private/Navigation/SlateNavigationDebugger.cpp
+++ b/Engine/Source/Runtime/Slate/Private/Navigation/SlateNavigationDebugger.cpp
@@ -22,6 +22,7 @@ namespace ArrowGeometryHelpers
 	constexpr float LineThickness = 2.0f;
 	constexpr float DashLineLength = 4.0f;
 	constexpr float CandidateColorAlpha = 0.25f;
+	constexpr float MinimumArrowGeometrySizeToPaint = 8.0f;
 	static const FVector2f ArrowBrushSize = FVector2f(30.0f, 30.0f);
 
 	static EUINavigation GetOpposingDirection(const EUINavigation InNavDirection)
@@ -76,10 +77,10 @@ namespace ArrowGeometryHelpers
 
 		TSharedRef<SWidget> ViewportWidget = InNavigationData.ViewportWidget.ToSharedRef();
 
-		const FGeometry ViewportGeometry = ViewportWidget->GetCachedGeometry().MakeChild(InNavigationData.WindowWidget->GetPositionInScreen() * -1);
+		const FGeometry ViewportGeometry = ViewportWidget->GetPaintSpaceGeometry();
 		const FSlateRect ViewportRect = FSlateRect::FromPointAndExtent(ViewportGeometry.GetAbsolutePosition(), ViewportGeometry.GetAbsoluteSize());
 		
-		const FGeometry WidgetGeometry = InWidget->GetCachedGeometry().MakeChild(InNavigationData.WindowWidget->GetPositionInScreen() * -1);
+		const FGeometry WidgetGeometry = InWidget->GetPaintSpaceGeometry();
 		const FSlateRect WidgetRect = FSlateRect::FromPointAndExtent(WidgetGeometry.GetAbsolutePosition(), WidgetGeometry.GetAbsoluteSize());
 
 		return !FSlateRect::DoRectanglesIntersect(ViewportRect, WidgetRect);
@@ -89,7 +90,7 @@ namespace ArrowGeometryHelpers
 	{
 		check(InNavigationData.CanPaint());
 
-		const FVector2f WidgetSize = InWidgetGeometry.GetLocalSize();
+		const FVector2f WidgetSize = InWidgetGeometry.GetAbsoluteSize();
 		const FVector2f DefaultArrowTextureSize = ArrowGeometryHelpers::ArrowBrushSize;
 
 		const bool bArrowsWiderThanWidget = WidgetSize.X - (4.0f * DefaultArrowTextureSize.X) <= 0.0f;
@@ -100,31 +101,34 @@ namespace ArrowGeometryHelpers
 		const FVector2f ArrowTextureSize = bScaleDownArrows ? FVector2f(ScaledDownArrowSize) : DefaultArrowTextureSize;
 
 		FVector2f ArrowPosition;
-		FVector2f AdjustmentForOutline;
 
 		if (InNavDirection == EUINavigation::Up)
 		{
-			ArrowPosition = FVector2f((WidgetSize.X - ArrowTextureSize.X) * 0.5, 0.0f);
-			AdjustmentForOutline = FVector2f(0.0f, ArrowGeometryHelpers::LineThickness);
+			const FVector2f TopCenter = InWidgetGeometry.GetAbsolutePositionAtCoordinates(FVector2f(0.5f, 0.0f));
+			ArrowPosition = TopCenter + FVector2f(ArrowTextureSize.X * -0.5f, 0.0f);
 		}
 		else if (InNavDirection == EUINavigation::Down)
 		{
-			ArrowPosition = FVector2f((WidgetSize.X - ArrowTextureSize.X) * 0.5, WidgetSize.Y - ArrowTextureSize.Y);
-			AdjustmentForOutline = FVector2f(0.0f, -2.0f * ArrowGeometryHelpers::LineThickness);
+			const FVector2f BottomCenter = InWidgetGeometry.GetAbsolutePositionAtCoordinates(FVector2f(0.5f, 1.0f));
+			ArrowPosition = BottomCenter + FVector2f(ArrowTextureSize.X * -0.5f, -1.0f * ArrowTextureSize.Y);
 		}
 		else if (InNavDirection == EUINavigation::Left)
 		{
-			ArrowPosition = FVector2f(0.0f, (WidgetSize.Y - ArrowTextureSize.Y) * 0.5f);
-			AdjustmentForOutline = FVector2f(ArrowGeometryHelpers::LineThickness, 0.0f);
+			const FVector2f LeftMiddle = InWidgetGeometry.GetAbsolutePositionAtCoordinates(FVector2f(0.0f, 0.5f));
+			ArrowPosition = LeftMiddle + FVector2f(0.0f, ArrowTextureSize.Y * -0.5f);
 		}
 		else if (InNavDirection == EUINavigation::Right)
 		{
-			ArrowPosition = FVector2f(WidgetSize.X - ArrowTextureSize.X, (WidgetSize.Y - ArrowTextureSize.Y) * 0.5f);
-			AdjustmentForOutline = FVector2f(-2.0f * ArrowGeometryHelpers::LineThickness, 0.0f);
+			const FVector2f RightMiddle = InWidgetGeometry.GetAbsolutePositionAtCoordinates(FVector2f(1.0f, 0.5f));
+			ArrowPosition = RightMiddle + FVector2f(-1.0f * ArrowTextureSize.X, ArrowTextureSize.Y * -0.5f);
+		}
+		else
+		{
+			ArrowPosition = InWidgetGeometry.GetAbsolutePositionAtCoordinates(FVector2f(0.5f, 0.5f));
 		}
 
-		const FGeometry ViewportWidgetGeometry = InNavigationData.ViewportWidget->GetCachedGeometry().MakeChild(InNavigationData.WindowWidget->GetPositionInScreen() * -1);
-		const FVector2f ArrowPositionInViewport = ViewportWidgetGeometry.AbsoluteToLocal(InWidgetGeometry.LocalToAbsolute(ArrowPosition)) + AdjustmentForOutline;
+		const FGeometry ViewportWidgetGeometry = InNavigationData.ViewportWidget->GetPaintSpaceGeometry();
+		const FVector2f ArrowPositionInViewport = ViewportWidgetGeometry.AbsoluteToLocal(ArrowPosition);
 
 		return ViewportWidgetGeometry.MakeChild(ArrowTextureSize, FSlateLayoutTransform(1.0f, ArrowPositionInViewport));
 	}
@@ -136,8 +140,8 @@ namespace ArrowGeometryHelpers
 		TSharedRef<SWindow> WindowWidget = InNavigationData.WindowWidget.ToSharedRef();
 		TSharedPtr<SWidget> ViewportWidget = InNavigationData.ViewportWidget.ToSharedRef();
 
-		const FGeometry ViewportWidgetGeometry = ViewportWidget->GetCachedGeometry().MakeChild(InNavigationData.WindowWidget->GetPositionInScreen() * -1);
-		const FGeometry WidgetGeometry = InWidget->GetCachedGeometry().MakeChild(InNavigationData.WindowWidget->GetPositionInScreen() * -1);
+		const FGeometry ViewportWidgetGeometry = ViewportWidget->GetPaintSpaceGeometry();
+		const FGeometry WidgetGeometry = InWidget->GetPaintSpaceGeometry();
 		const FVector2f WidgetPosition = WidgetGeometry.GetAbsolutePosition();
 
 		const bool bIsFocusArrow = InArrowType == FNavigationDebuggingStyle::EArrowType::Focus;
@@ -184,14 +188,20 @@ namespace ArrowGeometryHelpers
 
 	static void DrawArrow(const FGeometry& InArrowGeometry, const FSlateBrush* InArrowBrush, float ColorAlpha, FSlateWindowElementList& OutDrawElements, int32 LayerId)
 	{
-		FSlateDrawElement::MakeBox(
-			OutDrawElements,
-			LayerId,
-			InArrowGeometry.ToPaintGeometry(),
-			InArrowBrush,
-			ESlateDrawEffect::None,
-			FLinearColor::White.CopyWithNewOpacity(ColorAlpha)
-		);
+		const FVector2f ArrowSize = InArrowGeometry.GetAbsoluteSize();
+		const bool bIsTooSmall = ArrowSize.X < MinimumArrowGeometrySizeToPaint || ArrowSize.Y < MinimumArrowGeometrySizeToPaint;
+
+		if (!bIsTooSmall)
+		{
+			FSlateDrawElement::MakeBox(
+				OutDrawElements,
+				LayerId,
+				InArrowGeometry.ToPaintGeometry(),
+				InArrowBrush,
+				ESlateDrawEffect::None,
+				FLinearColor::White.CopyWithNewOpacity(ColorAlpha)
+			);
+		}
 	}
 }
 
@@ -486,7 +496,7 @@ int32 FSlateNavigationDebugger::PaintCandidateWidgets(const FSlateUserNavigation
 					{
 						HighlightedCandidateWidgets.FindOrAdd(CandidateWidgetRef).Add(Info.Direction);
 
-						const FGeometry CandidateWidgetGeometry = CandidateWidgetRef->GetCachedGeometry().MakeChild(InNavigationData.WindowWidget->GetPositionInScreen() * -1);						
+						const FGeometry CandidateWidgetGeometry = CandidateWidgetRef->GetPaintSpaceGeometry();
 						if (!bAlreadyOutlined)
 						{
 							const FLinearColor OutlineColor = FNavigationDebuggingStyle::Get().GetTargetWidgetColor().CopyWithNewOpacity(ArrowGeometryHelpers::CandidateColorAlpha);
@@ -521,7 +531,7 @@ int32 FSlateNavigationDebugger::PaintFocusTargetArrowsAndLines(const TPair<EUINa
 		const bool bTargetWidgetInViewport = !ArrowGeometryHelpers::IsWidgetOutOfViewport(InNavigationData, FocusedToDestination.Value);
 		if (bTargetWidgetInViewport)
 		{
-			const FGeometry TargetWidgetGeometry = FocusedToDestination.Value->GetCachedGeometry().MakeChild(InNavigationData.WindowWidget->GetPositionInScreen() * -1);
+			const FGeometry TargetWidgetGeometry = FocusedToDestination.Value->GetPaintSpaceGeometry();
 			const FGeometry TargetArrowGeometry = PaintOnScreenWidgetArrow(InNavigationData, FocusedToDestination.Key, TargetWidgetGeometry, FNavigationDebuggingStyle::EArrowType::Target, OutDrawElements, ++LayerId);
 			TargetWidgetArrowPosition = TargetArrowGeometry.GetAbsolutePositionAtCoordinates(ArrowGeometryHelpers::GetArrowCoordinatesForLine(ArrowGeometryHelpers::GetOpposingDirection(FocusedToDestination.Key)));
 		}
@@ -538,7 +548,7 @@ int32 FSlateNavigationDebugger::PaintFocusTargetArrowsAndLines(const TPair<EUINa
 		const bool bFocusWidgetInViewport = !ArrowGeometryHelpers::IsWidgetOutOfViewport(InNavigationData, InNavigationData.FocusedWidget.ToSharedRef());
 		if (bFocusWidgetInViewport)
 		{
-			const FGeometry FocusedWidgetGeometry = InNavigationData.FocusedWidget->GetCachedGeometry().MakeChild(InNavigationData.WindowWidget->GetPositionInScreen() * -1);
+			const FGeometry FocusedWidgetGeometry = InNavigationData.FocusedWidget->GetPaintSpaceGeometry();
 			const FGeometry FocusArrowGeometry = PaintOnScreenWidgetArrow(InNavigationData, FocusedToDestination.Key, FocusedWidgetGeometry, FNavigationDebuggingStyle::EArrowType::Focus, OutDrawElements, ++LayerId);
 			FocusedWidgetArrowPosition = FocusArrowGeometry.GetAbsolutePositionAtCoordinates(ArrowGeometryHelpers::GetArrowCoordinatesForLine(FocusedToDestination.Key));
 		}
