Splitting Candidate Widget CVar into 4 different CVars

#Synced-CL 51546663

#Synced-CL 51546663

diff --git a/Engine/Source/Runtime/Slate/Private/Navigation/SlateNavigationDebugger.cpp b/Engine/Source/Runtime/Slate/Private/Navigation/SlateNavigationDebugger.cpp
--- a/Engine/Source/Runtime/Slate/Private/Navigation/SlateNavigationDebugger.cpp
+++ b/Engine/Source/Runtime/Slate/Private/Navigation/SlateNavigationDebugger.cpp
@@ -10,12 +10,83 @@
 #include "Framework/Application/SlateApplication.h"
 #include "Framework/Application/SlateUser.h"
 
-static TAutoConsoleVariable<bool> CVarDisplayCandidateWidgets(
-	TEXT("Slate.NavigationDebugging.DisplayCandidateWidgets"),
+namespace DebugCandidateWidgets
+{
+
+static TAutoConsoleVariable<bool> CVarDisplayCandidateWidgets_Right(
+	TEXT("Slate.NavigationDebugging.CandidateWidgets.Right"),
+	false,
+	TEXT("Enable to display candidate widgets for EUINavigation::Right of the focused widget when `Slate.NavigationDebugging` is enabled"),
+	ECVF_Default);
+
+static TAutoConsoleVariable<bool> CVarDisplayCandidateWidgets_Left(
+	TEXT("Slate.NavigationDebugging.CandidateWidgets.Left"),
+	false,
+	TEXT("Enable to display candidate widgets for EUINavigation::Left of the focused widget when `Slate.NavigationDebugging` is enabled"),
+	ECVF_Default);
+
+static TAutoConsoleVariable<bool> CVarDisplayCandidateWidgets_Up(
+	TEXT("Slate.NavigationDebugging.CandidateWidgets.Up"),
+	false,
+	TEXT("Enable to display candidate widgets for EUINavigation::Up of the focused widget when `Slate.NavigationDebugging` is enabled"),
+	ECVF_Default);
+
+static TAutoConsoleVariable<bool> CVarDisplayCandidateWidgets_Down(
+	TEXT("Slate.NavigationDebugging.CandidateWidgets.Down"),
 	false,
-	TEXT("Enable to see candidate widgets"),
+	TEXT("Enable to display candidate widgets for EUINavigation::Down of the focused widget when `Slate.NavigationDebugging` is enabled"),
 	ECVF_Default);
 
+
+static TMap<EUINavigation, TAutoConsoleVariable<bool>*>& GetDirectionToConsoleVariableMap()
+{
+	static TMap<EUINavigation, TAutoConsoleVariable<bool>*> DirectionToCVar = {
+		{ EUINavigation::Left, &CVarDisplayCandidateWidgets_Left },
+		{ EUINavigation::Right, &CVarDisplayCandidateWidgets_Right },
+		{ EUINavigation::Up, &CVarDisplayCandidateWidgets_Up },
+		{ EUINavigation::Down, &CVarDisplayCandidateWidgets_Down },
+	};
+
+	return DirectionToCVar;
+}
+
+FAutoConsoleCommand ConsoleCommandDisplayCandidateWidgets_All(
+	TEXT("Slate.NavigationDebugging.CandidateWidgets.All"),
+	TEXT("Set enabled to display candidate widgets for 4 [Up, Down, Left, Right] Navigation Directions of the focused widget when `Slate.NavigationDebugging` is enabled"),
+	FConsoleCommandWithArgsDelegate::CreateLambda([](const TArray<FString>& InArgs)
+												  {
+													  if (!InArgs.IsEmpty())
+													  {
+														  const bool bEnableAll = InArgs[0].ToBool();
+														  
+														  for (TPair<EUINavigation, TAutoConsoleVariable<bool>*>& DirectionToCVarPair : GetDirectionToConsoleVariableMap())
+														  {
+															  DirectionToCVarPair.Value->AsVariable()->Set(bEnableAll, EConsoleVariableFlags::ECVF_SetByConsole);
+														  }
+													  }
+												  }));
+
+static bool ShouldDisplayCandidateWidgetsFor(const EUINavigation InDirection)
+{
+	const TAutoConsoleVariable<bool>* const * ConsoleObject = GetDirectionToConsoleVariableMap().Find(InDirection);
+	return ConsoleObject && (*ConsoleObject)->GetValueOnAnyThread();
+}
+
+static bool ShouldDisplayAnyCandidateWidgets()
+{
+	for (TPair<EUINavigation, TAutoConsoleVariable<bool>*>& DirectionToCVarPair : GetDirectionToConsoleVariableMap())
+	{
+		if (DirectionToCVarPair.Value->GetValueOnAnyThread())
+		{
+			return true;
+		}
+	}
+
+	return false;
+}
+
+}
+
 namespace ArrowGeometryHelpers
 {
 	constexpr bool bAntiAlias = true;
@@ -473,7 +544,7 @@ void FSlateNavigationDebugger::PaintDashedLine(const FSlateUserNavigationData& I
 
 int32 FSlateNavigationDebugger::PaintCandidateWidgets(const FSlateUserNavigationData& InNavigationData, const TArray<FHittestGrid::FDebuggingFindNextFocusableWidgetArgs>& InCandidateWidgets, FSlateWindowElementList& OutDrawElements, int32 LayerId) const
 {
-	if (CVarDisplayCandidateWidgets.GetValueOnAnyThread())
+	if (DebugCandidateWidgets::ShouldDisplayAnyCandidateWidgets())
 	{
 		check(InNavigationData.CanPaint());
 
@@ -489,10 +560,12 @@ int32 FSlateNavigationDebugger::PaintCandidateWidgets(const FSlateUserNavigation
 
 					const bool bNotCandidate = Result.Result.EqualTo(HittestGridDebuggingText::Valid);
 
+					const bool bDebuggingDirection = DebugCandidateWidgets::ShouldDisplayCandidateWidgetsFor(Info.Direction);
 					const bool bAlreadyOutlined = HighlightedCandidateWidgets.Contains(CandidateWidgetRef);
 					const bool bAlreadyPaintedArrow = bAlreadyOutlined && HighlightedCandidateWidgets[CandidateWidgetRef].Contains(Info.Direction);
 
-					if (!bNotCandidate && !bAlreadyPaintedArrow)
+					const bool bDisplayCandidateWidget = bDebuggingDirection && !bNotCandidate && !bAlreadyPaintedArrow;
+					if (bDisplayCandidateWidget)
 					{
 						HighlightedCandidateWidgets.FindOrAdd(CandidateWidgetRef).Add(Info.Direction);
 
@@ -600,7 +673,7 @@ int32 FSlateNavigationDebugger::PaintForUser(FSlateUser& InSlateUser, FSlateWind
 			LayerId = PaintFocusTargetArrowsAndLines(DirectionToWidget, NavigationData, OutDrawElements, LayerId);
 		}
 		
-		// Unbind to listen to candidate widgets
+		// Unbind from listening to candidate widgets
 		SetListenToHitTestGrid(false, HitTestGridInfo, HittestGridHandle);
 	}
 
