PR #14493: Add a margin size value to radial box settings

#Synced-CL 51596070

diff --git a/Engine/Source/Runtime/Slate/Private/Widgets/Layout/SRadialBox.cpp b/Engine/Source/Runtime/Slate/Private/Widgets/Layout/SRadialBox.cpp
--- a/Engine/Source/Runtime/Slate/Private/Widgets/Layout/SRadialBox.cpp
+++ b/Engine/Source/Runtime/Slate/Private/Widgets/Layout/SRadialBox.cpp
@@ -33,6 +33,7 @@ void SRadialBox::Construct( const FArguments& InArgs )
 	bClockwiseOrder = InArgs._bClockwiseOrder;
 	SectorCentralAngle = InternalNormalizeAngle(InArgs._SectorCentralAngle);
 	AngleBetweenItems = InArgs._AngleBetweenItems;
+	MarginSize = InArgs._MarginSize;
 
 	Slots.AddSlots(MoveTemp(const_cast<TArray<FSlot::FSlotArguments>&>(InArgs._Slots)));
 }
@@ -80,7 +81,7 @@ SRadialBox::FChildArranger::FChildArranger(const SRadialBox& InRadialBox, const
 void SRadialBox::FChildArranger::Arrange()
 {
 	const int32 NumItems = RadialBox.Slots.Num();
-	const float Radius = RadialBox.PreferredWidth.Get() / 2.f;
+	const float Radius = (RadialBox.PreferredWidth.Get() / 2.f) - RadialBox.MarginSize;
 
 	int32 TargetNumItems = FMath::Max(1, NumItems);
 	const bool bIsFullCircumference = RadialBox.SectorCentralAngle >= 360.f;
diff --git a/Engine/Source/Runtime/Slate/Public/Widgets/Layout/SRadialBox.h b/Engine/Source/Runtime/Slate/Public/Widgets/Layout/SRadialBox.h
--- a/Engine/Source/Runtime/Slate/Public/Widgets/Layout/SRadialBox.h
+++ b/Engine/Source/Runtime/Slate/Public/Widgets/Layout/SRadialBox.h
@@ -41,6 +41,7 @@ public:
 		, _AngleBetweenItems(0.f)
 		, _SectorCentralAngle(360.f)
 		, _bClockwiseOrder(false)
+		, _MarginSize(0.f)
 		{
 			_Visibility = EVisibility::SelfHitTestInvisible;
 		}
@@ -69,6 +70,9 @@ public:
 		/** If true children will be arranged in a clockwise order from starting angle  */
 		SLATE_ARGUMENT(bool, bClockwiseOrder)
 
+		/** How much space to leave between the widgets and the edge of the panel */
+		SLATE_ARGUMENT(float, MarginSize)
+
 	SLATE_END_ARGS()
 
 	SLATE_API SRadialBox();
@@ -102,6 +106,7 @@ public:
 	void SetDistributeItemsEvenly(bool bInDistributeItemsEvenly) { bDistributeItemsEvenly = bInDistributeItemsEvenly; }
 	void SetClockwiseOrder(bool bInClockwiseOrder) { bClockwiseOrder = bInClockwiseOrder; }
 	SLATE_API void SetSectorCentralAngle(float InSectorCentralAngle);
+	void SetMarginSize(float InMarginSize) { MarginSize = InMarginSize; }
 
 	SLATE_API void SetUseAllottedWidth(bool bInUseAllottedWidth);
 
@@ -133,6 +138,9 @@ private:
 	/** If true children will be arranged in a clockwise order from starting angle  */
 	bool bClockwiseOrder;
 
+	/** How much space to leave between the widgets and the edge of the panel */
+	float MarginSize;
+
 	class FChildArranger;
 	friend class SRadialBox::FChildArranger;
 };
diff --git a/Engine/Source/Runtime/UMG/Private/Components/DynamicEntryBoxBase.cpp b/Engine/Source/Runtime/UMG/Private/Components/DynamicEntryBoxBase.cpp
--- a/Engine/Source/Runtime/UMG/Private/Components/DynamicEntryBoxBase.cpp
+++ b/Engine/Source/Runtime/UMG/Private/Components/DynamicEntryBoxBase.cpp
@@ -229,6 +229,7 @@ void UDynamicEntryBoxBase::SetRadialSettings(const FRadialBoxSettings& InSetting
 		StaticCastSharedPtr<SRadialBox>(MyPanelWidget)->SetClockwiseOrder(InSettings.bClockwiseOrder);
 		StaticCastSharedPtr<SRadialBox>(MyPanelWidget)->SetAngleBetweenItems(InSettings.AngleBetweenItems);
 		StaticCastSharedPtr<SRadialBox>(MyPanelWidget)->SetSectorCentralAngle(InSettings.SectorCentralAngle);
+		StaticCastSharedPtr<SRadialBox>(MyPanelWidget)->SetMarginSize(InSettings.MarginSize);
 	}
 }
 
@@ -287,6 +288,7 @@ TSharedRef<SWidget> UDynamicEntryBoxBase::RebuildWidget()
 			.StartingAngle(RadialBoxSettings.StartingAngle)
 			.bDistributeItemsEvenly(RadialBoxSettings.bDistributeItemsEvenly)
 			.bClockwiseOrder(RadialBoxSettings.bClockwiseOrder)
+			.MarginSize(RadialBoxSettings.MarginSize)
 			.AngleBetweenItems(RadialBoxSettings.AngleBetweenItems)
 			.SectorCentralAngle(RadialBoxSettings.SectorCentralAngle);
 		break;
diff --git a/Engine/Source/Runtime/UMG/Public/Components/RadialBoxSettings.h b/Engine/Source/Runtime/UMG/Public/Components/RadialBoxSettings.h
--- a/Engine/Source/Runtime/UMG/Public/Components/RadialBoxSettings.h
+++ b/Engine/Source/Runtime/UMG/Public/Components/RadialBoxSettings.h
@@ -31,12 +31,17 @@ struct FRadialBoxSettings
 	UPROPERTY(EditAnywhere, Category = "Items", meta = (EditCondition = "bDistributeItemsEvenly", ClampMin = 0, ClampMax = 360))
 	float SectorCentralAngle;
 
+	/** How much space to leave between the widgets and the edge of the panel */
+	UPROPERTY(EditAnywhere, Category = "Items", meta = (ClampMin = 0))
+	float MarginSize;
+
 	FRadialBoxSettings()
 		: StartingAngle(0.f)
 		, bDistributeItemsEvenly(true)
 		, bClockwiseOrder(false)
 		, AngleBetweenItems(0.f)
 		, SectorCentralAngle(360.f)
+		, MarginSize(0.f)
 	{
 	}
 };
\ No newline at end of file
