Fix for Justification being delayed by one frame when constructed.

#Synced-CL 51747987

diff --git a/Engine/Source/Runtime/Slate/Private/Widgets/Text/SMultiLineEditableText.cpp b/Engine/Source/Runtime/Slate/Private/Widgets/Text/SMultiLineEditableText.cpp
--- a/Engine/Source/Runtime/Slate/Private/Widgets/Text/SMultiLineEditableText.cpp
+++ b/Engine/Source/Runtime/Slate/Private/Widgets/Text/SMultiLineEditableText.cpp
@@ -701,6 +701,7 @@ int32 SMultiLineEditableText::OnPaint( const FPaintArgs& Args, const FGeometry&
 	if (bAutoWrap && SlateEditableTextLayout.GetComputedWrappingWidth() == 0)
 	{
 		QUICK_SCOPE_CYCLE_COUNTER(STAT_Slate_AutoWrapRecompute)
+		const_cast<FSlateEditableTextLayout&>(SlateEditableTextLayout).SetCachedSize(AllottedGeometry);
 		const_cast<SMultiLineEditableText*>(this)->CacheDesiredSize(GetPrepassLayoutScaleMultiplier());
 	}
 
diff --git a/Engine/Source/Runtime/Slate/Private/Widgets/Text/SlateEditableTextLayout.cpp b/Engine/Source/Runtime/Slate/Private/Widgets/Text/SlateEditableTextLayout.cpp
--- a/Engine/Source/Runtime/Slate/Private/Widgets/Text/SlateEditableTextLayout.cpp
+++ b/Engine/Source/Runtime/Slate/Private/Widgets/Text/SlateEditableTextLayout.cpp
@@ -3580,8 +3580,16 @@ int32 FSlateEditableTextLayout::OnPaint(const FPaintArgs& Args, const FGeometry&
 {
 	// Update the auto-wrap size now that we have computed paint geometry; won't take effect until text frame
 	// Note: This is done here rather than in Tick(), because Tick() doesn't get called while resizing windows, but OnPaint() does
+	const FVector2f PreviousCachedSize = CachedSize;
 	CachedSize = FVector2f(AllottedGeometry.GetLocalSize());
 
+	// If CacheSize is set for the first time, we relayout to avoid a one frame jump for non left justifications.
+	if (PreviousCachedSize == FVector2f::Zero())
+	{
+		TextLayout->SetVisibleRegion(FVector2D(CachedSize), FVector2D(ScrollOffset) * TextLayout->GetScale());
+		TextLayout->UpdateIfNeeded();
+	}
+
 	// Only paint the hint text layout if we don't have any text set
 	if (TextLayout->IsEmpty() && HintTextLayout.IsValid())
 	{
@@ -3599,6 +3607,11 @@ int32 FSlateEditableTextLayout::OnPaint(const FPaintArgs& Args, const FGeometry&
 	return LayerId;
 }
 
+void FSlateEditableTextLayout::SetCachedSize(const FGeometry& AllottedGeometry)
+{
+	CachedSize = FVector2f(AllottedGeometry.GetLocalSize());
+}
+
 void FSlateEditableTextLayout::CacheDesiredSize(float LayoutScaleMultiplier)
 {
 	const float FontMaxCharHeight = FTextEditHelper::GetFontHeight(TextStyle.Font);
diff --git a/Engine/Source/Runtime/Slate/Public/Widgets/Text/SlateEditableTextLayout.h b/Engine/Source/Runtime/Slate/Public/Widgets/Text/SlateEditableTextLayout.h
--- a/Engine/Source/Runtime/Slate/Public/Widgets/Text/SlateEditableTextLayout.h
+++ b/Engine/Source/Runtime/Slate/Public/Widgets/Text/SlateEditableTextLayout.h
@@ -484,6 +484,8 @@ public:
 	}
 	SLATE_API int32 OnPaint(const FPaintArgs& Args, const FGeometry& AllottedGeometry, const FSlateRect& MyCullingRect, FSlateWindowElementList& OutDrawElements, int32 LayerId, const FWidgetStyle& InWidgetStyle, bool bParentEnabled) const;
 
+	SLATE_API void SetCachedSize(const FGeometry& AllottedGeometry);
+
 	SLATE_API void CacheDesiredSize(float LayoutScaleMultiplier);
 
 	SLATE_API FVector2D ComputeDesiredSize(float LayoutScaleMultiplier) const;
