Fix to redraw widgets when their focus changes

#Synced-CL 51747987

diff --git a/Engine/Source/Runtime/Slate/Private/Framework/Application/SlateApplication.cpp b/Engine/Source/Runtime/Slate/Private/Framework/Application/SlateApplication.cpp
--- a/Engine/Source/Runtime/Slate/Private/Framework/Application/SlateApplication.cpp
+++ b/Engine/Source/Runtime/Slate/Private/Framework/Application/SlateApplication.cpp
@@ -3081,6 +3081,12 @@ bool FSlateApplication::SetUserFocus(FSlateUser& User, const FWidgetPath& InFocu
 		// Switch worlds for widgets in the old path
 		FScopedSwitchWorldHack SwitchWorld(OldFocusedWidgetPath.Window.Pin());
 
+#if PLATFORM_UI_NEEDS_FOCUS_OUTLINES
+		if (OldFocusedWidget->CanSupportFocus() && OldFocusedWidget->SupportsKeyboardFocus())
+		{
+			OldFocusedWidget->Invalidate(EInvalidateWidgetReason::Paint);
+		}
+#endif
 #if WITH_SLATE_DEBUGGING
 		FSlateDebugging::BroadcastFocusLost(FocusEvent, OldFocusedWidgetPath, OldFocusedWidget, NewFocusedWidgetPath, NewFocusedWidget);
 #endif
@@ -3112,6 +3118,12 @@ bool FSlateApplication::SetUserFocus(FSlateUser& User, const FWidgetPath& InFocu
 		
 		const FArrangedWidget& WidgetToFocus = NewFocusedWidgetPath.Widgets.Last();
 
+#if PLATFORM_UI_NEEDS_FOCUS_OUTLINES
+		if (NewFocusedWidget->CanSupportFocus() && NewFocusedWidget->SupportsKeyboardFocus())
+		{
+			NewFocusedWidget->Invalidate(EInvalidateWidgetReason::Paint);
+		}
+#endif
 #if WITH_SLATE_DEBUGGING
 		FSlateDebugging::BroadcastFocusReceived(FocusEvent, OldFocusedWidgetPath, OldFocusedWidget, NewFocusedWidgetPath, NewFocusedWidget);
 #endif
diff --git a/Engine/Source/Runtime/SlateCore/Public/Widgets/SWidget.h b/Engine/Source/Runtime/SlateCore/Public/Widgets/SWidget.h
--- a/Engine/Source/Runtime/SlateCore/Public/Widgets/SWidget.h
+++ b/Engine/Source/Runtime/SlateCore/Public/Widgets/SWidget.h
@@ -914,6 +914,13 @@ public:
 	virtual FChildren* Debug_GetChildrenForReflector() { return GetAllChildren(); }
 #endif
 
+	/**
+	 * Checks is this widget can ever support keyboard focus.
+	 *
+	 * @return  True if this widget can ever support keyboard focus
+	 */
+	bool CanSupportFocus() const { return bCanSupportFocus; }
+
 	/**
 	 * Checks to see if this widget supports keyboard focus.  Override this in derived classes.
 	 *
