New IOS TextField implementation that matches Android look

diff --git a/Engine/Source/Runtime/Core/Private/IOS/IOSPlatformMisc.cpp b/Engine/Source/Runtime/Core/Private/IOS/IOSPlatformMisc.cpp
--- a/Engine/Source/Runtime/Core/Private/IOS/IOSPlatformMisc.cpp
+++ b/Engine/Source/Runtime/Core/Private/IOS/IOSPlatformMisc.cpp
@@ -32,6 +32,7 @@
 #include "Misc/ScopeExit.h"
 #include "Misc/SecureHash.h"
 #include "Modules/ModuleManager.h"
+#include <atomic>
 #include "Templates/Function.h"
 #include "Templates/UniquePtr.h"
 #include "Containers/StringConv.h"
@@ -86,6 +87,9 @@ void (* GMemoryWarningHandler)(const FGenericMemoryWarningContext& Context) = NU
 /** global for showing the splash screen */
 bool GShowSplashScreen = true;
 
+/** global for querying virtual keyboard input field height (without including keyboard itself) */
+std::atomic<float> GIOSVirtualKeyboardInputFieldHeight{0.0f};
+
 #if !UE_BUILD_SHIPPING
 /** global for showing debug console */
 bool GDebugConsoleOpen = false;
@@ -2413,3 +2417,8 @@ public:
 		return false;
 	}
 } GIOSExec;
+
+float FIOSPlatformMisc::GetVirtualKeyboardInputHeight()
+{
+	return GIOSVirtualKeyboardInputFieldHeight.load(std::memory_order_relaxed);
+}
diff --git a/Engine/Source/Runtime/Core/Public/IOS/IOSPlatformMisc.h b/Engine/Source/Runtime/Core/Public/IOS/IOSPlatformMisc.h
--- a/Engine/Source/Runtime/Core/Public/IOS/IOSPlatformMisc.h
+++ b/Engine/Source/Runtime/Core/Public/IOS/IOSPlatformMisc.h
@@ -246,6 +246,8 @@ struct CORE_API FIOSPlatformMisc : public FApplePlatformMisc
 		ColorGamut = EDisplayColorGamut::sRGB_D65;
 	}
 
+	static float GetVirtualKeyboardInputHeight();
+
 	static int32 GetMaxRefreshRate();
 
     // added these for now because Crashlytics doesn't properly break up different callstacks all ending in UE_LOG(LogXXX, Fatal, ...)
diff --git a/Engine/Source/Runtime/Slate/Private/Framework/Text/IOS/IOSPlatformTextField.cpp b/Engine/Source/Runtime/Slate/Private/Framework/Text/IOS/IOSPlatformTextField.cpp
--- a/Engine/Source/Runtime/Slate/Private/Framework/Text/IOS/IOSPlatformTextField.cpp
+++ b/Engine/Source/Runtime/Slate/Private/Framework/Text/IOS/IOSPlatformTextField.cpp
@@ -6,8 +6,18 @@
 #include "IOS/IOSAsyncTask.h"
 #include "IOS/IOSView.h"
 #include "Misc/ConfigCacheIni.h"
+#include <atomic>
 #include "Widgets/Input/IVirtualKeyboardEntry.h"
 
+static int32 GIOSTextFieldImplementation = 0;
+FAutoConsoleVariableRef CVarIOSTextFieldImplementation(
+		TEXT("ios.TextFieldImplementation"),
+		GIOSTextFieldImplementation,
+		TEXT("iOS text field implementation. 0 = UIAlertController popup (default), 1 = floating text field above keyboard."),
+		ECVF_Default);
+
+extern std::atomic<float> GIOSVirtualKeyboardInputFieldHeight;
+
 namespace
 {
 	void GetKeyboardConfig(TSharedPtr<IVirtualKeyboardEntry> TextEntryWidget, FKeyboardConfig& KeyboardConfig)
@@ -46,6 +56,7 @@ namespace
 
 FIOSPlatformTextField::FIOSPlatformTextField()
 	: TextField( nullptr )
+	, TextField2( nullptr )
 {
 }
 
@@ -64,6 +75,26 @@ FIOSPlatformTextField::~FIOSPlatformTextField()
 			}
 		});
 	}
+	if (TextField2 != nullptr)
+	{
+        UE_LOG(LogIOS, Log, TEXT("Deleting text field2: %p"), TextField2);
+        SlateTextField2* LocalTextField2 = TextField2;
+        TextField2 = nullptr;
+		dispatch_async(dispatch_get_main_queue(), ^{
+            NSLog(@"Finally releasing text field2 %@", LocalTextField2);
+			if (LocalTextField2 != nullptr && [LocalTextField2 respondsToSelector:@selector(hide)])
+			{
+            	[LocalTextField2 hide];
+				[LocalTextField2 release];
+			}
+		});
+	}
+}
+
+bool FIOSPlatformTextField::AllowMoveCursor()
+{
+	// OS field owns the cursor since changes in slate would not be replicated in the text field
+	return GIOSTextFieldImplementation != 1;
 }
 
 void FIOSPlatformTextField::ShowVirtualKeyboard(bool bShow, int32 UserIndex, TSharedPtr<IVirtualKeyboardEntry> TextEntryWidget)
@@ -98,7 +129,9 @@ void FIOSPlatformTextField::ShowVirtualKeyboard(bool bShow, int32 UserIndex, TSh
 	}
 	else
 	{
-		if(bShow)
+		const bool bUseNewImpl = (GIOSTextFieldImplementation == 1);
+
+		if (bShow)
 		{
 			// capture some gamethread strings before we toss over to main thread
 			NSString* TextContents = [NSString stringWithFString : TextEntryWidget->GetText().ToString()];
@@ -107,28 +140,59 @@ void FIOSPlatformTextField::ShowVirtualKeyboard(bool bShow, int32 UserIndex, TSh
 			GetKeyboardConfig(TextEntryWidget, KeyboardConfig);
 
 			// these functions must be run on the main thread
-			dispatch_async(dispatch_get_main_queue(),^ {
-				if (TextField == nullptr)
-				{
-					TextField = [[[SlateTextField alloc]init] retain];
-				}
-				[TextField show:TextEntryWidget text:TextContents placeholder:PlaceholderContents keyboardConfig:KeyboardConfig];
-			});
+			if (bUseNewImpl)
+			{
+				dispatch_async(dispatch_get_main_queue(), ^{
+					if (TextField2 == nullptr)
+					{
+						TextField2 = [[SlateTextField2 alloc] init];
+					}
+					[TextField2 show:TextEntryWidget text:TextContents placeholder:PlaceholderContents keyboardConfig:KeyboardConfig];
+				});
+			}
+			else
+			{
+				dispatch_async(dispatch_get_main_queue(), ^{
+					if (TextField == nullptr)
+					{
+						TextField = [[[SlateTextField alloc] init] retain];
+					}
+					[TextField show:TextEntryWidget text:TextContents placeholder:PlaceholderContents keyboardConfig:KeyboardConfig];
+				});
+			}
 		}
         else
         {
 #if !PLATFORM_TVOS
-			if (TextField != nullptr && [TextField hasTextWidget])
+			if (bUseNewImpl)
+			{
+				if (TextField2 != nullptr && [TextField2 hasTextWidget])
+				{
+	                UE_LOG(LogIOS, Log, TEXT("Hiding field2: %p"), TextField2);
+	                SlateTextField2* LocalTextField2 = TextField2;
+	                dispatch_async(dispatch_get_main_queue(), ^{
+	                    NSLog(@"Finally hiding text field2 %@", LocalTextField2);
+	                    if (LocalTextField2 != nullptr)
+	                    {
+	                        [LocalTextField2 hide];
+	                    }
+	                });
+				}
+			}
+			else
 			{
-                UE_LOG(LogIOS, Log, TEXT("Hiding field: %p"), TextField);
-                SlateTextField* LocalTextField = TextField;
-                dispatch_async(dispatch_get_main_queue(), ^{
-                    NSLog(@"Finally releasing text field %@", LocalTextField);
-                    if (LocalTextField != nullptr)
-                    {
-                        [LocalTextField hide];
-                    }
-                });
+				if (TextField != nullptr && [TextField hasTextWidget])
+				{
+	                UE_LOG(LogIOS, Log, TEXT("Hiding field: %p"), TextField);
+	                SlateTextField* LocalTextField = TextField;
+	                dispatch_async(dispatch_get_main_queue(), ^{
+	                    NSLog(@"Finally releasing text field %@", LocalTextField);
+	                    if (LocalTextField != nullptr)
+	                    {
+	                        [LocalTextField hide];
+	                    }
+	                });
+				}
 			}
 #endif
         }
@@ -336,3 +400,192 @@ void FIOSPlatformTextField::ShowVirtualKeyboard(bool bShow, int32 UserIndex, TSh
 }
 
 @end
+
+@implementation SlateTextField2
+
+-(id)init
+{
+	self = [super init];
+
+	if (self)
+	{
+		CGSize ScreenSize = [UIScreen mainScreen].bounds.size;
+		const CGFloat TextFieldHeight = 46.0f;
+
+		InputTextField = [[UITextField alloc] initWithFrame:CGRectMake(0, ScreenSize.height, ScreenSize.width, TextFieldHeight)];
+		InputTextField.borderStyle = UITextBorderStyleNone;
+		InputTextField.backgroundColor = [[UIColor secondarySystemBackgroundColor] colorWithAlphaComponent:0.98];
+		InputTextField.layer.cornerRadius = 19.0f;
+		InputTextField.layer.masksToBounds = YES;
+		InputTextField.layer.borderWidth = 1.0f;
+		InputTextField.layer.borderColor = [UIColor systemGray4Color].CGColor;
+		InputTextField.returnKeyType = UIReturnKeyDone;
+		InputTextField.clearButtonMode = UITextFieldViewModeWhileEditing;
+		InputTextField.delegate = self;
+		[InputTextField addTarget:self action:@selector(onTextChanged:) forControlEvents:UIControlEventEditingChanged];
+
+		UIView* LeftSpacer = [[UIView alloc] initWithFrame:CGRectMake(0, 0, 12, 1)];
+		InputTextField.leftView = LeftSpacer;
+		InputTextField.leftViewMode = UITextFieldViewModeAlways;
+		[LeftSpacer release];
+
+		// Register keyboard frame notification to position the text field above the keyboard
+		[[NSNotificationCenter defaultCenter]
+			addObserver:self
+			selector:@selector(keyboardWillShow:)
+			name:UIKeyboardWillShowNotification
+			object:nil];
+	}
+
+	return self;
+}
+
+-(void)dealloc
+{
+	[[NSNotificationCenter defaultCenter] removeObserver:self name:UIKeyboardWillShowNotification object:nil];
+
+	[InputTextField removeFromSuperview];
+	[InputTextField release];
+	InputTextField = nil;
+
+	[super dealloc];
+}
+
+-(void)keyboardWillShow:(NSNotification*)Notification
+{
+	if (!TextWidget.IsValid())
+	{
+		return;
+	}
+
+	NSDictionary* Info = [Notification userInfo];
+	CGRect KBFrame = [[Info objectForKey:UIKeyboardFrameEndUserInfoKey] CGRectValue];
+
+	CGFloat FieldHeight = InputTextField.frame.size.height;
+	CGFloat ScreenWidth = [UIScreen mainScreen].bounds.size.width;
+
+	const CGFloat InnerInset = 8.0f;
+	const CGFloat BottomInset = 4.0f;
+	UIEdgeInsets SafeArea = [IOSAppDelegate GetDelegate].IOSController.view.safeAreaInsets;
+	CGFloat LeftInset = SafeArea.left + InnerInset;
+	CGFloat RightInset = SafeArea.right + InnerInset;
+
+	InputTextField.frame = CGRectMake(LeftInset, KBFrame.origin.y - FieldHeight - BottomInset, ScreenWidth - LeftInset - RightInset, FieldHeight);
+
+	// Just the sum of the text field and inset, keyboard height is not included
+	GIOSVirtualKeyboardInputFieldHeight.store(FieldHeight + BottomInset, std::memory_order_relaxed);
+}
+
+-(void)onTextChanged:(UITextField*)InTextField
+{
+	if (!TextWidget.IsValid())
+	{
+		return;
+	}
+
+	FText CurrentText = FText::FromString(InTextField.text);
+	TWeakPtr<IVirtualKeyboardEntry> LocalWidget = TextWidget;
+
+	FAppleAsyncTask* AsyncTask = [[FAppleAsyncTask alloc] init];
+	AsyncTask.GameThreadCallback = ^ bool(void)
+	{
+		if (LocalWidget.IsValid())
+		{
+			LocalWidget.Pin()->SetTextFromVirtualKeyboard(CurrentText, ETextEntryType::TextEntryUpdated);
+		}
+		return true;
+	};
+	[AsyncTask FinishedTask];
+}
+
+-(void)textFieldDidChangeSelection:(UITextField*)InTextField
+{
+	if (!TextWidget.IsValid())
+	{
+		return;
+	}
+
+	UITextRange* SelectedRange = InTextField.selectedTextRange;
+	if (SelectedRange == nil)
+	{
+		return;
+	}
+
+	int SelStart = (int)[InTextField offsetFromPosition:InTextField.beginningOfDocument toPosition:SelectedRange.start];
+	int SelEnd = (int)[InTextField offsetFromPosition:InTextField.beginningOfDocument toPosition:SelectedRange.end];
+	TWeakPtr<IVirtualKeyboardEntry> LocalWidget = TextWidget;
+
+	FAppleAsyncTask* AsyncTask = [[FAppleAsyncTask alloc] init];
+	AsyncTask.GameThreadCallback = ^ bool(void)
+	{
+		if (LocalWidget.IsValid())
+		{
+			LocalWidget.Pin()->SetSelectionFromVirtualKeyboard(SelStart, SelEnd);
+		}
+		return true;
+	};
+	[AsyncTask FinishedTask];
+}
+
+-(BOOL)textFieldShouldReturn:(UITextField*)InTextField
+{
+	FText LocalEntry = FText::FromString(InTextField.text);
+	TWeakPtr<IVirtualKeyboardEntry> LocalWidget = TextWidget;
+
+	[self hide];
+
+	FAppleAsyncTask* AsyncTask = [[FAppleAsyncTask alloc] init];
+	AsyncTask.GameThreadCallback = ^ bool(void)
+	{
+		if (LocalWidget.IsValid())
+		{
+			LocalWidget.Pin()->SetTextFromVirtualKeyboard(LocalEntry, ETextEntryType::TextEntryAccepted);
+		}
+		return true;
+	};
+	[AsyncTask FinishedTask];
+
+	return YES;
+}
+
+-(void)show:(TSharedPtr<IVirtualKeyboardEntry>)InTextWidget text:(NSString*)TextContents placeholder:(NSString*)PlaceholderContents keyboardConfig:(FKeyboardConfig)KeyboardConfig
+{
+	TextWidget = InTextWidget;
+
+	InputTextField.text = TextContents;
+	InputTextField.attributedPlaceholder = [[NSAttributedString alloc] 
+		initWithString:PlaceholderContents 
+		attributes:@{NSForegroundColorAttributeName: [UIColor secondaryLabelColor]}];
+	InputTextField.keyboardType = KeyboardConfig.KeyboardType;
+	InputTextField.autocorrectionType = KeyboardConfig.AutocorrectionType;
+	InputTextField.autocapitalizationType = KeyboardConfig.AutocapitalizationType;
+	InputTextField.secureTextEntry = KeyboardConfig.bSecureTextEntry;
+
+	// Attach the text field to the window on first use.
+	if (InputTextField.superview == nil)
+	{
+		UIWindow* AppWindow = [IOSAppDelegate GetDelegate].IOSController.view.window;
+		[AppWindow addSubview:InputTextField];
+	}
+
+	[InputTextField becomeFirstResponder];
+}
+
+-(void)hide
+{
+	TextWidget = nullptr;
+
+	GIOSVirtualKeyboardInputFieldHeight.store(0.0f, std::memory_order_relaxed);
+
+	CGSize ScreenSize = [UIScreen mainScreen].bounds.size;
+	InputTextField.frame = CGRectMake(0, ScreenSize.height, InputTextField.frame.size.width, InputTextField.frame.size.height);
+
+	[InputTextField resignFirstResponder];
+}
+
+-(bool)hasTextWidget
+{
+	return TextWidget.IsValid();
+}
+
+@end
diff --git a/Engine/Source/Runtime/Slate/Public/Framework/Text/IOS/IOSPlatformTextField.h b/Engine/Source/Runtime/Slate/Public/Framework/Text/IOS/IOSPlatformTextField.h
--- a/Engine/Source/Runtime/Slate/Public/Framework/Text/IOS/IOSPlatformTextField.h
+++ b/Engine/Source/Runtime/Slate/Public/Framework/Text/IOS/IOSPlatformTextField.h
@@ -12,6 +12,7 @@
 class IVirtualKeyboardEntry;
 
 @class SlateTextField;
+@class SlateTextField2;
 
 class FIOSPlatformTextField : public IPlatformTextField
 {
@@ -20,10 +21,11 @@ public:
 	virtual ~FIOSPlatformTextField();
 
 	virtual void ShowVirtualKeyboard(bool bShow, int32 UserIndex, TSharedPtr<IVirtualKeyboardEntry> TextEntryWidget) override;
-	virtual bool AllowMoveCursor() override { return true; };
+	virtual bool AllowMoveCursor() override;
 
 private:
 	SlateTextField* TextField;
+	SlateTextField2* TextField2;
 };
 
 typedef FIOSPlatformTextField FPlatformTextField;
@@ -48,3 +50,16 @@ typedef FIOSPlatformTextField FPlatformTextField;
 -(bool)hasTextWidget;
 
 @end
+
+@interface SlateTextField2 : NSObject <UITextFieldDelegate>
+{
+	TWeakPtr<IVirtualKeyboardEntry> TextWidget;
+
+	UITextField* InputTextField;
+}
+
+-(void)show:(TSharedPtr<IVirtualKeyboardEntry>)InTextWidget text:(NSString*)TextContents placeholder:(NSString*)PlaceholderContents keyboardConfig:(FKeyboardConfig)KeyboardConfig;
+-(void)hide;
+-(bool)hasTextWidget;
+
+@end
