iOS text field build fix for TVOS

#Synced-CL 51927141

diff --git a/Engine/Source/Runtime/Slate/Private/Framework/Text/IOS/IOSPlatformTextField.cpp b/Engine/Source/Runtime/Slate/Private/Framework/Text/IOS/IOSPlatformTextField.cpp
--- a/Engine/Source/Runtime/Slate/Private/Framework/Text/IOS/IOSPlatformTextField.cpp
+++ b/Engine/Source/Runtime/Slate/Private/Framework/Text/IOS/IOSPlatformTextField.cpp
@@ -75,6 +75,7 @@ FIOSPlatformTextField::~FIOSPlatformTextField()
 			}
 		});
 	}
+#if !PLATFORM_TVOS
 	if (TextField2 != nullptr)
 	{
         UE_LOG(LogIOS, Log, TEXT("Deleting text field2: %p"), TextField2);
@@ -89,6 +90,7 @@ FIOSPlatformTextField::~FIOSPlatformTextField()
 			}
 		});
 	}
+#endif
 }
 
 bool FIOSPlatformTextField::AllowMoveCursor()
@@ -129,7 +131,11 @@ void FIOSPlatformTextField::ShowVirtualKeyboard(bool bShow, int32 UserIndex, TSh
 	}
 	else
 	{
+#if PLATFORM_TVOS
+		const bool bUseNewImpl = false;
+#else
 		const bool bUseNewImpl = (GIOSTextFieldImplementation == 1);
+#endif
 
 		if (bShow)
 		{
@@ -142,6 +148,7 @@ void FIOSPlatformTextField::ShowVirtualKeyboard(bool bShow, int32 UserIndex, TSh
 			// these functions must be run on the main thread
 			if (bUseNewImpl)
 			{
+#if !PLATFORM_TVOS
 				dispatch_async(dispatch_get_main_queue(), ^{
 					if (TextField2 == nullptr)
 					{
@@ -149,6 +156,7 @@ void FIOSPlatformTextField::ShowVirtualKeyboard(bool bShow, int32 UserIndex, TSh
 					}
 					[TextField2 show:TextEntryWidget text:TextContents placeholder:PlaceholderContents keyboardConfig:KeyboardConfig];
 				});
+#endif
 			}
 			else
 			{
@@ -401,6 +409,7 @@ void FIOSPlatformTextField::ShowVirtualKeyboard(bool bShow, int32 UserIndex, TSh
 
 @end
 
+#if !PLATFORM_TVOS
 @implementation SlateTextField2
 
 -(id)init
@@ -589,3 +598,4 @@ void FIOSPlatformTextField::ShowVirtualKeyboard(bool bShow, int32 UserIndex, TSh
 }
 
 @end
+#endif
diff --git a/Engine/Source/Runtime/Slate/Public/Framework/Text/IOS/IOSPlatformTextField.h b/Engine/Source/Runtime/Slate/Public/Framework/Text/IOS/IOSPlatformTextField.h
--- a/Engine/Source/Runtime/Slate/Public/Framework/Text/IOS/IOSPlatformTextField.h
+++ b/Engine/Source/Runtime/Slate/Public/Framework/Text/IOS/IOSPlatformTextField.h
@@ -12,7 +12,9 @@
 class IVirtualKeyboardEntry;
 
 @class SlateTextField;
+#if !PLATFORM_TVOS
 @class SlateTextField2;
+#endif
 
 class FIOSPlatformTextField : public IPlatformTextField
 {
@@ -25,7 +27,9 @@ public:
 
 private:
 	SlateTextField* TextField;
+#if !PLATFORM_TVOS
 	SlateTextField2* TextField2;
+#endif
 };
 
 typedef FIOSPlatformTextField FPlatformTextField;
@@ -51,6 +55,7 @@ typedef FIOSPlatformTextField FPlatformTextField;
 
 @end
 
+#if !PLATFORM_TVOS
 @interface SlateTextField2 : NSObject <UITextFieldDelegate>
 {
 	TWeakPtr<IVirtualKeyboardEntry> TextWidget;
@@ -63,3 +68,4 @@ typedef FIOSPlatformTextField FPlatformTextField;
 -(bool)hasTextWidget;
 
 @end
+#endif
