Second iOS text field build fix for TVOS

#Synced-CL 51927141

diff --git a/Engine/Source/Runtime/Slate/Private/Framework/Text/IOS/IOSPlatformTextField.cpp b/Engine/Source/Runtime/Slate/Private/Framework/Text/IOS/IOSPlatformTextField.cpp
--- a/Engine/Source/Runtime/Slate/Private/Framework/Text/IOS/IOSPlatformTextField.cpp
+++ b/Engine/Source/Runtime/Slate/Private/Framework/Text/IOS/IOSPlatformTextField.cpp
@@ -56,7 +56,9 @@ namespace
 
 FIOSPlatformTextField::FIOSPlatformTextField()
 	: TextField( nullptr )
+#if !PLATFORM_TVOS
 	, TextField2( nullptr )
+#endif
 {
 }
 
