Limit new IOS text field to plain iOS only, excluding TVOS and VisionOS

#Synced-CL 51927141

diff --git a/Engine/Source/Runtime/Slate/Private/Framework/Text/IOS/IOSPlatformTextField.cpp b/Engine/Source/Runtime/Slate/Private/Framework/Text/IOS/IOSPlatformTextField.cpp
--- a/Engine/Source/Runtime/Slate/Private/Framework/Text/IOS/IOSPlatformTextField.cpp
+++ b/Engine/Source/Runtime/Slate/Private/Framework/Text/IOS/IOSPlatformTextField.cpp
@@ -56,7 +56,7 @@ namespace
 
 FIOSPlatformTextField::FIOSPlatformTextField()
 	: TextField( nullptr )
-#if !PLATFORM_TVOS
+#if PLATFORM_IOS && !PLATFORM_TVOS && !PLATFORM_VISIONOS
 	, TextField2( nullptr )
 #endif
 {
@@ -77,7 +77,7 @@ FIOSPlatformTextField::~FIOSPlatformTextField()
 			}
 		});
 	}
-#if !PLATFORM_TVOS
+#if PLATFORM_IOS && !PLATFORM_TVOS && !PLATFORM_VISIONOS
 	if (TextField2 != nullptr)
 	{
         UE_LOG(LogIOS, Log, TEXT("Deleting text field2: %p"), TextField2);
@@ -133,10 +133,10 @@ void FIOSPlatformTextField::ShowVirtualKeyboard(bool bShow, int32 UserIndex, TSh
 	}
 	else
 	{
-#if PLATFORM_TVOS
-		const bool bUseNewImpl = false;
-#else
+#if PLATFORM_IOS && !PLATFORM_TVOS && !PLATFORM_VISIONOS
 		const bool bUseNewImpl = (GIOSTextFieldImplementation == 1);
+#else
+		const bool bUseNewImpl = false;
 #endif
 
 		if (bShow)
@@ -150,7 +150,7 @@ void FIOSPlatformTextField::ShowVirtualKeyboard(bool bShow, int32 UserIndex, TSh
 			// these functions must be run on the main thread
 			if (bUseNewImpl)
 			{
-#if !PLATFORM_TVOS
+#if PLATFORM_IOS && !PLATFORM_TVOS && !PLATFORM_VISIONOS
 				dispatch_async(dispatch_get_main_queue(), ^{
 					if (TextField2 == nullptr)
 					{
@@ -176,6 +176,7 @@ void FIOSPlatformTextField::ShowVirtualKeyboard(bool bShow, int32 UserIndex, TSh
 #if !PLATFORM_TVOS
 			if (bUseNewImpl)
 			{
+#if PLATFORM_IOS && !PLATFORM_TVOS && !PLATFORM_VISIONOS
 				if (TextField2 != nullptr && [TextField2 hasTextWidget])
 				{
 	                UE_LOG(LogIOS, Log, TEXT("Hiding field2: %p"), TextField2);
@@ -188,6 +189,7 @@ void FIOSPlatformTextField::ShowVirtualKeyboard(bool bShow, int32 UserIndex, TSh
 	                    }
 	                });
 				}
+#endif
 			}
 			else
 			{
@@ -411,7 +413,7 @@ void FIOSPlatformTextField::ShowVirtualKeyboard(bool bShow, int32 UserIndex, TSh
 
 @end
 
-#if !PLATFORM_TVOS
+#if PLATFORM_IOS && !PLATFORM_TVOS && !PLATFORM_VISIONOS
 @implementation SlateTextField2
 
 -(id)init
diff --git a/Engine/Source/Runtime/Slate/Public/Framework/Text/IOS/IOSPlatformTextField.h b/Engine/Source/Runtime/Slate/Public/Framework/Text/IOS/IOSPlatformTextField.h
--- a/Engine/Source/Runtime/Slate/Public/Framework/Text/IOS/IOSPlatformTextField.h
+++ b/Engine/Source/Runtime/Slate/Public/Framework/Text/IOS/IOSPlatformTextField.h
@@ -12,7 +12,7 @@
 class IVirtualKeyboardEntry;
 
 @class SlateTextField;
-#if !PLATFORM_TVOS
+#if PLATFORM_IOS && !PLATFORM_TVOS && !PLATFORM_VISIONOS
 @class SlateTextField2;
 #endif
 
@@ -27,7 +27,7 @@ public:
 
 private:
 	SlateTextField* TextField;
-#if !PLATFORM_TVOS
+#if PLATFORM_IOS && !PLATFORM_TVOS && !PLATFORM_VISIONOS
 	SlateTextField2* TextField2;
 #endif
 };
@@ -55,7 +55,7 @@ typedef FIOSPlatformTextField FPlatformTextField;
 
 @end
 
-#if !PLATFORM_TVOS
+#if PLATFORM_IOS && !PLATFORM_TVOS && !PLATFORM_VISIONOS
 @interface SlateTextField2 : NSObject <UITextFieldDelegate>
 {
 	TWeakPtr<IVirtualKeyboardEntry> TextWidget;
