Slate: Deprecate Game window type

[RN] EWindowType::GameWindow is now deprecated
#Synced-CL 52219356

diff --git a/Engine/Plugins/Media/PixelStreaming/Source/PixelStreamingEditor/Private/PixelStreamingEditorUtils.cpp b/Engine/Plugins/Media/PixelStreaming/Source/PixelStreamingEditor/Private/PixelStreamingEditorUtils.cpp
--- a/Engine/Plugins/Media/PixelStreaming/Source/PixelStreamingEditor/Private/PixelStreamingEditorUtils.cpp
+++ b/Engine/Plugins/Media/PixelStreaming/Source/PixelStreamingEditor/Private/PixelStreamingEditorUtils.cpp
@@ -39,8 +39,10 @@ namespace UE::EditorPixelStreaming
 			case EWindowType::CursorDecorator:
 				return TEXT("CursorDecorator");
 
+PRAGMA_DISABLE_DEPRECATION_WARNINGS
 			case EWindowType::GameWindow:
 				return TEXT("GameWindow");
+PRAGMA_ENABLE_DEPRECATION_WARNINGS
 
 			default:
 				return TEXT("Unknown Window Type!");
diff --git a/Engine/Plugins/VirtualProduction/Avalanche/Source/AvalancheMedia/Private/Broadcast/OutputDevices/AvaBroadcastDisplayMediaCapture.cpp b/Engine/Plugins/VirtualProduction/Avalanche/Source/AvalancheMedia/Private/Broadcast/OutputDevices/AvaBroadcastDisplayMediaCapture.cpp
--- a/Engine/Plugins/VirtualProduction/Avalanche/Source/AvalancheMedia/Private/Broadcast/OutputDevices/AvaBroadcastDisplayMediaCapture.cpp
+++ b/Engine/Plugins/VirtualProduction/Avalanche/Source/AvalancheMedia/Private/Broadcast/OutputDevices/AvaBroadcastDisplayMediaCapture.cpp
@@ -56,9 +56,8 @@ public:
 		static FWindowStyle BackgroundlessStyle = FCoreStyle::Get().GetWidgetStyle<FWindowStyle>("Window");
 		BackgroundlessStyle.SetBackgroundBrush(FSlateNoResource());
 
-		// Attempt at creating a borderless game window.
+		// Attempt at creating a borderless window.
 		TSharedRef<SWindow> NewWindow = SNew(SWindow)
-			.Type(EWindowType::GameWindow)
 			.Style(&BackgroundlessStyle)
 			.ScreenPosition(Position)
 			.ClientSize(Size)
diff --git a/Engine/Source/Editor/PIEPreviewDeviceProfileSelector/Private/PIEPreviewDeviceProfileSelectorModule.cpp b/Engine/Source/Editor/PIEPreviewDeviceProfileSelector/Private/PIEPreviewDeviceProfileSelectorModule.cpp
--- a/Engine/Source/Editor/PIEPreviewDeviceProfileSelector/Private/PIEPreviewDeviceProfileSelectorModule.cpp
+++ b/Engine/Source/Editor/PIEPreviewDeviceProfileSelector/Private/PIEPreviewDeviceProfileSelectorModule.cpp
@@ -169,7 +169,6 @@ TSharedRef<SWindow> FPIEPreviewDeviceModule::CreatePIEPreviewDeviceWindow(FVecto
 	static FWindowStyle BackgroundlessStyle = FCoreStyle::Get().GetWidgetStyle<FWindowStyle>("Window");
 	BackgroundlessStyle.SetBackgroundBrush(FSlateNoResource());
 	TSharedRef<SPIEPreviewWindow> Window = SNew(SPIEPreviewWindow)
-		.Type(EWindowType::GameWindow)
 		.Style(&BackgroundlessStyle)
 		.ClientSize(ClientSize)
 		.Title(WindowTitle)
diff --git a/Engine/Source/Runtime/ApplicationCore/Public/GenericPlatform/GenericWindowDefinition.h b/Engine/Source/Runtime/ApplicationCore/Public/GenericPlatform/GenericWindowDefinition.h
--- a/Engine/Source/Runtime/ApplicationCore/Public/GenericPlatform/GenericWindowDefinition.h
+++ b/Engine/Source/Runtime/ApplicationCore/Public/GenericPlatform/GenericWindowDefinition.h
@@ -25,8 +25,8 @@ enum class EWindowType
 	Notification,
 	/** Value indicating that this is a window used for a cursor decorator */
 	CursorDecorator,
-	/** Value indicating that this is a game window */
-	GameWindow
+
+	GameWindow UE_DEPRECATED(5.8, "There is no longer any meaningful distinction between game and normal window. Use Normal instead. Windows now opt in to the Slate Console Debugger using a boolean property.")
 };
 
 
diff --git a/Engine/Source/Runtime/Engine/Private/GameEngine.cpp b/Engine/Source/Runtime/Engine/Private/GameEngine.cpp
--- a/Engine/Source/Runtime/Engine/Private/GameEngine.cpp
+++ b/Engine/Source/Runtime/Engine/Private/GameEngine.cpp
@@ -709,7 +709,6 @@ TSharedRef<SWindow> UGameEngine::CreateGameWindow()
 	}
 
 	TSharedRef<SWindow> Window = SNew(SWindow)
-	.Type(EWindowType::GameWindow)
 	.Style(bUseBorderlessWindow ? &BorderlessStyle : &FCoreStyle::Get().GetWidgetStyle<FWindowStyle>("Window"))
 	.ClientSize(FVector2D(ResX, ResY))
 	.AdjustInitialSizeAndPositionForDPIScale(false)
diff --git a/Engine/Source/Runtime/Slate/Private/Framework/Docking/TabManager.cpp b/Engine/Source/Runtime/Slate/Private/Framework/Docking/TabManager.cpp
--- a/Engine/Source/Runtime/Slate/Private/Framework/Docking/TabManager.cpp
+++ b/Engine/Source/Runtime/Slate/Private/Framework/Docking/TabManager.cpp
@@ -118,7 +118,9 @@ namespace UE::Slate::Private
 		TSharedRef<SWindow> TopLevelWindow = InWindow;
 
 		EWindowType WindowType = TopLevelWindow->GetType();
+PRAGMA_DISABLE_DEPRECATION_WARNINGS
 		while (!(WindowType == EWindowType::Normal || WindowType == EWindowType::GameWindow))
+PRAGMA_ENABLE_DEPRECATION_WARNINGS
 		{
 			// The window might be a menu or a tooltip if this the case check the parent window instead
 			if (TSharedPtr<SWindow> ParentWindow = TopLevelWindow->GetParentWindow())
diff --git a/Engine/Source/Runtime/Slate/Private/Test/SWindowZoneTest.cpp b/Engine/Source/Runtime/Slate/Private/Test/SWindowZoneTest.cpp
--- a/Engine/Source/Runtime/Slate/Private/Test/SWindowZoneTest.cpp
+++ b/Engine/Source/Runtime/Slate/Private/Test/SWindowZoneTest.cpp
@@ -113,13 +113,11 @@ namespace
 	};
 
 	TSharedRef<SWindow> CreateTestWindow(
-		EWindowType WindowType,
 		bool bUseOSWindowBorder,
 		ESizingRule SizingRule,
 		EWindowMode::Type WindowMode = EWindowMode::Windowed)
 	{
 		TSharedRef<SWindow> Result = SNew(SWindow)
-			.Type(WindowType)
 			.UseOSWindowBorder(bUseOSWindowBorder)
 			.ClientSize(FVector2D(WindowSize, WindowSize))
 			.ScreenPosition(FVector2f(256, 256))
@@ -132,7 +130,6 @@ namespace
 	}
 
 	TSharedRef<SWindow> CreateTestWindow_ZoneOverride(
-		EWindowType WindowType,
 		bool bUseOSWindowBorder,
 		ESizingRule SizingRule,
 		EWindowZone::Type ZoneOverride,
@@ -144,7 +141,6 @@ namespace
 		OutZoneWidget = ZoneWidget;
 
 		TSharedRef<SWindow> Result = SNew(SWindow)
-			.Type(WindowType)
 			.UseOSWindowBorder(bUseOSWindowBorder)
 			.ClientSize(FVector2D(WindowSize, WindowSize))
 			.ScreenPosition(FVector2f(100, 100))
@@ -183,43 +179,35 @@ namespace
 
 	/**
 	 * Test GetCurrentWindowZone returns NotInWindow for positions outside the window.
-	 * Tests all combinations of:
-	 * - Window types: GameWindow, Normal (2)
-	 * - Sizing rules: UserSized, FixedSize, Autosized (3)
-	 * Total: 6 combinations
+	 * Tests all sizing rules: UserSized, FixedSize, Autosized (3)
 	 */
 	IMPLEMENT_SIMPLE_AUTOMATION_TEST(FSWindowZoneNotInWindowTest, "Slate.Window.GetCurrentWindowZone.NotInWindow", SlateWindowZoneTestFlags)
 
 	bool FSWindowZoneNotInWindowTest::RunTest(const FString& Parameters)
 	{
 		FCursorModeScope CursorModeScope(*this);
-		for (EWindowType WindowType : { EWindowType::GameWindow, EWindowType::Normal })
+		for (ESizingRule SizingRule : { ESizingRule::UserSized, ESizingRule::FixedSize, ESizingRule::Autosized })
 		{
-			for (ESizingRule SizingRule : { ESizingRule::UserSized, ESizingRule::FixedSize, ESizingRule::Autosized })
+			TSharedRef<SWindow> Window = CreateTestWindow(/*bUseOSWindowBorder*/ true, SizingRule);
+			for (const FVector2D& OutsidePos : {
+				FTestPositions::Outside(),
+				FTestPositions::OutsideTopRight(),
+				FTestPositions::OutsideBottomLeft(),
+				FTestPositions::OutsideBottomRight(),
+				FTestPositions::OutsideTop(),
+				FTestPositions::OutsideBottom(),
+			})
 			{
-				TSharedRef<SWindow> Window = CreateTestWindow(WindowType, /*bUseOSWindowBorder*/ true, SizingRule);
-				for (const FVector2D& OutsidePos : {
-					FTestPositions::Outside(),
-					FTestPositions::OutsideTopRight(),
-					FTestPositions::OutsideBottomLeft(),
-					FTestPositions::OutsideBottomRight(),
-					FTestPositions::OutsideTop(),
-					FTestPositions::OutsideBottom(),
-				})
-				{
-					constexpr const ANSICHAR* SizingRuleNames[] = { "UserSized", "FixedSize", "Autosized" };
+				constexpr const ANSICHAR* SizingRuleNames[] = { "UserSized", "FixedSize", "Autosized" };
 
-					const EWindowZone::Type ActualZone = Window->GetCurrentWindowZone(OutsidePos);
-					TestEqual<int32>(
-						FString::Printf(TEXT("%S/%S at (%.0f,%.0f)"),
-							WindowType == EWindowType::Normal ? "Normal" : "GameWindow", SizingRuleNames[static_cast<int32>(SizingRule)],
-							OutsidePos.X, OutsidePos.Y),
-						ActualZone,
-						EWindowZone::NotInWindow);
-				}
-
-				FSlateApplication::Get().DestroyWindowImmediately(Window);
+				const EWindowZone::Type ActualZone = Window->GetCurrentWindowZone(OutsidePos);
+				TestEqual<int32>(
+					FString::Printf(TEXT("%S at (%.0f,%.0f)"), SizingRuleNames[static_cast<int32>(SizingRule)], OutsidePos.X, OutsidePos.Y),
+					ActualZone,
+					EWindowZone::NotInWindow);
 			}
+
+			FSlateApplication::Get().DestroyWindowImmediately(Window);
 		}
 
 		return true;
@@ -234,58 +222,52 @@ namespace
 	bool FSWindowZoneWindowedFullscreenTest::RunTest(const FString& Parameters)
 	{
 		FCursorModeScope CursorModeScope(*this);
-		for (EWindowType WindowType : { EWindowType::GameWindow, EWindowType::Normal })
+		for (bool bOsWindowBorder : { false, true })
 		{
-			for (bool bOsWindowBorder : { false, true })
+			// Test without widget override
 			{
-				// Test without widget override
+				TSharedRef<SWindow> Window = CreateTestWindow(bOsWindowBorder, ESizingRule::UserSized, EWindowMode::WindowedFullscreen);
+
+				for (const FVector2D& Pos : {
+					FTestPositions::TopLeft(),
+					FTestPositions::Top(),
+					FTestPositions::TopRight(),
+					FTestPositions::Left(),
+					FTestPositions::Center(),
+					FTestPositions::Right(),
+					FTestPositions::BottomLeft(),
+					FTestPositions::Bottom(),
+					FTestPositions::BottomRight(),
+				})
 				{
-					TSharedRef<SWindow> Window = CreateTestWindow(
-						WindowType, bOsWindowBorder, ESizingRule::UserSized, EWindowMode::WindowedFullscreen);
-
-					for (const FVector2D& Pos : {
-						FTestPositions::TopLeft(),
-						FTestPositions::Top(),
-						FTestPositions::TopRight(),
-						FTestPositions::Left(),
-						FTestPositions::Center(),
-						FTestPositions::Right(),
-						FTestPositions::BottomLeft(),
-						FTestPositions::Bottom(),
-						FTestPositions::BottomRight(),
-					})
-					{
-						const EWindowZone::Type ActualZone = Window->GetCurrentWindowZone(Pos);
-						TestEqual<int32>(
-							FString::Printf(TEXT("%S %S WindowedFullscreen at (%.0f,%.0f)"),
-								WindowType == EWindowType::Normal ? "Normal" : "GameWindow",
-								bOsWindowBorder ? "with OS window border" : "without OS window border",
-								Pos.X, Pos.Y),
-							ActualZone,
-							EWindowZone::ClientArea);
-					}
-
-					FSlateApplication::Get().DestroyWindowImmediately(Window);
+					const EWindowZone::Type ActualZone = Window->GetCurrentWindowZone(Pos);
+					TestEqual<int32>(
+						FString::Printf(TEXT("WindowedFullscreen %S at (%.0f,%.0f)"),
+							bOsWindowBorder ? "with OS window border" : "without OS window border",
+							Pos.X, Pos.Y),
+						ActualZone,
+						EWindowZone::ClientArea);
 				}
 
-				// Test with widget zone override
-				{
-					TSharedPtr<SZoneOverrideWidget> ZoneWidget;
-					TSharedRef<SWindow> Window = CreateTestWindow_ZoneOverride(
-						WindowType, bOsWindowBorder, ESizingRule::UserSized, EWindowZone::TitleBar, ZoneWidget, EWindowMode::WindowedFullscreen);
+				FSlateApplication::Get().DestroyWindowImmediately(Window);
+			}
 
-					FSlateApplication::Get().ForceRedrawWindow(Window);
+			// Test with widget zone override
+			{
+				TSharedPtr<SZoneOverrideWidget> ZoneWidget;
+				TSharedRef<SWindow> Window = CreateTestWindow_ZoneOverride(
+					bOsWindowBorder, ESizingRule::UserSized, EWindowZone::TitleBar, ZoneWidget, EWindowMode::WindowedFullscreen);
 
-					const EWindowZone::Type ActualZone = Window->GetCurrentWindowZone(FTestPositions::Center());
-					TestEqual<int32>(
-						FString::Printf(TEXT("%S %S WindowedFullscreen respects widget override"),
-							WindowType == EWindowType::Normal ? "Normal" : "GameWindow",
-							bOsWindowBorder ? "with OS window border" : "without OS window border"),
-						ActualZone,
-						EWindowZone::TitleBar);
+				FSlateApplication::Get().ForceRedrawWindow(Window);
 
-					FSlateApplication::Get().DestroyWindowImmediately(Window);
-				}
+				const EWindowZone::Type ActualZone = Window->GetCurrentWindowZone(FTestPositions::Center());
+				TestEqual<int32>(
+					FString::Printf(TEXT("%S WindowedFullscreen respects widget override"),
+						bOsWindowBorder ? "with OS window border" : "without OS window border"),
+					ActualZone,
+					EWindowZone::TitleBar);
+
+				FSlateApplication::Get().DestroyWindowImmediately(Window);
 			}
 		}
 
@@ -307,7 +289,7 @@ namespace
 		{
 			// Test: Edge positions should return ClientArea (no resize handles)
 			{
-				TSharedRef<SWindow> Window = CreateTestWindow(EWindowType::Normal, /*bUseOSWindowBorder*/ true, SizingRule);
+				TSharedRef<SWindow> Window = CreateTestWindow(/*bUseOSWindowBorder*/ true, SizingRule);
 
 				for (const FVector2D& EdgePos : {
 					FTestPositions::TopLeft(),
@@ -341,7 +323,6 @@ namespace
 			{
 				TSharedPtr<SZoneOverrideWidget> ZoneWidget;
 				TSharedRef<SWindow> Window = CreateTestWindow_ZoneOverride(
-					EWindowType::Normal,
 					/*bUseOSWindowBorder*/ true,
 					SizingRule,
 					EWindowZone::TitleBar,
@@ -389,26 +370,16 @@ namespace
 			{ FTestPositions::BottomRight(), EWindowZone::BottomRightBorder, "BottomRightBorder" },
 		};
 
-		for (EWindowType WindowType : { EWindowType::GameWindow, EWindowType::Normal })
-		{
-			TSharedRef<SWindow> Window = CreateTestWindow(
-				WindowType,
-				/*bUseOSWindowBorder*/ true,
-				ESizingRule::UserSized);
-
-			for (const FZoneTestCase& TestCase : TestCases)
-			{
-				const EWindowZone::Type ActualZone = Window->GetCurrentWindowZone(TestCase.Position);
-				TestEqual<int32>(
-					FString::Printf(TEXT("%S UserSized: %S"),
-						WindowType == EWindowType::GameWindow ? "GameWindow" : "Normal", TestCase.Name),
-					ActualZone,
-					TestCase.ExpectedZone);
-			}
+		TSharedRef<SWindow> Window = CreateTestWindow(/*bUseOSWindowBorder*/ true, ESizingRule::UserSized);
 
-			FSlateApplication::Get().DestroyWindowImmediately(Window);
+		for (const FZoneTestCase& TestCase : TestCases)
+		{
+			const EWindowZone::Type ActualZone = Window->GetCurrentWindowZone(TestCase.Position);
+			TestEqual<int32>(FString::Printf(TEXT("UserSized: %S"), TestCase.Name), ActualZone, TestCase.ExpectedZone);
 		}
 
+		FSlateApplication::Get().DestroyWindowImmediately(Window);
+
 		return true;
 	}
 } // namespace UE::Slate::Test
diff --git a/Engine/Source/Runtime/Slate/Private/Test/SlateUserCursorLockTest.cpp b/Engine/Source/Runtime/Slate/Private/Test/SlateUserCursorLockTest.cpp
--- a/Engine/Source/Runtime/Slate/Private/Test/SlateUserCursorLockTest.cpp
+++ b/Engine/Source/Runtime/Slate/Private/Test/SlateUserCursorLockTest.cpp
@@ -290,7 +290,6 @@ namespace
 	 */
 	bool TestCursorLockRounding(
 		FAutomationTestBase* Test,
-		EWindowType WindowType,
 		bool bUseOSWindowBorder,
 		const FCursorLockTestCase& TestCase,
 		const ANSICHAR* WindowDescription)
@@ -320,7 +319,6 @@ namespace
 
 		// Create the test window at integral coordinates
 		TSharedRef<SWindow> TestWindow = SNew(SWindow)
-			.Type(WindowType)
 			.UseOSWindowBorder(bUseOSWindowBorder)
 			.ClientSize(FVector2f(FMath::CeilToFloat(ClientWidth), FMath::CeilToFloat(ClientHeight)))
 			.ScreenPosition(FVector2f(WindowOriginX, WindowOriginY))
@@ -424,7 +422,6 @@ namespace
 	 */
 	bool RunCursorLockRoundingTestsForWindowConfig(
 		FAutomationTestBase* Test,
-		EWindowType WindowType,
 		bool bUseOSWindowBorder,
 		const ANSICHAR* WindowDescription)
 	{
@@ -474,7 +471,7 @@ namespace
 		bool bAllPassed = true;
 		for (const FCursorLockTestCase& TestCase : TestCases)
 		{
-			if (!TestCursorLockRounding(Test, WindowType, bUseOSWindowBorder, TestCase, WindowDescription))
+			if (!TestCursorLockRounding(Test, bUseOSWindowBorder, TestCase, WindowDescription))
 			{
 				bAllPassed = false;
 			}
@@ -490,48 +487,18 @@ namespace
 
 namespace UE::Slate::Test
 {
-	IMPLEMENT_SIMPLE_AUTOMATION_TEST(FSlateUserCursorLockBorderlessGameWindowTest, "Slate.User.CursorLock.BorderlessGameWindow", SlateUserCursorLockTestFlags)
+	IMPLEMENT_SIMPLE_AUTOMATION_TEST(FSlateUserCursorLockBorderlessTest, "Slate.User.CursorLock.Borderless", SlateUserCursorLockTestFlags)
 
-	bool FSlateUserCursorLockBorderlessGameWindowTest::RunTest(const FString& Parameters)
+	bool FSlateUserCursorLockBorderlessTest::RunTest(const FString& Parameters)
 	{
-		return RunCursorLockRoundingTestsForWindowConfig(
-			this,
-			EWindowType::GameWindow,
-			/*bUseOSWindowBorder*/ false,
-			"Borderless GameWindow");
+		return RunCursorLockRoundingTestsForWindowConfig(this, /*bUseOSWindowBorder*/ false, "Borderless");
 	}
 
-	IMPLEMENT_SIMPLE_AUTOMATION_TEST(FSlateUserCursorLockBorderlessNormalWindowTest, "Slate.User.CursorLock.BorderlessNormalWindow", SlateUserCursorLockTestFlags)
+	IMPLEMENT_SIMPLE_AUTOMATION_TEST(FSlateUserCursorLockBorderedTest, "Slate.User.CursorLock.Bordered", SlateUserCursorLockTestFlags)
 
-	bool FSlateUserCursorLockBorderlessNormalWindowTest::RunTest(const FString& Parameters)
+	bool FSlateUserCursorLockBorderedTest::RunTest(const FString& Parameters)
 	{
-		return RunCursorLockRoundingTestsForWindowConfig(
-			this,
-			EWindowType::Normal,
-			/*bUseOSWindowBorder*/ false,
-			"Borderless Normal");
-	}
-
-	IMPLEMENT_SIMPLE_AUTOMATION_TEST(FSlateUserCursorLockBorderedGameWindowTest, "Slate.User.CursorLock.BorderedGameWindow", SlateUserCursorLockTestFlags)
-
-	bool FSlateUserCursorLockBorderedGameWindowTest::RunTest(const FString& Parameters)
-	{
-		return RunCursorLockRoundingTestsForWindowConfig(
-			this,
-			EWindowType::GameWindow,
-			/*bUseOSWindowBorder*/ true,
-			"Bordered GameWindow");
-	}
-
-	IMPLEMENT_SIMPLE_AUTOMATION_TEST(FSlateUserCursorLockBorderedNormalWindowTest, "Slate.User.CursorLock.BorderedNormalWindow", SlateUserCursorLockTestFlags)
-
-	bool FSlateUserCursorLockBorderedNormalWindowTest::RunTest(const FString& Parameters)
-	{
-		return RunCursorLockRoundingTestsForWindowConfig(
-			this,
-			EWindowType::Normal,
-			/*bUseOSWindowBorder*/ true,
-			"Bordered Normal");
+		return RunCursorLockRoundingTestsForWindowConfig(this, /*bUseOSWindowBorder*/ true, "Bordered");
 	}
 } // namespace UE::Slate::Test
 
diff --git a/Engine/Source/Runtime/SlateCore/Private/Widgets/SWindow.cpp b/Engine/Source/Runtime/SlateCore/Private/Widgets/SWindow.cpp
--- a/Engine/Source/Runtime/SlateCore/Private/Widgets/SWindow.cpp
+++ b/Engine/Source/Runtime/SlateCore/Private/Widgets/SWindow.cpp
@@ -569,7 +569,9 @@ void SWindow::ConstructWindowInternals()
 	UpdateWindowContentVisibility();
 
 	// create window
+PRAGMA_DISABLE_DEPRECATION_WARNINGS
 	if ((Type == EWindowType::Normal || Type == EWindowType::GameWindow) && !bHasOSWindowBorder && !bVirtualWindow)
+PRAGMA_ENABLE_DEPRECATION_WARNINGS
 	{
 		// create window content slot
 		MainWindowArea->AddSlot()
