TabManager - Adding HasSidebaredTabs public function

#rnx
#Synced-CL 52063874
#Synced-CL 52063874

diff --git a/Engine/Source/Runtime/Slate/Private/Framework/Docking/TabManager.cpp b/Engine/Source/Runtime/Slate/Private/Framework/Docking/TabManager.cpp
--- a/Engine/Source/Runtime/Slate/Private/Framework/Docking/TabManager.cpp
+++ b/Engine/Source/Runtime/Slate/Private/Framework/Docking/TabManager.cpp
@@ -3255,6 +3255,11 @@ const TSharedPtr<const FTabSpawnerEntry> FTabManager::FindTabSpawnerFor(FName Ta
 		: TSharedPtr<const FTabSpawnerEntry>();
 }
 
+bool FTabManager::HasSidebaredTabs() const
+{
+	return !TemporarilySidebaredTabs.IsEmpty();
+}
+
 int32 FTabManager::FindTabInCollapsedAreas( const FTabMatcher& Matcher )
 {
 	for ( int32 CollapsedDockAreaIndex=0; CollapsedDockAreaIndex < CollapsedDockAreas.Num(); ++CollapsedDockAreaIndex )
diff --git a/Engine/Source/Runtime/Slate/Public/Framework/Docking/TabManager.h b/Engine/Source/Runtime/Slate/Public/Framework/Docking/TabManager.h
--- a/Engine/Source/Runtime/Slate/Public/Framework/Docking/TabManager.h
+++ b/Engine/Source/Runtime/Slate/Public/Framework/Docking/TabManager.h
@@ -1216,6 +1216,9 @@ class FTabManager : public TSharedFromThis<FTabManager>
 		/** Find a tab spawner by name. */
 		SLATE_API const TSharedPtr<const FTabSpawnerEntry> FindTabSpawnerFor(FName TabId) const;
 
+	    /** Return whether there are any temporarily sidebarred tabs **/
+		SLATE_API bool HasSidebaredTabs() const;
+	
 		UE_EXPERIMENTAL(5.7, "The panel drawer is still in development and might be removed. The api, implementation details and its form may change in near releases.")
 		SLATE_API bool HasPanelDrawer(const TSharedPtr<SWindow>& ParentWindow) const;
 
