Add support for dynamic invalidation the editor windows title bar

#Synced-CL 52348392

diff --git a/Engine/Source/Editor/LevelEditor/Private/LevelEditor.cpp b/Engine/Source/Editor/LevelEditor/Private/LevelEditor.cpp
--- a/Engine/Source/Editor/LevelEditor/Private/LevelEditor.cpp
+++ b/Engine/Source/Editor/LevelEditor/Private/LevelEditor.cpp
@@ -398,7 +398,11 @@ private:
 	FGeometry CachedGeometry;
 };
 
-
+template <>
+struct TWidgetTypeTraits<SProjectBadge>
+{
+	static constexpr bool SupportsInvalidation() { return true; }
+};
 
 TSharedRef<SDockTab> FLevelEditorModule::SpawnLevelEditor( const FSpawnTabArgs& InArgs )
 {
diff --git a/Engine/Source/Runtime/Slate/Private/Framework/Application/SWindowTitleBar.cpp b/Engine/Source/Runtime/Slate/Private/Framework/Application/SWindowTitleBar.cpp
--- a/Engine/Source/Runtime/Slate/Private/Framework/Application/SWindowTitleBar.cpp
+++ b/Engine/Source/Runtime/Slate/Private/Framework/Application/SWindowTitleBar.cpp
@@ -1,6 +1,7 @@
 // Copyright Epic Games, Inc. All Rights Reserved.
 
 #include "Framework/Application/SWindowTitleBar.h"
+#include "Widgets/SInvalidationPanel.h"
 
 void SWindowTitleBar::Construct( const FArguments& InArgs, const TSharedRef<SWindow>& InWindow, const TSharedPtr<SWidget>& InCenterContent, EHorizontalAlignment InCenterContentAlignment )
 {
@@ -185,38 +186,47 @@ void SWindowTitleBar::MakeTitleBarContentWidgets( TSharedPtr< SWidget >& OutLeft
 
 	if (ShowAppIcon && bHasWindowButtons)
 	{
-		OutLeftContent = 
-			SNew(SHorizontalBox)
-#if PLATFORM_MAC
-			// The Mac has traffic-light buttons to the left, so we need some empty space under them (but no icon).
-			+ SHorizontalBox::Slot()
-			.AutoWidth()
-			.HAlign(HAlign_Left)
-			.VAlign(VAlign_Top)
+		OutLeftContent =
+#if UE_SLATE_WITH_DYNAMIC_INVALIDATION
+			SNew(SInvalidationPanel)
+			.UseDynamicInvalidation(true)
 			[
-				// This leaves spaces for the macOS "traffic light" close/minimize/maximize buttons. Without this space, the traffic lights would render on top of the main menu on the Mac.
-				SNew(SSpacer)
-				.Size(FVector2D(64, 10))
-			]
+#endif
+				SNew(SHorizontalBox)
+#if PLATFORM_MAC
+				// The Mac has traffic-light buttons to the left, so we need some empty space under them (but no icon).
+				+ SHorizontalBox::Slot()
+				.AutoWidth()
+				.HAlign(HAlign_Left)
+				.VAlign(VAlign_Top)
+				[
+					// This leaves spaces for the macOS "traffic light" close/minimize/maximize buttons. Without this space, the traffic lights would render on top of the main menu on the Mac.
+					SNew(SSpacer)
+					.Size(FVector2D(64, 10))
+				]
 #else
-			+ SHorizontalBox::Slot()
-			.AutoWidth()
-			.HAlign(HAlign_Left)
-			.VAlign(VAlign_Top)
-			.Padding(FMargin(0, 0, 0, 4.0f))
-			[
-				SAssignNew(AppIconWidget, SAppIconWidget)
-				.IconColorAndOpacity(this, &SWindowTitleBar::GetWindowTitleContentColor)
-			]
+				+ SHorizontalBox::Slot()
+				.AutoWidth()
+				.HAlign(HAlign_Left)
+				.VAlign(VAlign_Top)
+				.Padding(FMargin(0, 0, 0, 4.0f))
+				[
+					SAssignNew(AppIconWidget, SAppIconWidget)
+					.IconColorAndOpacity(this, &SWindowTitleBar::GetWindowTitleContentColor)
+				]
 #endif //PLATFORM_MAC
-			+ SHorizontalBox::Slot()
-			.VAlign(VAlign_Top)
-			.FillWidth(1)
+				+ SHorizontalBox::Slot()
+				.VAlign(VAlign_Top)
+				.FillWidth(1)
 #if PLATFORM_MAC
-			// Align the main menu with the macOS "traffic light" buttons (close, minimize, and maximize).
-			.Padding(FMargin(0, 3.0f, 0, 0))
+				// Align the main menu with the macOS "traffic light" buttons (close, minimize, and maximize).
+				.Padding(FMargin(0, 3.0f, 0, 0))
 #endif // PLATFORM_MAC
-			.Expose(WindowMenuSlot);
+				.Expose(WindowMenuSlot)
+#if UE_SLATE_WITH_DYNAMIC_INVALIDATION
+			]
+#endif
+			;
 
 		// Default everything to use the small icon unless specifically set to use the large icon.
 		SetAllowMenuBar(false);
@@ -230,42 +240,52 @@ void SWindowTitleBar::MakeTitleBarContentWidgets( TSharedPtr< SWidget >& OutLeft
 
 	if (bHasWindowButtons)
 	{
-		OutRightContent = SNew(SBox)
-			.Visibility(EVisibility::SelfHitTestInvisible)
-			.Padding(FMargin(2.0f, 0.0f, 0.0f, 0.0f))
+		OutRightContent = 
+#if UE_SLATE_WITH_DYNAMIC_INVALIDATION
+			SNew(SInvalidationPanel)
+			.UseDynamicInvalidation(true)
 			[
-				// Expose a slot for optional content inside the title bar on the right side.
-				SNew(SHorizontalBox)
-				+ SHorizontalBox::Slot()
-				.Expose(RightSideContentSlot)
+#endif
+				SNew(SBox)
+				.Visibility(EVisibility::SelfHitTestInvisible)
+				.Padding(FMargin(2.0f, 0.0f, 0.0f, 0.0f))
+				[
+					// Expose a slot for optional content inside the title bar on the right side.
+					SNew(SHorizontalBox)
+					+ SHorizontalBox::Slot()
+					.Expose(RightSideContentSlot)
 
 // We don't need to add these as macOS draws them for us.
 #if !PLATFORM_MAC
-				// Minimize
-				+ SHorizontalBox::Slot()
-				.VAlign(VAlign_Top)
-				.AutoWidth()
-				[
-					MinimizeButton.ToSharedRef()
-				]
-
-				// Maximize/Restore
-				+ SHorizontalBox::Slot()
-				.VAlign(VAlign_Top)
-				.AutoWidth()
-				[
-					MaximizeRestoreButton.ToSharedRef()
-				]
-
-				// Close button
-				+ SHorizontalBox::Slot()
-				.VAlign(VAlign_Top)
-				.AutoWidth()
-				[
-					CloseButton.ToSharedRef()
-				]
+					// Minimize
+					+ SHorizontalBox::Slot()
+					.VAlign(VAlign_Top)
+					.AutoWidth()
+					[
+						MinimizeButton.ToSharedRef()
+					]
+
+					// Maximize/Restore
+					+ SHorizontalBox::Slot()
+					.VAlign(VAlign_Top)
+					.AutoWidth()
+					[
+						MaximizeRestoreButton.ToSharedRef()
+					]
+
+					// Close button
+					+ SHorizontalBox::Slot()
+					.VAlign(VAlign_Top)
+					.AutoWidth()
+					[
+						CloseButton.ToSharedRef()
+					]
 #endif //!PLATFORM_MAC
-			];
+				]
+#if UE_SLATE_WITH_DYNAMIC_INVALIDATION
+			]
+#endif
+			;
 	}
 	else
 	{
diff --git a/Engine/Source/Runtime/Slate/Private/Framework/MultiBox/SMenuEntryBlock.h b/Engine/Source/Runtime/Slate/Private/Framework/MultiBox/SMenuEntryBlock.h
--- a/Engine/Source/Runtime/Slate/Private/Framework/MultiBox/SMenuEntryBlock.h
+++ b/Engine/Source/Runtime/Slate/Private/Framework/MultiBox/SMenuEntryBlock.h
@@ -17,6 +17,12 @@
 
 class FActiveTimerHandle;
 
+template <>
+struct TWidgetTypeTraits<class SMenuEntryBlock>
+{
+	static constexpr bool SupportsInvalidation() { return true; }
+};
+
 /**
  * Menu entry MultiBlock
  */
diff --git a/Engine/Source/Runtime/Slate/Private/Framework/MultiBox/SMenuSeparatorBlock.h b/Engine/Source/Runtime/Slate/Private/Framework/MultiBox/SMenuSeparatorBlock.h
--- a/Engine/Source/Runtime/Slate/Private/Framework/MultiBox/SMenuSeparatorBlock.h
+++ b/Engine/Source/Runtime/Slate/Private/Framework/MultiBox/SMenuSeparatorBlock.h
@@ -6,6 +6,12 @@
 #include "Widgets/DeclarativeSyntaxSupport.h"
 #include "Framework/MultiBox/MultiBox.h"
 
+template <>
+struct TWidgetTypeTraits<class SMenuSeparatorBlock>
+{
+	static constexpr bool SupportsInvalidation() { return true; }
+};
+
 /**
  * Menu separator MultiBlock
  */
diff --git a/Engine/Source/Runtime/Slate/Private/Framework/MultiBox/SMenuSpacerBlock.h b/Engine/Source/Runtime/Slate/Private/Framework/MultiBox/SMenuSpacerBlock.h
--- a/Engine/Source/Runtime/Slate/Private/Framework/MultiBox/SMenuSpacerBlock.h
+++ b/Engine/Source/Runtime/Slate/Private/Framework/MultiBox/SMenuSpacerBlock.h
@@ -6,6 +6,12 @@
 #include "Widgets/DeclarativeSyntaxSupport.h"
 #include "Framework/MultiBox/MultiBox.h"
 
+template <>
+struct TWidgetTypeTraits<class SMenuSpacerBlock>
+{
+	static constexpr bool SupportsInvalidation() { return true; }
+};
+
 /**
  * Menu spacer MultiBlock, a variant of the separator block that has no drawn representation, but consumes space in the menu
  */
diff --git a/Engine/Source/Runtime/Slate/Public/Framework/Application/SWindowTitleBar.h b/Engine/Source/Runtime/Slate/Public/Framework/Application/SWindowTitleBar.h
--- a/Engine/Source/Runtime/Slate/Public/Framework/Application/SWindowTitleBar.h
+++ b/Engine/Source/Runtime/Slate/Public/Framework/Application/SWindowTitleBar.h
@@ -37,6 +37,13 @@ namespace SWindowTitleBarDefs
 }
 
 
+template <>
+struct TWidgetTypeTraits<class SAppIconWidget>
+{
+	static constexpr bool SupportsInvalidation() { return true; }
+};
+
+
 /** Widget that represents the app icon + system menu button, usually drawn in the top left of a Windows app */
 class SAppIconWidget
 	: public SCompoundWidget
