Remove invalidation from the window title bar since the min/max buttons reflect the OS state and that isn't well supported through invalidation.

#Synced-CL 52353515

#Synced-CL 52353515

diff --git a/Engine/Source/Runtime/Slate/Private/Framework/Application/SWindowTitleBar.cpp b/Engine/Source/Runtime/Slate/Private/Framework/Application/SWindowTitleBar.cpp
--- a/Engine/Source/Runtime/Slate/Private/Framework/Application/SWindowTitleBar.cpp
+++ b/Engine/Source/Runtime/Slate/Private/Framework/Application/SWindowTitleBar.cpp
@@ -1,7 +1,6 @@
 // Copyright Epic Games, Inc. All Rights Reserved.
 
 #include "Framework/Application/SWindowTitleBar.h"
-#include "Widgets/SInvalidationPanel.h"
 
 void SWindowTitleBar::Construct( const FArguments& InArgs, const TSharedRef<SWindow>& InWindow, const TSharedPtr<SWidget>& InCenterContent, EHorizontalAlignment InCenterContentAlignment )
 {
@@ -186,47 +185,38 @@ void SWindowTitleBar::MakeTitleBarContentWidgets( TSharedPtr< SWidget >& OutLeft
 
 	if (ShowAppIcon && bHasWindowButtons)
 	{
-		OutLeftContent =
-#if UE_SLATE_WITH_DYNAMIC_INVALIDATION
-			SNew(SInvalidationPanel)
-			.UseDynamicInvalidation(true)
-			[
-#endif
-				SNew(SHorizontalBox)
+		OutLeftContent = 
+			SNew(SHorizontalBox)
 #if PLATFORM_MAC
-				// The Mac has traffic-light buttons to the left, so we need some empty space under them (but no icon).
-				+ SHorizontalBox::Slot()
-				.AutoWidth()
-				.HAlign(HAlign_Left)
-				.VAlign(VAlign_Top)
-				[
-					// This leaves spaces for the macOS "traffic light" close/minimize/maximize buttons. Without this space, the traffic lights would render on top of the main menu on the Mac.
-					SNew(SSpacer)
-					.Size(FVector2D(64, 10))
-				]
+			// The Mac has traffic-light buttons to the left, so we need some empty space under them (but no icon).
+			+ SHorizontalBox::Slot()
+			.AutoWidth()
+			.HAlign(HAlign_Left)
+			.VAlign(VAlign_Top)
+			[
+				// This leaves spaces for the macOS "traffic light" close/minimize/maximize buttons. Without this space, the traffic lights would render on top of the main menu on the Mac.
+				SNew(SSpacer)
+				.Size(FVector2D(64, 10))
+			]
 #else
-				+ SHorizontalBox::Slot()
-				.AutoWidth()
-				.HAlign(HAlign_Left)
-				.VAlign(VAlign_Top)
-				.Padding(FMargin(0, 0, 0, 4.0f))
-				[
-					SAssignNew(AppIconWidget, SAppIconWidget)
-					.IconColorAndOpacity(this, &SWindowTitleBar::GetWindowTitleContentColor)
-				]
+			+ SHorizontalBox::Slot()
+			.AutoWidth()
+			.HAlign(HAlign_Left)
+			.VAlign(VAlign_Top)
+			.Padding(FMargin(0, 0, 0, 4.0f))
+			[
+				SAssignNew(AppIconWidget, SAppIconWidget)
+				.IconColorAndOpacity(this, &SWindowTitleBar::GetWindowTitleContentColor)
+			]
 #endif //PLATFORM_MAC
-				+ SHorizontalBox::Slot()
-				.VAlign(VAlign_Top)
-				.FillWidth(1)
+			+ SHorizontalBox::Slot()
+			.VAlign(VAlign_Top)
+			.FillWidth(1)
 #if PLATFORM_MAC
-				// Align the main menu with the macOS "traffic light" buttons (close, minimize, and maximize).
-				.Padding(FMargin(0, 3.0f, 0, 0))
+			// Align the main menu with the macOS "traffic light" buttons (close, minimize, and maximize).
+			.Padding(FMargin(0, 3.0f, 0, 0))
 #endif // PLATFORM_MAC
-				.Expose(WindowMenuSlot)
-#if UE_SLATE_WITH_DYNAMIC_INVALIDATION
-			]
-#endif
-			;
+			.Expose(WindowMenuSlot);
 
 		// Default everything to use the small icon unless specifically set to use the large icon.
 		SetAllowMenuBar(false);
@@ -240,52 +230,42 @@ void SWindowTitleBar::MakeTitleBarContentWidgets( TSharedPtr< SWidget >& OutLeft
 
 	if (bHasWindowButtons)
 	{
-		OutRightContent = 
-#if UE_SLATE_WITH_DYNAMIC_INVALIDATION
-			SNew(SInvalidationPanel)
-			.UseDynamicInvalidation(true)
+		OutRightContent = SNew(SBox)
+			.Visibility(EVisibility::SelfHitTestInvisible)
+			.Padding(FMargin(2.0f, 0.0f, 0.0f, 0.0f))
 			[
-#endif
-				SNew(SBox)
-				.Visibility(EVisibility::SelfHitTestInvisible)
-				.Padding(FMargin(2.0f, 0.0f, 0.0f, 0.0f))
-				[
-					// Expose a slot for optional content inside the title bar on the right side.
-					SNew(SHorizontalBox)
-					+ SHorizontalBox::Slot()
-					.Expose(RightSideContentSlot)
+				// Expose a slot for optional content inside the title bar on the right side.
+				SNew(SHorizontalBox)
+				+ SHorizontalBox::Slot()
+				.Expose(RightSideContentSlot)
 
 // We don't need to add these as macOS draws them for us.
 #if !PLATFORM_MAC
-					// Minimize
-					+ SHorizontalBox::Slot()
-					.VAlign(VAlign_Top)
-					.AutoWidth()
-					[
-						MinimizeButton.ToSharedRef()
-					]
-
-					// Maximize/Restore
-					+ SHorizontalBox::Slot()
-					.VAlign(VAlign_Top)
-					.AutoWidth()
-					[
-						MaximizeRestoreButton.ToSharedRef()
-					]
-
-					// Close button
-					+ SHorizontalBox::Slot()
-					.VAlign(VAlign_Top)
-					.AutoWidth()
-					[
-						CloseButton.ToSharedRef()
-					]
-#endif //!PLATFORM_MAC
+				// Minimize
+				+ SHorizontalBox::Slot()
+				.VAlign(VAlign_Top)
+				.AutoWidth()
+				[
+					MinimizeButton.ToSharedRef()
 				]
-#if UE_SLATE_WITH_DYNAMIC_INVALIDATION
-			]
-#endif
-			;
+
+				// Maximize/Restore
+				+ SHorizontalBox::Slot()
+				.VAlign(VAlign_Top)
+				.AutoWidth()
+				[
+					MaximizeRestoreButton.ToSharedRef()
+				]
+
+				// Close button
+				+ SHorizontalBox::Slot()
+				.VAlign(VAlign_Top)
+				.AutoWidth()
+				[
+					CloseButton.ToSharedRef()
+				]
+#endif //!PLATFORM_MAC
+			];
 	}
 	else
 	{
