PR #12129: Fix widgets that inherit from STableViewBase consuming RMB mouse-down input events when RMB-scrolling is disabled
#Synced-CL 52272722

diff --git a/Engine/Source/Runtime/Slate/Private/Widgets/Views/STableViewBase.cpp b/Engine/Source/Runtime/Slate/Private/Widgets/Views/STableViewBase.cpp
--- a/Engine/Source/Runtime/Slate/Private/Widgets/Views/STableViewBase.cpp
+++ b/Engine/Source/Runtime/Slate/Private/Widgets/Views/STableViewBase.cpp
@@ -538,31 +538,37 @@ FReply STableViewBase::OnMouseButtonDown( const FGeometry& MyGeometry, const FPo
 		OnRightMouseButtonDown(MouseEvent);
 	}
 
-	if ( MouseEvent.GetEffectingButton() == EKeys::RightMouseButton && ScrollBar->IsNeeded() )
+	if (bEnableRightClickScrolling)
 	{
-		AmountScrolledWhileRightMouseDown = 0;
+		if (MouseEvent.GetEffectingButton() == EKeys::RightMouseButton && ScrollBar->IsNeeded())
+		{
+			AmountScrolledWhileRightMouseDown = 0;
 
-		// NOTE: We don't bother capturing the mouse, unless the user starts dragging a few pixels (see the
-		// mouse move handling here.)  This is important so that the item row has a chance to select
-		// items when the right mouse button is released.  Just keep in mind that you might not get
-		// an OnMouseButtonUp event for the right mouse button if the user moves off of the table before
-		// they reach our scroll threshold
-		return FReply::Handled();
-	}
-	else if ( this->HasMouseCapture() )
-	{
-		// Consume all mouse buttons while we are RMB-dragging.
-		return FReply::Handled();
+			// NOTE: We don't bother capturing the mouse, unless the user starts dragging a few pixels (see the
+			// mouse move handling here.)  This is important so that the item row has a chance to select
+			// items when the right mouse button is released.  Just keep in mind that you might not get
+			// an OnMouseButtonUp event for the right mouse button if the user moves off of the table before
+			// they reach our scroll threshold
+			return FReply::Handled();
+		}
+		else if (this->HasMouseCapture())
+		{
+			// Consume all mouse buttons while we are RMB-dragging.
+			return FReply::Handled();
+		}
 	}
 	return FReply::Unhandled();			
 }
 
 FReply STableViewBase::OnMouseButtonDoubleClick( const FGeometry& InMyGeometry, const FPointerEvent& InMouseEvent )
 {
-	if ( this->HasMouseCapture() )
+	if (bEnableRightClickScrolling)
 	{
-		// Consume all other mouse buttons while we are RMB-dragging.
-		return FReply::Handled();
+		if ( this->HasMouseCapture() )
+		{
+			// Consume all other mouse buttons while we are RMB-dragging.
+			return FReply::Handled();
+		}
 	}
 	return FReply::Unhandled();			
 
