[Dynamic Invalidation] Fix `SButton` missing default initial foreground color

`BorderForegroundColorAttribute` was not initialized with a default value, leaving `SButton` with a default magenta foreground color which was visible for a frame until the style was explicitly applied.

#Synced-CL 52095318

diff --git a/Engine/Source/Runtime/Slate/Private/Widgets/Input/SButton.cpp b/Engine/Source/Runtime/Slate/Private/Widgets/Input/SButton.cpp
--- a/Engine/Source/Runtime/Slate/Private/Widgets/Input/SButton.cpp
+++ b/Engine/Source/Runtime/Slate/Private/Widgets/Input/SButton.cpp
@@ -64,7 +64,7 @@ void SButton::PrivateRegisterAttributes(FSlateAttributeInitializer& AttributeIni
 }
 
 SButton::SButton()
-	: BorderForegroundColorAttribute(*this)
+	: BorderForegroundColorAttribute(*this, FSlateColor::UseStyle())
 	, ContentPaddingAttribute(*this)
 	, NormalPaddingAttribute(*this)
 	, PressedPaddingAttribute(*this)
