Update TextLayout when text changed by virtual keyboard to prevent flicker while typing
#Synced-CL 52052884

diff --git a/Engine/Source/Runtime/Slate/Private/Widgets/Text/SlateEditableTextLayout.cpp b/Engine/Source/Runtime/Slate/Private/Widgets/Text/SlateEditableTextLayout.cpp
--- a/Engine/Source/Runtime/Slate/Private/Widgets/Text/SlateEditableTextLayout.cpp
+++ b/Engine/Source/Runtime/Slate/Private/Widgets/Text/SlateEditableTextLayout.cpp
@@ -3427,6 +3427,7 @@ void FSlateEditableTextLayout::UpdateTextChangedByVirtualKeyboard()
 	if (bTextChangedByVirtualKeyboard)
 	{
 		SetEditableText(VirtualKeyboardText);
+		TextLayout->UpdateIfNeeded();
 		// Let outsiders know that the text content has been changed
 		BroadcastOnTextChanged(GetEditableText());
 		bTextChangedByVirtualKeyboard = false;
