Change the SInlineEditableTextBlock to accept the delay time as an argument and skip creating a timer if set to zero
- Set the value to 0 in the TEDS Outliner

#Synced-CL 52470375

diff --git a/Engine/Plugins/Experimental/EditorDataStorageFeatures/Source/TedsOutliner/Private/Widgets/OutlinerTextWidget.cpp b/Engine/Plugins/Experimental/EditorDataStorageFeatures/Source/TedsOutliner/Private/Widgets/OutlinerTextWidget.cpp
--- a/Engine/Plugins/Experimental/EditorDataStorageFeatures/Source/TedsOutliner/Private/Widgets/OutlinerTextWidget.cpp
+++ b/Engine/Plugins/Experimental/EditorDataStorageFeatures/Source/TedsOutliner/Private/Widgets/OutlinerTextWidget.cpp
@@ -111,6 +111,7 @@ namespace UE::Editor::Outliner::TextWidget::Private
 				.IsSelected(WidgetRowBinder.BindEvent(&FExternalWidgetSelectionColumn::IsSelected))
 				.HighlightText(OutlinerWidgetRowBinder.BindData(&FHighlightTextColumn::HighlightText))
 				.ColorAndOpacity(Private::GetForegroundColor(DataStorage, DataRow, WidgetRow))
+				.DoubleSelectDelay(0.0f)
 			);
 
 			BindTableViewerRenameDelegate();
diff --git a/Engine/Source/Runtime/Slate/Private/Widgets/Text/SInlineEditableTextBlock.cpp b/Engine/Source/Runtime/Slate/Private/Widgets/Text/SInlineEditableTextBlock.cpp
--- a/Engine/Source/Runtime/Slate/Private/Widgets/Text/SInlineEditableTextBlock.cpp
+++ b/Engine/Source/Runtime/Slate/Private/Widgets/Text/SInlineEditableTextBlock.cpp
@@ -64,7 +64,7 @@ void SInlineEditableTextBlock::Construct( const FArguments& InArgs )
 
 	bIsMultiLine = InArgs._MultiLine;
 	bDelayedLeftClickEntersEditMode = InArgs._DelayedLeftClickEntersEditMode;
-	DoubleSelectDelay = 0.0f;
+	DoubleSelectDelay = InArgs._DoubleSelectDelay;
 
 	OnEnterEditingMode = InArgs._OnEnterEditingMode;
 	OnExitEditingMode = InArgs._OnExitEditingMode;
@@ -305,7 +305,7 @@ void SInlineEditableTextBlock::SetOnVerifyTextChanged(FOnVerifyTextChanged InOnV
 
 FReply SInlineEditableTextBlock::OnMouseButtonDown( const FGeometry& MyGeometry, const FPointerEvent& MouseEvent )
 {
-	if (!bDelayedLeftClickEntersEditMode || !MouseEvent.IsMouseButtonDown(EKeys::LeftMouseButton) || MouseEvent.IsControlDown() || MouseEvent.IsShiftDown())
+	if (!bDelayedLeftClickEntersEditMode || DoubleSelectDelay == 0.0f || !MouseEvent.IsMouseButtonDown(EKeys::LeftMouseButton) || MouseEvent.IsControlDown() || MouseEvent.IsShiftDown())
 	{
 		return FReply::Unhandled();
 	}
@@ -314,7 +314,7 @@ FReply SInlineEditableTextBlock::OnMouseButtonDown( const FGeometry& MyGeometry,
 	{
 		if (IsSelected.Execute() && !bIsReadOnly.Get() && !ActiveTimerHandle.IsValid())
 		{
-			ActiveTimerHandle = RegisterActiveTimer(0.5f, FWidgetActiveTimerDelegate::CreateSP(this, &SInlineEditableTextBlock::TriggerEditMode));
+			ActiveTimerHandle = RegisterActiveTimer(DoubleSelectDelay, FWidgetActiveTimerDelegate::CreateSP(this, &SInlineEditableTextBlock::TriggerEditMode));
 		}
 	}
 	else
@@ -344,6 +344,23 @@ FReply SInlineEditableTextBlock::OnDragOver( const FGeometry& MyGeometry, const
 
 FReply SInlineEditableTextBlock::OnMouseButtonDoubleClick( const FGeometry& InMyGeometry, const FPointerEvent& InMouseEvent )
 {
+	if (DoubleSelectDelay == 0.0f)
+	{
+		if (IsSelected.IsBound())
+		{
+			if (IsSelected.Execute() && !bIsReadOnly.Get() && !FSlateApplication::Get().IsDragDropping())
+			{
+				EnterEditingMode();
+			}
+		}
+		else if (HasKeyboardFocus() && !bIsReadOnly.Get())
+		{
+			EnterEditingMode();
+		}
+		
+		return FReply::Handled();
+	}
+	
 	auto PinnedActiveTimerHandle = ActiveTimerHandle.Pin();
 	if (PinnedActiveTimerHandle.IsValid())
 	{
diff --git a/Engine/Source/Runtime/Slate/Public/Widgets/Text/SInlineEditableTextBlock.h b/Engine/Source/Runtime/Slate/Public/Widgets/Text/SInlineEditableTextBlock.h
--- a/Engine/Source/Runtime/Slate/Public/Widgets/Text/SInlineEditableTextBlock.h
+++ b/Engine/Source/Runtime/Slate/Public/Widgets/Text/SInlineEditableTextBlock.h
@@ -53,6 +53,7 @@ class SInlineEditableTextBlock: public SCompoundWidget
 		, _LineBreakPolicy()
 		, _IsReadOnly(false)
 		, _MultiLine(false)
+		, _DoubleSelectDelay(0.5f)
 		, _DelayedLeftClickEntersEditMode(true)
 		, _ModiferKeyForNewLine(EModifierKey::None)
 		, _OverflowPolicy()
@@ -107,6 +108,9 @@ class SInlineEditableTextBlock: public SCompoundWidget
 		/** True if the editable text block is multi-line */
 		SLATE_ARGUMENT( bool, MultiLine )
 
+		/** The length of delay to enter edit mode after selection, if zero, it will enter edit mode on double click */
+		SLATE_ARGUMENT( float, DoubleSelectDelay )
+
 		/** True if left-clicking will enter edit mode after a delay. */
 		SLATE_ARGUMENT( bool, DelayedLeftClickEntersEditMode )
 
@@ -302,7 +306,7 @@ private:
 protected:
 	/** When selection of widget is managed by another widget, this delays the "double select" from 
 	occurring immediately, offering a chance for double clicking to take action. */
-	float DoubleSelectDelay;
+	float DoubleSelectDelay = 0.5f;
 
 	/** Attribute to look up if the widget is multiline */
 	uint8 bIsMultiLine:1;
