Fixed the min/max buttons disappearing on small windows

Cause: both the left content (menus) and right content (min/max/close buttons) were set to auto width with an empty slot filling out the rest of the space. When that middle slot couldn't size down anymore, the right content would clip out of the window.

Solution: Only the buttons are auto width now at the top level, the other slot resizes. An additional hbox is nested inside the resizing slot to contain the left content and the extenstion content. Inside that hbox, the extension content is the one that resizes dynamically while the left content is auto width. This makes its so that the extension content clips first, then the left content.

#Synced-CL 52453503

diff --git a/Engine/Source/Runtime/Slate/Private/Framework/Application/SWindowTitleBar.cpp b/Engine/Source/Runtime/Slate/Private/Framework/Application/SWindowTitleBar.cpp
--- a/Engine/Source/Runtime/Slate/Private/Framework/Application/SWindowTitleBar.cpp
+++ b/Engine/Source/Runtime/Slate/Private/Framework/Application/SWindowTitleBar.cpp
@@ -234,11 +234,7 @@ void SWindowTitleBar::MakeTitleBarContentWidgets( TSharedPtr< SWidget >& OutLeft
 			.Visibility(EVisibility::SelfHitTestInvisible)
 			.Padding(FMargin(2.0f, 0.0f, 0.0f, 0.0f))
 			[
-				// Expose a slot for optional content inside the title bar on the right side.
 				SNew(SHorizontalBox)
-				+ SHorizontalBox::Slot()
-				.Expose(RightSideContentSlot)
-
 // We don't need to add these as macOS draws them for us.
 #if !PLATFORM_MAC
 				// Minimize
@@ -282,6 +278,16 @@ TSharedRef<SWidget> SWindowTitleBar::MakeTitleBarContent( TSharedPtr<SWidget> Ce
 
 	MakeTitleBarContentWidgets(LeftContent, RightContent);
 
+	// Extension content lives in its own FillWidth slot so it clips before
+	// the window buttons (AutoWidth, right) and the app icon + menus (AutoWidth, left) do.
+	TSharedPtr<SWidget> ExtensionContent = SNew(SBox)
+		.Visibility(EVisibility::SelfHitTestInvisible)
+		[
+			SNew(SHorizontalBox)
+			+ SHorizontalBox::Slot()
+			.Expose(RightSideContentSlot)
+		];
+
 	// create window title if no content was provided
 	if (!CenterContent.IsValid())
 	{
@@ -307,10 +313,13 @@ TSharedRef<SWidget> SWindowTitleBar::MakeTitleBarContent( TSharedPtr<SWidget> Ce
 
 	// calculate content dimensions
 	LeftContent->SlatePrepass();
+	ExtensionContent->SlatePrepass();
 	RightContent->SlatePrepass();
 
 	FVector2D LeftSize = LeftContent->GetDesiredSize();
-	FVector2D RightSize = RightContent->GetDesiredSize();
+	FVector2D ExtensionSize = ExtensionContent->GetDesiredSize();
+	FVector2D ButtonsSize = RightContent->GetDesiredSize();
+	FVector2D RightSize = FVector2D(ExtensionSize.X + ButtonsSize.X, FMath::Max(ExtensionSize.Y, ButtonsSize.Y));
 
 	if (CenterContentAlignment == HAlign_Center)
 	{
@@ -357,21 +366,52 @@ TSharedRef<SWidget> SWindowTitleBar::MakeTitleBarContent( TSharedPtr<SWidget> Ce
 					.Size(FVector2D(RightSize.X, SpacerHeight))
 				]
 			]
+			// Foreground: nested layout with explicit clip priority.
+			// 1. RightContent (outer AutoWidth): never clips.
+			// 2. ExtensionContent (inner FillWidth + ClipToBounds): clips first.
+			// 3. LeftContent (inner AutoWidth, clipped by outer SBox): clips last.
 			+ SOverlay::Slot()
 			[
 				SNew(SHorizontalBox)
 				.Visibility(EVisibility::SelfHitTestInvisible)
+
+				// Outer fill slot: absorbs all available space left of the buttons.
+				// ClipToBounds ensures LeftContent cannot bleed into the button area.
 				+ SHorizontalBox::Slot()
-				.AutoWidth()
-				.HAlign(HAlign_Left)
+				.FillWidth(1.0f)
 				.VAlign(VAlign_Top)
 				[
-					LeftContent.ToSharedRef()
+					SNew(SBox)
+					.Clipping(EWidgetClipping::ClipToBounds)
+					[
+						SNew(SHorizontalBox)
+
+						// LeftContent: AutoWidth inside the clipped container.
+						// Only clips once the ExtensionContent area is fully gone.
+						+ SHorizontalBox::Slot()
+						.AutoWidth()
+						.HAlign(HAlign_Left)
+						.VAlign(VAlign_Top)
+						[
+							LeftContent.ToSharedRef()
+						]
+
+						// ExtensionContent: FillWidth so it clips before LeftContent is affected.
+						+ SHorizontalBox::Slot()
+						.FillWidth(1.0f)
+						.HAlign(HAlign_Right)
+						.VAlign(VAlign_Top)
+						[
+							SNew(SBox)
+							.Clipping(EWidgetClipping::ClipToBounds)
+							[
+								ExtensionContent.ToSharedRef()
+							]
+						]
+					]
 				]
 
-				+ SHorizontalBox::Slot()
-				.FillWidth(1.0f)
-
+				// RightContent: AutoWidth so it never clips.
 				+ SHorizontalBox::Slot()
 				.AutoWidth()
 				.HAlign(HAlign_Right)
