Hide iOS text field instead of repositioning out of screen

#Synced-CL 52463319

diff --git a/Engine/Source/Runtime/Slate/Private/Framework/Text/IOS/IOSPlatformTextField.cpp b/Engine/Source/Runtime/Slate/Private/Framework/Text/IOS/IOSPlatformTextField.cpp
--- a/Engine/Source/Runtime/Slate/Private/Framework/Text/IOS/IOSPlatformTextField.cpp
+++ b/Engine/Source/Runtime/Slate/Private/Framework/Text/IOS/IOSPlatformTextField.cpp
@@ -426,6 +426,7 @@ void FIOSPlatformTextField::ShowVirtualKeyboard(bool bShow, int32 UserIndex, TSh
 		const CGFloat TextFieldHeight = 46.0f;
 
 		InputTextField = [[UITextField alloc] initWithFrame:CGRectMake(0, ScreenSize.height, ScreenSize.width, TextFieldHeight)];
+		InputTextField.hidden = YES;
 		InputTextField.borderStyle = UITextBorderStyleNone;
 		InputTextField.backgroundColor = [[UIColor secondarySystemBackgroundColor] colorWithAlphaComponent:0.98];
 		InputTextField.layer.cornerRadius = 19.0f;
@@ -484,6 +485,7 @@ void FIOSPlatformTextField::ShowVirtualKeyboard(bool bShow, int32 UserIndex, TSh
 	CGFloat RightInset = SafeArea.right + InnerInset;
 
 	InputTextField.frame = CGRectMake(LeftInset, KBFrame.origin.y - FieldHeight - BottomInset, ScreenWidth - LeftInset - RightInset, FieldHeight);
+	InputTextField.hidden = NO;
 
 	// Just the sum of the text field and inset, keyboard height is not included
 	GIOSVirtualKeyboardInputFieldHeight.store(FieldHeight + BottomInset, std::memory_order_relaxed);
@@ -581,6 +583,8 @@ void FIOSPlatformTextField::ShowVirtualKeyboard(bool bShow, int32 UserIndex, TSh
 		[AppWindow addSubview:InputTextField];
 	}
 
+	InputTextField.hidden = NO;
+
 	[InputTextField becomeFirstResponder];
 }
 
@@ -590,8 +594,7 @@ void FIOSPlatformTextField::ShowVirtualKeyboard(bool bShow, int32 UserIndex, TSh
 
 	GIOSVirtualKeyboardInputFieldHeight.store(0.0f, std::memory_order_relaxed);
 
-	CGSize ScreenSize = [UIScreen mainScreen].bounds.size;
-	InputTextField.frame = CGRectMake(0, ScreenSize.height, InputTextField.frame.size.width, InputTextField.frame.size.height);
+	InputTextField.hidden = YES;
 
 	[InputTextField resignFirstResponder];
 }
