Add CVar WindowsCursor.DynamicLockToggle to toggle the mouse lock to be active only when the window is in the foreground on Windows

Also fix math for exclusive fullscreen on secondary monitors, which won't have top-left anchored to (0, 0).

#Synced-CL 52463319

diff --git a/Engine/Source/Runtime/ApplicationCore/Private/GenericPlatform/GenericWindow.cpp b/Engine/Source/Runtime/ApplicationCore/Private/GenericPlatform/GenericWindow.cpp
--- a/Engine/Source/Runtime/ApplicationCore/Private/GenericPlatform/GenericWindow.cpp
+++ b/Engine/Source/Runtime/ApplicationCore/Private/GenericPlatform/GenericWindow.cpp
@@ -159,6 +159,11 @@ bool FGenericWindow::IsForegroundWindow() const
 	return true;
 }
 
+bool FGenericWindow::SupportsDynamicCursorLock() const
+{
+	return false;
+}
+
 bool FGenericWindow::IsFullscreenSupported() const
 {
 	// empty default functionality
diff --git a/Engine/Source/Runtime/ApplicationCore/Private/Windows/WindowsCursor.cpp b/Engine/Source/Runtime/ApplicationCore/Private/Windows/WindowsCursor.cpp
--- a/Engine/Source/Runtime/ApplicationCore/Private/Windows/WindowsCursor.cpp
+++ b/Engine/Source/Runtime/ApplicationCore/Private/Windows/WindowsCursor.cpp
@@ -24,6 +24,12 @@ FAutoConsoleVariableRef CVarUseInvisibleCursorForNoneCursorType(
 	TEXT("If enabled, sets the platform HCursor to a transparent cursor instead of null when the mouse cursor type to None."),
 	ECVF_Default);
 
+static TAutoConsoleVariable<bool> CVarDynamicLockToggle(
+	TEXT("WindowsCursor.DynamicLockToggle"),
+	true,
+	TEXT("If enabled, lock the cursor automatically if our window is in the foreground."),
+	ECVF_ReadOnly);
+
 FWindowsCursor::FWindowsCursor()
 {
 	// Load up cursors that we'll be using
@@ -276,10 +282,55 @@ void FWindowsCursor::Show( bool bShow )
 	}
 }
 
+static TOptional<RECT> ClipBounds;
+
+static void HandleDynamicClip()
+{
+	static DWORD ProcessId = GetCurrentProcessId();
+	DWORD idProcess = 0;
+	HWND hWnd = GetForegroundWindow();
+	bool bAllowClip = ClipBounds.IsSet() && GetWindowThreadProcessId(hWnd, &idProcess) && (idProcess == ProcessId);
+	if (bAllowClip)
+	{
+		RECT WindowRect;
+		if (GetWindowRect(hWnd, &WindowRect))
+		{
+			const RECT& ClipRect = ClipBounds.GetValue();
+			bAllowClip =
+				(FMath::Max(WindowRect.left, ClipRect.left) < FMath::Min(WindowRect.right, ClipRect.right)) &&
+				(FMath::Max(WindowRect.top, ClipRect.top) < FMath::Min(WindowRect.bottom, ClipRect.bottom));
+		}
+	}
+	::ClipCursor(bAllowClip ? &ClipBounds.GetValue() : nullptr);
+}
+
+void FWindowsCursor::Tick()
+{
+	if (CVarDynamicLockToggle.GetValueOnGameThread())
+	{
+		HandleDynamicClip();
+	}
+}
+
 void FWindowsCursor::Lock( const RECT* const Bounds )
 {
 	// Lock/Unlock the cursor
-	::ClipCursor(Bounds);
+	if (CVarDynamicLockToggle.GetValueOnGameThread())
+	{
+		if (Bounds)
+		{
+			ClipBounds = *Bounds;
+		}
+		else
+		{
+			ClipBounds.Reset();
+		}
+		HandleDynamicClip();
+	}
+	else
+	{
+		::ClipCursor(Bounds);
+	}
 	// If the cursor is not visible and we're running game, assume we're in a mode where the mouse is controlling the camera and lock it to the center of the widget.
 }
 
diff --git a/Engine/Source/Runtime/ApplicationCore/Private/Windows/WindowsWindow.cpp b/Engine/Source/Runtime/ApplicationCore/Private/Windows/WindowsWindow.cpp
--- a/Engine/Source/Runtime/ApplicationCore/Private/Windows/WindowsWindow.cpp
+++ b/Engine/Source/Runtime/ApplicationCore/Private/Windows/WindowsWindow.cpp
@@ -1140,6 +1140,12 @@ bool FWindowsWindow::IsForegroundWindow() const
 	return ::GetForegroundWindow() == HWnd;
 }
 
+bool FWindowsWindow::SupportsDynamicCursorLock() const
+{
+	static const TConsoleVariableData<bool>* const CVarDynamicLockToggleData = IConsoleManager::Get().FindTConsoleVariableDataBool(TEXT("WindowsCursor.DynamicLockToggle"));
+	return CVarDynamicLockToggleData->GetValueOnGameThread();
+}
+
 bool FWindowsWindow::IsFullscreenSupported() const
 {
 #if defined(UE_WINDOWS_FORCE_NONEXCLUSIVE_FULLSCREEN) && UE_WINDOWS_FORCE_NONEXCLUSIVE_FULLSCREEN
diff --git a/Engine/Source/Runtime/ApplicationCore/Public/GenericPlatform/GenericWindow.h b/Engine/Source/Runtime/ApplicationCore/Public/GenericPlatform/GenericWindow.h
--- a/Engine/Source/Runtime/ApplicationCore/Public/GenericPlatform/GenericWindow.h
+++ b/Engine/Source/Runtime/ApplicationCore/Public/GenericPlatform/GenericWindow.h
@@ -189,7 +189,10 @@ public:
 	/** @return true if the window is in the foreground */
 	APPLICATIONCORE_API virtual bool IsForegroundWindow() const;
 
-	/** @return true if the window is in the foreground */
+	/** @return true if the window supports dynamic cursor locking */
+	APPLICATIONCORE_API virtual bool SupportsDynamicCursorLock() const;
+
+	/** @return true if the window supports fullscreen */
 	APPLICATIONCORE_API virtual bool IsFullscreenSupported() const;
 
 	/**
diff --git a/Engine/Source/Runtime/ApplicationCore/Public/GenericPlatform/ICursor.h b/Engine/Source/Runtime/ApplicationCore/Public/GenericPlatform/ICursor.h
--- a/Engine/Source/Runtime/ApplicationCore/Public/GenericPlatform/ICursor.h
+++ b/Engine/Source/Runtime/ApplicationCore/Public/GenericPlatform/ICursor.h
@@ -107,6 +107,11 @@ public:
 	 */
 	virtual void Show( bool bShow ) = 0;
 
+	/**
+	 * Perform operations that can't be done on an event-basis
+	 */
+	virtual void Tick() {}
+
 	/**
 	 * Locks the cursor to the passed in bounds
 	 * 
diff --git a/Engine/Source/Runtime/ApplicationCore/Public/Windows/WindowsCursor.h b/Engine/Source/Runtime/ApplicationCore/Public/Windows/WindowsCursor.h
--- a/Engine/Source/Runtime/ApplicationCore/Public/Windows/WindowsCursor.h
+++ b/Engine/Source/Runtime/ApplicationCore/Public/Windows/WindowsCursor.h
@@ -5,6 +5,7 @@
 #include "CoreTypes.h"
 #include "HAL/PlatformMemory.h"
 #include "GenericPlatform/ICursor.h"
+#include "Windows/WindowsHWrapper.h"
 
 class FWindowsCursor : public ICursor
 {
@@ -38,6 +39,7 @@ public:
 
 	virtual void Show( bool bShow ) override;
 
+	virtual void Tick() override;
 	virtual void Lock( const RECT* const Bounds ) override;
 
 	virtual void SetTypeShape(EMouseCursor::Type InCursorType, void* CursorHandle) override;
@@ -47,8 +49,8 @@ private:
 	EMouseCursor::Type CurrentType;
 
 	/** Cursors */
-	Windows::HCURSOR CursorHandles[ EMouseCursor::TotalCursorCount ];
+	HCURSOR CursorHandles[ EMouseCursor::TotalCursorCount ];
 
 	/** Override Cursors */
-	Windows::HCURSOR CursorOverrideHandles[EMouseCursor::TotalCursorCount];
+	HCURSOR CursorOverrideHandles[EMouseCursor::TotalCursorCount];
 };
\ No newline at end of file
diff --git a/Engine/Source/Runtime/ApplicationCore/Public/Windows/WindowsWindow.h b/Engine/Source/Runtime/ApplicationCore/Public/Windows/WindowsWindow.h
--- a/Engine/Source/Runtime/ApplicationCore/Public/Windows/WindowsWindow.h
+++ b/Engine/Source/Runtime/ApplicationCore/Public/Windows/WindowsWindow.h
@@ -131,6 +131,7 @@ public:
 	APPLICATIONCORE_API virtual int32 GetWindowTitleBarSize() const override;
 	virtual void* GetOSWindowHandle() const  override { return HWnd; }
 	APPLICATIONCORE_API virtual bool IsForegroundWindow() const override;
+	APPLICATIONCORE_API virtual bool SupportsDynamicCursorLock() const override;
 	APPLICATIONCORE_API virtual bool IsFullscreenSupported() const override;
 	APPLICATIONCORE_API virtual void SetText(const TCHAR* const Text) override;
 	APPLICATIONCORE_API virtual void DrawAttention(const FWindowDrawAttentionParameters& Parameters) override;
diff --git a/Engine/Source/Runtime/Engine/Private/Slate/SceneViewport.cpp b/Engine/Source/Runtime/Engine/Private/Slate/SceneViewport.cpp
--- a/Engine/Source/Runtime/Engine/Private/Slate/SceneViewport.cpp
+++ b/Engine/Source/Runtime/Engine/Private/Slate/SceneViewport.cpp
@@ -670,14 +670,17 @@ FReply FSceneViewport::AcquireFocusAndCapture(FIntPoint MousePosition, EFocusCau
 			if (FSlateApplication::Get().GetTransformFullscreenMouseInput() && !GIsEditor && Window.IsValid() && Window->GetWindowMode() == EWindowMode::Fullscreen)
 			{
 				const FDisplayMetrics& CachedDisplayMetrics = FSlateApplication::Get().GetCachedDisplayMetricsByRef();
-				FVector2D WindowSize = Window->GetSizeInScreen();
-				FVector2D DisplaySize = CachedDisplayMetrics.MonitorInfo.IsEmpty()
-					? FVector2D((float)CachedDisplayMetrics.PrimaryDisplayWidth, (float)CachedDisplayMetrics.PrimaryDisplayHeight)
+				const FVector2f DisplaySize = CachedDisplayMetrics.MonitorInfo.IsEmpty()
+					? FVector2f((float)CachedDisplayMetrics.PrimaryDisplayWidth, (float)CachedDisplayMetrics.PrimaryDisplayHeight)
 					: Window->GetFullScreenInfo().GetSize();
-				FVector2D CorrectionScale = DisplaySize / WindowSize;
-				MousePosBeforeHiddenDueToCapture = { (int32)(MousePosition.X * CorrectionScale.X), (int32)(MousePosition.Y * CorrectionScale.Y) };
+				const FVector2f CorrectionScale = DisplaySize / Window->GetSizeInScreen();
+				const FVector2f PositionInScreen = Window->GetPositionInScreen();
+				FVector2f AdjustedMousePosition = FVector2f(MousePosition);
+				AdjustedMousePosition -= PositionInScreen;
+				AdjustedMousePosition *= CorrectionScale;
+				AdjustedMousePosition += PositionInScreen;
+				MousePosBeforeHiddenDueToCapture = AdjustedMousePosition.IntPoint();
 			}
-
 		}
 
 		if ( bCursorHiddenDueToCapture || !bShouldShowMouseCursor )
@@ -1346,7 +1349,6 @@ FReply FSceneViewport::OnViewportActivated(const FWindowActivateEvent& InActivat
 		// Determine if we're in permanent capture mode.  This cannot be cached as part of bShouldCaptureMouseOnActivate because it could change between window activate and deactivate
 		const bool bPermanentCapture = ViewportClient->IsInPermanentCapture();
 
-
 		// If we are activating and had Mouse Capture on deactivate then we should get focus again
 		// It's important to note in the case of:
 		//    InActivateEvent.ActivationType == FWindowActivateEvent::EA_ActivateByMouse
diff --git a/Engine/Source/Runtime/Slate/Private/Framework/Application/SlateApplication.cpp b/Engine/Source/Runtime/Slate/Private/Framework/Application/SlateApplication.cpp
--- a/Engine/Source/Runtime/Slate/Private/Framework/Application/SlateApplication.cpp
+++ b/Engine/Source/Runtime/Slate/Private/Framework/Application/SlateApplication.cpp
@@ -1915,7 +1915,10 @@ FWidgetPath FSlateApplication::LocateWidgetInWindow(UE::Slate::FDeprecateVector2
 				? FVector2f(CachedDisplayMetrics.PrimaryDisplayWidth, CachedDisplayMetrics.PrimaryDisplayHeight)
 				: Window->GetFullScreenInfo().GetSize2f();
 
+			const FVector2f PositionInScreen = Window->GetPositionInScreen();
+			CursorPosition -= PositionInScreen;
 			CursorPosition *= WindowSize / DisplaySize;
+			CursorPosition += PositionInScreen;
 		}
 
 		TArray<FWidgetAndPointer> WidgetsAndCursors = Window->GetHittestGrid().GetBubblePath(CursorPosition, GetCursorRadius(), bIgnoreEnabledStatus, UserIndex);
@@ -4682,7 +4685,17 @@ FPointerEvent FSlateApplication::TransformPointerEvent(const FPointerEvent& Poin
 				? FVector2f(CachedDisplayMetrics.PrimaryDisplayWidth, CachedDisplayMetrics.PrimaryDisplayHeight)
 				: Window->GetFullScreenInfo().GetSize2f();
 
-			TransformedPointerEvent = FPointerEvent(PointerEvent, PointerEvent.GetScreenSpacePosition() * WindowSize / DisplaySize, PointerEvent.GetLastScreenSpacePosition() * WindowSize / DisplaySize);
+			const FVector2f PositionInScreen = Window->GetPositionInScreen();
+			const FVector2f Ratio = WindowSize / DisplaySize;
+			FVector2f Position = PointerEvent.GetScreenSpacePosition();
+			Position -= PositionInScreen;
+			Position *= Ratio;
+			Position += PositionInScreen;
+			FVector2f LastPosition = PointerEvent.GetLastScreenSpacePosition();
+			LastPosition -= PositionInScreen;
+			LastPosition *= Ratio;
+			LastPosition += PositionInScreen;
+			TransformedPointerEvent = FPointerEvent(PointerEvent, Position, LastPosition);
 		}
 	}
 
@@ -7119,7 +7132,6 @@ bool FSlateApplication::ProcessWindowActivatedEvent( const FWindowActivateEvent&
 			ActiveModalWindow->FlashWindow();
 		}
 
-		
 		TSharedRef<SWindow> Window = ActivateEvent.GetAffectedWindow();
 		TSharedPtr<ISlateViewport> Viewport = Window->GetViewport();
 		if (Viewport.IsValid())
diff --git a/Engine/Source/Runtime/Slate/Private/Framework/Application/SlateUser.cpp b/Engine/Source/Runtime/Slate/Private/Framework/Application/SlateUser.cpp
--- a/Engine/Source/Runtime/Slate/Private/Framework/Application/SlateUser.cpp
+++ b/Engine/Source/Runtime/Slate/Private/Framework/Application/SlateUser.cpp
@@ -937,6 +937,7 @@ void FSlateUser::UpdateCursor()
 		return;
 	}
 
+	Cursor->Tick();
 	if (LockingWidgetPath.IsValid())
 	{
 		const FWidgetPath PathToWidget = LockingWidgetPath.ToWidgetPath(FWeakWidgetPath::EInterruptedPathHandling::ReturnInvalid);
@@ -1010,7 +1011,7 @@ void FSlateUser::LockCursorInternal(const FWidgetPath& WidgetPath)
 
 	// Do not attempt to lock the cursor to the window if its not in the foreground.  It would cause annoying side effects
 	TSharedPtr<const FGenericWindow> NativeWindow = WidgetPath.GetWindow()->GetNativeWindow();
-	if (NativeWindow && NativeWindow->IsForegroundWindow())
+	if (NativeWindow && (NativeWindow->IsForegroundWindow() || NativeWindow->SupportsDynamicCursorLock()))
 	{
 		// The last widget in the path should be the widget we are locking the cursor to
 		FSlateRect SlateClipRect = WidgetPath.Widgets.Last().Geometry.GetLayoutBoundingRect();
@@ -1022,13 +1023,13 @@ void FSlateUser::LockCursorInternal(const FWidgetPath& WidgetPath)
 		if (FSlateApplication::Get().GetTransformFullscreenMouseInput() && !GIsEditor && NativeWindow->GetWindowMode() == EWindowMode::Fullscreen)
 		{
 			const FDisplayMetrics& CachedDisplayMetrics = FSlateApplication::Get().GetCachedDisplayMetricsByRef();
-			FVector2f DisplaySize = { (float)CachedDisplayMetrics.PrimaryDisplayWidth, (float)CachedDisplayMetrics.PrimaryDisplayHeight };
-			FVector2f DisplayDistortion = SlateClipRect.GetSize() / DisplaySize;
+			const FVector2f DisplaySize = CachedDisplayMetrics.MonitorInfo.IsEmpty()
+				? FVector2f(CachedDisplayMetrics.PrimaryDisplayWidth, CachedDisplayMetrics.PrimaryDisplayHeight)
+				: WidgetPath.GetWindow()->GetFullScreenInfo().GetSize2f();
+			FVector2f DisplayDistortion = DisplaySize / SlateClipRect.GetSize();
 
-			SlateClipRect.Left /= DisplayDistortion.X;
-			SlateClipRect.Top /= DisplayDistortion.Y;
-			SlateClipRect.Right /= DisplayDistortion.X;
-			SlateClipRect.Bottom /= DisplayDistortion.Y;
+			SlateClipRect.Right = ((SlateClipRect.Right - SlateClipRect.Left) * DisplayDistortion.X) + SlateClipRect.Left;
+			SlateClipRect.Bottom = ((SlateClipRect.Bottom - SlateClipRect.Top) * DisplayDistortion.Y) + SlateClipRect.Top;
 		}
 
 		// Note: Slate uses floating point bounds, but the cursor needs integers. Round the clip rect inwardly, to guarantee it remains inside the
