Software/Hardware cursor handling

#Synced-CL 52513890

diff --git a/Engine/Source/Editor/Experimental/EditorInteractiveToolsFramework/Private/EditorGizmos/EditorTRSGizmo.cpp b/Engine/Source/Editor/Experimental/EditorInteractiveToolsFramework/Private/EditorGizmos/EditorTRSGizmo.cpp
--- a/Engine/Source/Editor/Experimental/EditorInteractiveToolsFramework/Private/EditorGizmos/EditorTRSGizmo.cpp
+++ b/Engine/Source/Editor/Experimental/EditorInteractiveToolsFramework/Private/EditorGizmos/EditorTRSGizmo.cpp
@@ -1973,6 +1973,11 @@ void UEditorTRSGizmo::UpdateCursor(const EGizmoElementInteractionState InState,
 			else if (const TSharedPtr<SWidget> CursorWidget = GetCursorWidget(InHitPartId);
 				CursorWidget.IsValid())
 			{
+				if (CurrentMode == EGizmoTransformMode::Rotate)
+				{
+					CursorAPI->SetUseSoftwareCursor(true);
+				}
+
 				CursorAPI->SetCursorOverrideWidget(CursorWidget.ToSharedRef());
 			}
 			else if (!CanInteractWithPart(InHitPartId))
diff --git a/Engine/Source/Runtime/Slate/Private/Widgets/SViewport.cpp b/Engine/Source/Runtime/Slate/Private/Widgets/SViewport.cpp
--- a/Engine/Source/Runtime/Slate/Private/Widgets/SViewport.cpp
+++ b/Engine/Source/Runtime/Slate/Private/Widgets/SViewport.cpp
@@ -174,7 +174,10 @@ int32 SViewport::OnPaint( const FPaintArgs& Args, const FGeometry& AllottedGeome
 
 	if( ViewportInterfacePin.IsValid() && ViewportInterfacePin->IsSoftwareCursorVisible() )
 	{
-		const FVector2f CursorPosScreenSpace = FSlateApplication::Get().GetCursorPos();		
+		const FVector2D ToViewportPixelScale = UE::Slate::Viewport::GetViewportLocalToPixelScale(AllottedGeometry, ViewportInterfacePin->GetSize());
+		FVector2f CursorPositionLocalSpace = UE::Slate::CastToVector2f(ViewportInterfacePin->GetSoftwareCursorPosition() / ToViewportPixelScale);
+
+		const FVector2f CursorPosScreenSpace = FSlateApplication::Get().GetCursorPos();
 		// @todo Slate: why are we calling OnCursorQuery in here?
 		FCursorReply Reply = ViewportInterfacePin->OnCursorQuery( AllottedGeometry,
 			FPointerEvent(
@@ -185,34 +188,52 @@ int32 SViewport::OnPaint( const FPaintArgs& Args, const FGeometry& AllottedGeome
 				TSet<FKey>(),
 				FModifierKeysState() )
 		 );
-		EMouseCursor::Type CursorType = Reply.GetCursorType();
 
-		const FSlateBrush* Brush = [CursorType]
+		// Check if a cursor widget should be used instead of the standard brush
+		TOptional<TSharedRef<SWidget>> MappedCursorWidget = ViewportInterfacePin->OnMapCursor(Reply);
+		if (MappedCursorWidget.IsSet())
 		{
-			switch (CursorType)
-			{
-			case EMouseCursor::CardinalCross:
-				return FCoreStyle::Get().GetBrush(TEXT("SoftwareCursor_CardinalCross"));
-			case EMouseCursor::GrabHand:
-				return FCoreStyle::Get().GetBrush(TEXT("SoftwareCursor_Hand"));
-			case EMouseCursor::GrabHandClosed:
-			default: // Default is included here for legacy compatability; grab was the default.
-				return FCoreStyle::Get().GetBrush(TEXT("SoftwareCursor_Grab"));
-			}
-		}();
-		
-		const FVector2f CursorSize = Brush->ImageSize / AllottedGeometry.Scale;
-
-		const FVector2D ToViewportPixelScale = UE::Slate::Viewport::GetViewportLocalToPixelScale(AllottedGeometry, ViewportInterfacePin->GetSize());
-		FVector2f CursorPositionLocalSpace = UE::Slate::CastToVector2f(ViewportInterfacePin->GetSoftwareCursorPosition() / ToViewportPixelScale);
+			TSharedRef<SWidget> CursorWidget = MappedCursorWidget.GetValue();
+			CursorWidget->SetVisibility(EVisibility::HitTestInvisible);
+			CursorWidget->SlatePrepass(AllottedGeometry.Scale);
+
+			const FVector2f CursorSize = UE::Slate::CastToVector2f(CursorWidget->GetDesiredSize());
+
+			LayerId++;
+			CursorWidget->Paint(
+				Args,
+				AllottedGeometry.MakeChild(CursorSize, FSlateLayoutTransform(CursorPositionLocalSpace - (CursorSize * 0.5f))),
+				MyCullingRect,
+				OutDrawElements, LayerId, InWidgetStyle, bEnabled);
+		}
+		else
+		{
+			EMouseCursor::Type CursorType = Reply.GetCursorType();
 
-		LayerId++;
-		FSlateDrawElement::MakeBox(
-			OutDrawElements,
-			LayerId,
-			AllottedGeometry.ToPaintGeometry( CursorSize, FSlateLayoutTransform(CursorPositionLocalSpace - (CursorSize *.5f )) ),
-			Brush
-		);
+			const FSlateBrush* Brush = [CursorType]
+			{
+				switch (CursorType)
+				{
+				case EMouseCursor::CardinalCross:
+					return FCoreStyle::Get().GetBrush(TEXT("SoftwareCursor_CardinalCross"));
+				case EMouseCursor::GrabHand:
+					return FCoreStyle::Get().GetBrush(TEXT("SoftwareCursor_Hand"));
+				case EMouseCursor::GrabHandClosed:
+				default: // Default is included here for legacy compatability; grab was the default.
+					return FCoreStyle::Get().GetBrush(TEXT("SoftwareCursor_Grab"));
+				}
+			}();
+
+			const FVector2f CursorSize = Brush->ImageSize / AllottedGeometry.Scale;
+
+			LayerId++;
+			FSlateDrawElement::MakeBox(
+				OutDrawElements,
+				LayerId,
+				AllottedGeometry.ToPaintGeometry( CursorSize, FSlateLayoutTransform(CursorPositionLocalSpace - (CursorSize * .5f)) ),
+				Brush
+			);
+		}
 	}
 
 	// If there are any custom hit testable widgets in the 3D world we need to register their custom hit test path here.
@@ -245,7 +266,21 @@ FCursorReply SViewport::OnCursorQuery( const FGeometry& MyGeometry, const FPoint
 
 TOptional<TSharedRef<SWidget>> SViewport::OnMapCursor(const FCursorReply& CursorReply) const
 {
-	return ViewportInterface.IsValid() ? ViewportInterface.Pin()->OnMapCursor(CursorReply) : TOptional<TSharedRef<SWidget>>();
+	if (ViewportInterface.IsValid())
+	{
+		TSharedPtr<ISlateViewport> ViewportPin = ViewportInterface.Pin();
+
+		// When the software cursor is active, OnPaint handles cursor widget drawing
+		// at the software cursor position. Don't also let Slate draw it at the
+		// hardware cursor position via DrawCursor.
+		if (ViewportPin->IsSoftwareCursorVisible())
+		{
+			return TOptional<TSharedRef<SWidget>>();
+		}
+
+		return ViewportPin->OnMapCursor(CursorReply);
+	}
+	return TOptional<TSharedRef<SWidget>>();
 }
 
 FReply SViewport::OnMouseButtonDown( const FGeometry& MyGeometry, const FPointerEvent& MouseEvent )
