Slate: Only use SplitterHandleSize adjustment with Fixed size columns, not Fill sized to avoid Fill sized columns outgrowing Fixed size columns when resizing

#Synced-CL 52367967

diff --git a/Engine/Source/Runtime/Slate/Private/Widgets/Views/SHeaderRow.cpp b/Engine/Source/Runtime/Slate/Private/Widgets/Views/SHeaderRow.cpp
--- a/Engine/Source/Runtime/Slate/Private/Widgets/Views/SHeaderRow.cpp
+++ b/Engine/Source/Runtime/Slate/Private/Widgets/Views/SHeaderRow.cpp
@@ -1052,9 +1052,10 @@ void SHeaderRow::RegenerateWidgets()
 
 				case EColumnSizeMode::FillSized:
 				{
-					auto GetColumnWidthAsOptionalSize = [&SomeColumn, SplitterHandleSizeCopy = SplitterHandleSize]() -> FOptionalSize
+
+					auto GetColumnWidthAsOptionalSize = [&SomeColumn]() -> FOptionalSize
 					{
-						const float DesiredWidth = SomeColumn.GetWidth() - SplitterHandleSizeCopy;
+						const float DesiredWidth = SomeColumn.GetWidth();
 						return FOptionalSize(DesiredWidth);
 					};
 
