PR #14289: Fix SScrollBox can have a hard reference to itself and, as a result, not be destroyed when there are no other references left

#tests Run with it for a while without issue.
#Synced-CL 52573034

diff --git a/Engine/Source/Runtime/Slate/Private/Widgets/Layout/SScrollBox.cpp b/Engine/Source/Runtime/Slate/Private/Widgets/Layout/SScrollBox.cpp
--- a/Engine/Source/Runtime/Slate/Private/Widgets/Layout/SScrollBox.cpp
+++ b/Engine/Source/Runtime/Slate/Private/Widgets/Layout/SScrollBox.cpp
@@ -615,8 +615,14 @@ void SScrollBox::ScrollToEnd()
 
 void SScrollBox::ScrollDescendantIntoView(const TSharedPtr<SWidget>& WidgetToScrollIntoView, bool InAnimateScroll, EDescendantScrollDestination InDestination, float InScrollPadding)
 {
-	ScrollIntoViewRequest = [this, WidgetToScrollIntoView, InAnimateScroll, InDestination, InScrollPadding] (FGeometry AllottedGeometry) {
-		InternalScrollDescendantIntoView(AllottedGeometry, WidgetToScrollIntoView, InAnimateScroll, InDestination, InScrollPadding);
+	TWeakPtr<SWidget> WidgetWeak = WidgetToScrollIntoView;
+
+	ScrollIntoViewRequest = [this, WidgetWeak, InAnimateScroll, InDestination, InScrollPadding](FGeometry AllottedGeometry)
+	{
+		if (WidgetWeak.IsValid())
+		{
+			InternalScrollDescendantIntoView(AllottedGeometry, WidgetWeak.Pin(), InAnimateScroll, InDestination, InScrollPadding);
+		}
 	};
 
 	BeginInertialScrolling();
