Axis-Aware Autosizing for Dock Tabs (SDockTab)
- Add SetAutosizePolicy() that will override ShouldAutosize() to enable auto-width and auto-height auto policies.
- For use with the new Sequencer Simple View that autosizes the tab to the height of the simple view.

Docked tabs using ShouldAutosize(true) force SSplitter::SizeToContent on the parent splitter axis, causing tabs to consume all available space (e.g., covering adjacent tabs like Curve Editor when docked side-by-side).

#Synced-CL 52671255

diff --git a/Engine/Source/Runtime/Slate/Private/Framework/Docking/SDockingTabStack.cpp b/Engine/Source/Runtime/Slate/Private/Framework/Docking/SDockingTabStack.cpp
--- a/Engine/Source/Runtime/Slate/Private/Framework/Docking/SDockingTabStack.cpp
+++ b/Engine/Source/Runtime/Slate/Private/Framework/Docking/SDockingTabStack.cpp
@@ -1,6 +1,7 @@
 // Copyright Epic Games, Inc. All Rights Reserved.
 
 #include "Framework/Docking/SDockingTabStack.h"
+#include "Algo/AllOf.h"
 #include "Framework/Commands/UIAction.h"
 #include "Framework/Commands/UICommandList.h"
 #include "Widgets/Text/STextBlock.h"
@@ -1320,23 +1321,45 @@ TSharedPtr<SDockTab> SDockingTabStack::GetSelectedTabFromWell()
 
 SSplitter::ESizeRule SDockingTabStack::GetSizeRule() const
 {
-	int32 NumTabs = this->GetNumTabs();
-	if (NumTabs > 0)
-	{
-		for (int32 Index = 0; Index < NumTabs; ++Index)
-		{
-			if (!this->GetTabs()[Index]->ShouldAutosize())
-			{
-				return SSplitter::FractionOfParent;
-			}
-		}
-		// If all tabs in this stack are sized to content, then the stack's cell should size to Content.
-		return SSplitter::SizeToContent;
-	}
-	else
+	const TSlotlessChildren<SDockTab> ChildTabs = GetTabs();
+
+	const int32 NumTabs = ChildTabs.Num();
+	if (NumTabs <= 0)
 	{
 		return SSplitter::FractionOfParent;
 	}
+
+	const TSharedPtr<SDockingSplitter> ParentSplitter = ParentNodePtr.Pin();
+	const TArrayView<const TSharedRef<SDockTab>> ChildTabsArray = MakeArrayView(&ChildTabs[0], ChildTabs.Num());
+
+	const bool bAllTabsAllowAutosize = Algo::AllOf(ChildTabsArray
+		,[&ParentSplitter](const TSharedRef<SDockTab>& InTab)
+	{
+		const EDockTabAutosizePolicy Policy = InTab->GetAutosizePolicy();
+
+		// No parent splitter: treat like generic autosize check
+		if (!ParentSplitter.IsValid())
+		{
+			return Policy != EDockTabAutosizePolicy::None;
+		}
+
+		const EOrientation ParentOrientation = ParentSplitter->GetOrientation();
+
+		switch (Policy)
+		{
+		default:
+		case EDockTabAutosizePolicy::None:
+			return false;
+		case EDockTabAutosizePolicy::LayoutAxis:
+			return true;
+		case EDockTabAutosizePolicy::HorizontalOnly:
+			return ParentOrientation == Orient_Horizontal;
+		case EDockTabAutosizePolicy::VerticalOnly:
+			return ParentOrientation == Orient_Vertical;
+		}
+	});
+
+	return bAllTabsAllowAutosize ? SSplitter::SizeToContent : SSplitter::FractionOfParent;
 }
 
 void SDockingTabStack::SetTabWellHidden( bool bShouldHideTabWell )
diff --git a/Engine/Source/Runtime/Slate/Private/Widgets/Docking/SDockTab.cpp b/Engine/Source/Runtime/Slate/Private/Widgets/Docking/SDockTab.cpp
--- a/Engine/Source/Runtime/Slate/Private/Widgets/Docking/SDockTab.cpp
+++ b/Engine/Source/Runtime/Slate/Private/Widgets/Docking/SDockTab.cpp
@@ -425,7 +425,7 @@ SDockTab::SDockTab()
 	, OnTabActivated()
 	, OnCanCloseTab()
 	, ContentAreaPadding( 2.f )
-	, bShouldAutosize(false)
+	, AutosizePolicy(EDockTabAutosizePolicy::None)
 	, TabColorScale(FLinearColor::Transparent)
 	, LastActivationTime(0.0)
 {
@@ -615,7 +615,8 @@ void SDockTab::Construct( const FArguments& InArgs )
 	this->OnExtendContextMenu = InArgs._OnExtendContextMenu;
 	this->TabLabel = InArgs._Label;
 	this->TabLabelSuffix = InArgs._LabelSuffix;
-	this->bShouldAutosize = InArgs._ShouldAutosize;
+	this->AutosizePolicy = (InArgs._ShouldAutosize && InArgs._AutosizePolicy == EDockTabAutosizePolicy::None)
+		? EDockTabAutosizePolicy::LayoutAxis : InArgs._AutosizePolicy; // Map old arg to new arg
 	this->TabColorScale = InArgs._TabColorScale;
 	this->IconColor = InArgs._IconColor;
 
@@ -998,15 +999,24 @@ void SDockTab::SetTabIcon( const TAttribute<const FSlateBrush*> InTabIcon )
 	TabIcon = InTabIcon;
 }
 
-
 bool SDockTab::ShouldAutosize() const
 {
-	return bShouldAutosize;
+	return AutosizePolicy != EDockTabAutosizePolicy::None;
 }
 
 void SDockTab::SetShouldAutosize(const bool bNewShouldAutosize)
 {
-	bShouldAutosize = bNewShouldAutosize;
+	AutosizePolicy = bNewShouldAutosize ? EDockTabAutosizePolicy::LayoutAxis : EDockTabAutosizePolicy::None;
+}
+
+EDockTabAutosizePolicy SDockTab::GetAutosizePolicy() const
+{
+	return AutosizePolicy;
+}
+
+void SDockTab::SetAutosizePolicy(const EDockTabAutosizePolicy InNewPolicy)
+{
+	AutosizePolicy = InNewPolicy;
 }
 
 FReply SDockTab::OnCloseButtonClicked()
diff --git a/Engine/Source/Runtime/Slate/Public/Widgets/Docking/SDockTab.h b/Engine/Source/Runtime/Slate/Public/Widgets/Docking/SDockTab.h
--- a/Engine/Source/Runtime/Slate/Public/Widgets/Docking/SDockTab.h
+++ b/Engine/Source/Runtime/Slate/Public/Widgets/Docking/SDockTab.h
@@ -33,13 +33,26 @@ enum ETabRole : uint8
 	NumRoles
 };
 
-	/** The cause of a tab activation */
+/** The cause of a tab activation */
 enum ETabActivationCause : uint8
 {
 	UserClickedOnTab,
 	SetDirectly
 };
 
+/** Defines the policy for automatically sizing dock tabs. */
+enum class EDockTabAutosizePolicy : uint8
+{
+	/** No autosize */
+	None,
+	/** Autosize on whatever axis the layout is using (old ShouldAutosize(true) behavior) */
+	LayoutAxis,
+	/** Autosize only when the parent splitter is horizontal (width axis) */
+	HorizontalOnly,
+	/** Autosize only when the parent splitter is vertical (height axis) */
+	VerticalOnly
+};
+
 class FMenuBuilder;
 
 namespace UE::Slate::Private
@@ -90,6 +103,7 @@ public:
 		, _OnTabRelocated()
 		, _OnTabDraggedOverDockArea()
 		, _ShouldAutosize(false)
+		, _AutosizePolicy(EDockTabAutosizePolicy::None)
 		, _OnCanCloseTab()
 		, _CanEverClose(true)
 		, _OnPersistVisualState()
@@ -117,6 +131,8 @@ public:
 		SLATE_EVENT( FSimpleDelegate, OnTabRelocated )
 		SLATE_EVENT( FSimpleDelegate, OnTabDraggedOverDockArea )
 		SLATE_ARGUMENT( bool, ShouldAutosize )
+		/** Overrides and acts as a replacement for ShouldAutosize to allow more auto-sizing policies */
+		SLATE_ARGUMENT( EDockTabAutosizePolicy, AutosizePolicy )
 		/** When the close button is pressed, checks whether the tab can be closed in that moment. Example: Show dialog and ask user whether they're sure to close. */
 		SLATE_EVENT( FCanCloseTab, OnCanCloseTab )
 		/** Whether this tab can ever be closed. Example: certain programs may want to show tabs for the lifetime of the program. */
@@ -252,6 +268,12 @@ public:
 	/** Set whether this tab should be sized based on its content. */
 	SLATE_API void SetShouldAutosize(const bool bNewShouldAutosize);
 
+	/** Gets the policy for how its content is auto-sized */
+	SLATE_API EDockTabAutosizePolicy GetAutosizePolicy() const;
+
+	/** Overrides ShouldAutosize() and sets the policy for how the tab should be auto-sized to its content. */
+	SLATE_API void SetAutosizePolicy(const EDockTabAutosizePolicy InNewPolicy);
+
 	/** @return true if the tab can be closed right now. Example: Callback could ask user via dialog. */
 	SLATE_API bool CanCloseTab() const;
 
@@ -512,8 +534,8 @@ protected:
 
 	TAttribute<FMargin> ContentAreaPadding;
 
-	/** Should this tab be auto-sized based on its content? */
-	bool bShouldAutosize;
+	/** How should this tab be auto-sized based on its content? */
+	EDockTabAutosizePolicy AutosizePolicy;
 
 	/** Color of this tab */
 	TAttribute<FLinearColor> TabColorScale;
